#ttp 251010 - AUDIO: Integrate Omni Radius located in Ambient Properties that PCF added for Jack

#proj UE4.Engine

#summary Integrated Omni Radius feature that was previously added to UE3

#change Brought across code from UE3 as best as possible, some could be simplified due to all attenuation being done via one class now and other aspects are now handled by FSoundParseParameters instead of the Audio component. It will also currently only work for PC and Xbox, other platforms don't have the same simple method of setting an 'inner radius' on a sound source.

diff --git a/Engine/Source/Editor/DetailCustomizations/Private/AttenuationSettingsCustomizations.cpp b/Engine/Source/Editor/DetailCustomizations/Private/AttenuationSettingsCustomizations.cpp
--- a/Engine/Source/Editor/DetailCustomizations/Private/AttenuationSettingsCustomizations.cpp
+++ b/Engine/Source/Editor/DetailCustomizations/Private/AttenuationSettingsCustomizations.cpp
@@ -167,11 +167,12 @@ void FAttenuationSettingsCustomization::CustomizeStructChildren( TSharedRef<IPro
 	ConeOffsetRow.Visibility(TAttribute<EVisibility>(this, &FAttenuationSettingsCustomization::IsConeSelected));
 
 	ChildBuilder.AddChildProperty(PropertyHandles.FindChecked(GET_MEMBER_NAME_CHECKED(FAttenuationSettings, FalloffDistance)).ToSharedRef());
+	ChildBuilder.AddChildProperty(PropertyHandles.FindChecked(GET_MEMBER_NAME_CHECKED(FAttenuationSettings, OmniRadius)).ToSharedRef());
 	ChildBuilder.AddChildProperty(PropertyHandles.FindChecked(GET_MEMBER_NAME_CHECKED(FAttenuationSettings, bAttenuateWithLPF)).ToSharedRef());
 	ChildBuilder.AddChildProperty(PropertyHandles.FindChecked(GET_MEMBER_NAME_CHECKED(FAttenuationSettings, LPFRadiusMin)).ToSharedRef());
 	ChildBuilder.AddChildProperty(PropertyHandles.FindChecked(GET_MEMBER_NAME_CHECKED(FAttenuationSettings, LPFRadiusMax)).ToSharedRef());
 
-	if (PropertyHandles.Num() != 11)
+	if (PropertyHandles.Num() != 12)
 	{
 		FString PropertyList;
 		for (auto It(PropertyHandles.CreateConstIterator()); It; ++It)
diff --git a/Engine/Source/Runtime/Engine/Classes/Components/AudioComponent.h b/Engine/Source/Runtime/Engine/Classes/Components/AudioComponent.h
--- a/Engine/Source/Runtime/Engine/Classes/Components/AudioComponent.h
+++ b/Engine/Source/Runtime/Engine/Classes/Components/AudioComponent.h
@@ -66,6 +66,7 @@ struct FSoundParseParameters
 	float HighFrequencyGain;
 
 	float StartTime;
+	float OmniRadius;
 
 	uint32 bUseSpatialization:1;
 	uint32 bLooping:1;
@@ -78,6 +79,7 @@ struct FSoundParseParameters
 		, Pitch(1.f)
 		, HighFrequencyGain(1.f)
 		, StartTime(-1.f)
+		, OmniRadius(0.0f)
 		, bUseSpatialization(false)
 		, bLooping(false)
 	{
diff --git a/Engine/Source/Runtime/Engine/Classes/Sound/SoundAttenuation.h b/Engine/Source/Runtime/Engine/Classes/Sound/SoundAttenuation.h
--- a/Engine/Source/Runtime/Engine/Classes/Sound/SoundAttenuation.h
+++ b/Engine/Source/Runtime/Engine/Classes/Sound/SoundAttenuation.h
@@ -73,6 +73,10 @@ struct ENGINE_API FAttenuationSettings
 	UPROPERTY(EditAnywhere, Category=Attenuation, meta=(DisplayName = "dB Attenuation At Max", ClampMax = "0" ))
 	float dBAttenuationAtMax;
 
+	/** At what distance we start transforming into omnidirectional soundsource */
+	UPROPERTY(EditAnywhere, Category=Attenuation, meta=(ClampMin = "0"))
+	float OmniRadius;
+
 	UPROPERTY()
 	float RadiusMin_DEPRECATED;
 
diff --git a/Engine/Source/Runtime/Engine/Private/ActiveSound.cpp b/Engine/Source/Runtime/Engine/Private/ActiveSound.cpp
--- a/Engine/Source/Runtime/Engine/Private/ActiveSound.cpp
+++ b/Engine/Source/Runtime/Engine/Private/ActiveSound.cpp
@@ -201,6 +201,7 @@ void FActiveSound::UpdateWaveInstances( FAudioDevice* AudioDevice, TArray<FWaveI
 		if (bHasAttenuationSettings)
 		{
 			AttenuationSettings.ApplyAttenuation(ParseParams.Transform, Listener.Transform.GetTranslation(), ParseParams.Volume, ParseParams.HighFrequencyGain );
+			ParseParams.OmniRadius = AttenuationSettings.OmniRadius;
 			ParseParams.bUseSpatialization = AttenuationSettings.bSpatialize;
 		}
 
diff --git a/Engine/Source/Runtime/Engine/Private/SoundNodeAttenuation.cpp b/Engine/Source/Runtime/Engine/Private/SoundNodeAttenuation.cpp
--- a/Engine/Source/Runtime/Engine/Private/SoundNodeAttenuation.cpp
+++ b/Engine/Source/Runtime/Engine/Private/SoundNodeAttenuation.cpp
@@ -108,6 +108,7 @@ void USoundNodeAttenuation::ParseNodes( FAudioDevice* AudioDevice, const UPTRINT
 	{
 		const FListener& Listener = AudioDevice->Listeners[0];
 		Settings->ApplyAttenuation(UpdatedParseParams.Transform, Listener.Transform.GetTranslation(), UpdatedParseParams.Volume, UpdatedParseParams.HighFrequencyGain);
+		UpdatedParseParams.OmniRadius = Settings->OmniRadius;
 		UpdatedParseParams.bUseSpatialization |= Settings->bSpatialize;
 	}
 	else
diff --git a/Engine/Source/Runtime/Engine/Private/SoundWave.cpp b/Engine/Source/Runtime/Engine/Private/SoundWave.cpp
--- a/Engine/Source/Runtime/Engine/Private/SoundWave.cpp
+++ b/Engine/Source/Runtime/Engine/Private/SoundWave.cpp
@@ -437,6 +437,7 @@ void USoundWave::Parse( FAudioDevice* AudioDevice, const UPTRINT NodeWaveInstanc
 		WaveInstance->bApplyRadioFilter = ActiveSound.bApplyRadioFilter;
 		WaveInstance->StartTime = ParseParams.StartTime;
 		WaveInstance->UserIndex = ActiveSound.UserIndex;
+		WaveInstance->OmniRadius = ParseParams.OmniRadius;
 
 		bool bAlwaysPlay = false;
 
diff --git a/Engine/Source/Runtime/Engine/Public/Audio.h b/Engine/Source/Runtime/Engine/Public/Audio.h
--- a/Engine/Source/Runtime/Engine/Public/Audio.h
+++ b/Engine/Source/Runtime/Engine/Public/Audio.h
@@ -219,6 +219,8 @@ struct FWaveInstance
 	FVector				Velocity;
 	/** Current location */
 	FVector				Location;
+	/** At what distance we start transforming into omnidirectional soundsource */
+	float				OmniRadius;
 	/** Cached type hash */
 	uint32				TypeHash;
 	/** Hash value for finding the wave instance based on the path through the cue to get to it */
diff --git a/Engine/Source/Runtime/Windows/XAudio2/Private/XAudio2Source.cpp b/Engine/Source/Runtime/Windows/XAudio2/Private/XAudio2Source.cpp
--- a/Engine/Source/Runtime/Windows/XAudio2/Private/XAudio2Source.cpp
+++ b/Engine/Source/Runtime/Windows/XAudio2/Private/XAudio2Source.cpp
@@ -412,9 +412,13 @@ void FXAudio2SoundSource::GetChannelVolumes( float ChannelVolumes[CHANNELOUT_COU
 		{
 			// Calculate direction from listener to sound, where the sound is at the origin if unspatialised.
 			FVector Direction = FVector::ZeroVector;
+			float NormalizedOmniRadius = 0;
 			if( WaveInstance->bUseSpatialization )
 			{
-				Direction = AudioDevice->InverseTransform.TransformPosition( WaveInstance->Location ).SafeNormal();
+				FVector UnnormalizedDirection = AudioDevice->InverseTransform.TransformPosition( WaveInstance->Location );
+				Direction = UnnormalizedDirection.SafeNormal();
+				float Distance = UnnormalizedDirection.Size();
+				NormalizedOmniRadius = (Distance > 0) ? (WaveInstance->OmniRadius / Distance) : FLT_MAX;
 			}
 
 			// Calculate 5.1 channel volume
@@ -448,7 +452,7 @@ void FXAudio2SoundSource::GetChannelVolumes( float ChannelVolumes[CHANNELOUT_COU
 			ChannelVolumes[CHANNELOUT_RADIO] = 0.0f;
 
 			// Call the spatialisation magic
-			AudioDevice->SpatializationHelper.CalculateDolbySurroundRate( OrientFront, ListenerPosition, EmitterPosition, ChannelVolumes );
+			AudioDevice->SpatializationHelper.CalculateDolbySurroundRate( OrientFront, ListenerPosition, EmitterPosition, NormalizedOmniRadius, ChannelVolumes );
 
 			// Handle any special post volume processing
 			if( WaveInstance->bApplyRadioFilter )
@@ -1213,7 +1217,7 @@ void FSpatializationHelper::Init()
  * @param	OutVolumes				An array of floats with one volume for each output channel.
  * @param	OutReverbLevel			The reverb volume
  */
-void FSpatializationHelper::CalculateDolbySurroundRate( const FVector& OrientFront, const FVector& ListenerPosition, const FVector& EmitterPosition, float* OutVolumes  )
+void FSpatializationHelper::CalculateDolbySurroundRate( const FVector& OrientFront, const FVector& ListenerPosition, const FVector& EmitterPosition, float OmniRadius, float* OutVolumes  )
 {
 	uint32 CalculateFlags = X3DAUDIO_CALCULATE_MATRIX | X3DAUDIO_CALCULATE_DOPPLER | X3DAUDIO_CALCULATE_REVERB;
 
@@ -1226,6 +1230,8 @@ void FSpatializationHelper::CalculateDolbySurroundRate( const FVector& OrientFro
 	Emitter.Position.x = EmitterPosition.X;
 	Emitter.Position.y = EmitterPosition.Y;
 	Emitter.Position.z = EmitterPosition.Z;
+	Emitter.InnerRadius = OmniRadius*OmniRadius;
+	Emitter.InnerRadiusAngle = 0;
 
 	X3DAudioCalculate( X3DInstance, &Listener, &Emitter, CalculateFlags, &DSPSettings );
 
diff --git a/Engine/Source/Runtime/Windows/XAudio2/Private/XAudio2Support.h b/Engine/Source/Runtime/Windows/XAudio2/Private/XAudio2Support.h
--- a/Engine/Source/Runtime/Windows/XAudio2/Private/XAudio2Support.h
+++ b/Engine/Source/Runtime/Windows/XAudio2/Private/XAudio2Support.h
@@ -452,7 +452,7 @@ public:
 	 * @param	OutVolumes				An array of floats with one volume for each output channel.
 	 * @param	OutReverbLevel			The reverb volume
 	 */
-	void CalculateDolbySurroundRate( const FVector& OrientFront, const FVector& ListenerPosition, const FVector& EmitterPosition, float* OutVolumes );
+	void CalculateDolbySurroundRate( const FVector& OrientFront, const FVector& ListenerPosition, const FVector& EmitterPosition, float OmniRadius, float* OutVolumes );
 };
 
 /** Variables required for the early init */
