Fix for compile error on Win32.

diff --git a/Engine/Source/Runtime/Windows/D3D11RHI/Private/D3D11RHI.cpp b/Engine/Source/Runtime/Windows/D3D11RHI/Private/D3D11RHI.cpp
--- a/Engine/Source/Runtime/Windows/D3D11RHI/Private/D3D11RHI.cpp
+++ b/Engine/Source/Runtime/Windows/D3D11RHI/Private/D3D11RHI.cpp
@@ -29,7 +29,7 @@ extern void UniformBufferBeginFrame();
 // http://developer.download.nvidia.com/devzone/devcenter/gamegraphics/files/OptimusRenderingPolicies.pdf
 // The following line is to favor the high performance NVIDIA GPU if there are multiple GPUs
 // Has to be .exe module to be correctly detected.
-extern "C" { _declspec(dllexport) uint32 NvOptimusEnablement = 0x00000001; }
+// extern "C" { _declspec(dllexport) uint32 NvOptimusEnablement = 0x00000001; }
 
 void FD3D11DynamicRHI::RHIBeginFrame()
 {
