Fix memcmp where memcpy was meant from static analysis
ttp#332368

diff --git a/Engine/Source/Runtime/Windows/D3D11RHI/Private/D3D11StateCachePrivate.h b/Engine/Source/Runtime/Windows/D3D11RHI/Private/D3D11StateCachePrivate.h
--- a/Engine/Source/Runtime/Windows/D3D11RHI/Private/D3D11StateCachePrivate.h
+++ b/Engine/Source/Runtime/Windows/D3D11RHI/Private/D3D11StateCachePrivate.h
@@ -543,7 +543,7 @@ template <EShaderFrequency ShaderFrequency>
 			CurrentBlendState->AddRef();
 		}
 		*SampleMask = CurrentBlendSampleMask;
-		FMemory::Memcmp(BlendFactor, CurrentBlendFactor, sizeof(CurrentBlendFactor));
+		FMemory::Memcpy(BlendFactor, CurrentBlendFactor, sizeof(CurrentBlendFactor));
 #else
 		Direct3DDeviceIMContext->OMGetBlendState(BlendState, BlendFactor, SampleMask);
 #endif
