Fix XB1.

diff --git a/Engine/Source/Runtime/Windows/D3D11RHI/Private/D3D11Commands.cpp b/Engine/Source/Runtime/Windows/D3D11RHI/Private/D3D11Commands.cpp
--- a/Engine/Source/Runtime/Windows/D3D11RHI/Private/D3D11Commands.cpp
+++ b/Engine/Source/Runtime/Windows/D3D11RHI/Private/D3D11Commands.cpp
@@ -39,24 +39,27 @@ DECLARE_ISBOUNDSHADER(ComputeShader)
 #define VALIDATE_BOUND_SHADER(s)
 #endif
 
-
-#define GPA_WINDOWS 1
-#include <GPUPerfAPI/Gpa.h>
+#if !PLATFORM_XBOXONE
+	#define GPA_WINDOWS 1
+	#include <GPUPerfAPI/Gpa.h>
+#endif
 
 void FD3D11DynamicRHI::RHIGpuTimeBegin(uint32 Hash, bool bCompute)
 {
-	char Str[256];
-	if(GpaBegin(Str, Hash, bCompute, (void*)Direct3DDevice))
-	{
-		OutputDebugStringA(Str);
-	}
-	return;
+	#if !PLATFORM_XBOXONE
+		char Str[256];
+		if(GpaBegin(Str, Hash, bCompute, (void*)Direct3DDevice))
+		{
+			OutputDebugStringA(Str);
+		}
+	#endif
 }
 
 void FD3D11DynamicRHI::RHIGpuTimeEnd(uint32 Hash, bool bCompute)
 {
-	GpaEnd(Hash, bCompute);
-	return;
+	#if !PLATFORM_XBOXONE
+		GpaEnd(Hash, bCompute);
+	#endif
 }
 
 
