#UE4:
* fixed TTP 333641 LIVE: Severe perf decrease from 4.0.2 to 4.1
  heuristic is now rejecting software renderer

diff --git a/Engine/Source/Runtime/Windows/D3D11RHI/Private/Windows/WindowsD3D11Device.cpp b/Engine/Source/Runtime/Windows/D3D11RHI/Private/Windows/WindowsD3D11Device.cpp
--- a/Engine/Source/Runtime/Windows/D3D11RHI/Private/Windows/WindowsD3D11Device.cpp
+++ b/Engine/Source/Runtime/Windows/D3D11RHI/Private/Windows/WindowsD3D11Device.cpp
@@ -254,6 +254,9 @@ void FD3D11DynamicRHIModule::FindAdapter()
 				const bool bIsIntegrated = FCString::Stristr(AdapterDesc.Description,TEXT("Intel")) != 0;
 				// PerfHUD is for performance profiling
 				const bool bIsPerfHUD = !FCString::Stricmp(AdapterDesc.Description,TEXT("NVIDIA PerfHUD"));
+				// Software renderer
+				const bool bIsSoftware = !FCString::Stricmp(AdapterDesc.Description,TEXT("Software Adapter"))
+					|| !FCString::Stricmp(AdapterDesc.Description,TEXT("Microsoft Basic Render Driver"));
 
 				FD3D11Adapter CurrentAdapter(AdapterIndex, ActualFeatureLevel);
 
@@ -263,6 +266,14 @@ void FD3D11DynamicRHIModule::FindAdapter()
 					continue;
 				}
 
+				if(bIsSoftware)
+				{
+					// We don't want software renderer. Ideally we specify D3D_DRIVER_TYPE_HARDWARE on creation but
+					// wen we specify an adapter we need to specify D3D_DRIVER_TYPE_UNKNOWN (otherwise the call fails).
+					// We cannot check the device type later (seems this is missing functionality in D3D).
+					continue;
+				}
+
 				if(CVarValue >= 0 && AdapterIndex != CVarValue)
 				{
 					// the user wants a specific adapter, not this one
