* Improved heuristics to detect Intel integrated hardware, prevent "Microsoft Basic Render Driver" from being used.

diff --git a/Engine/Source/Runtime/Windows/D3D11RHI/Private/Windows/WindowsD3D11Device.cpp b/Engine/Source/Runtime/Windows/D3D11RHI/Private/Windows/WindowsD3D11Device.cpp
--- a/Engine/Source/Runtime/Windows/D3D11RHI/Private/Windows/WindowsD3D11Device.cpp
+++ b/Engine/Source/Runtime/Windows/D3D11RHI/Private/Windows/WindowsD3D11Device.cpp
@@ -135,6 +135,10 @@ static bool SafeTestD3D11CreateDevice(IDXGIAdapter* Adapter,D3D_FEATURE_LEVEL Ma
 
 	__try
 	{
+		// We don't want software renderer. Ideally we specify D3D_DRIVER_TYPE_HARDWARE on creation but
+		// when we specify an adapter we need to specify D3D_DRIVER_TYPE_UNKNOWN (otherwise the call fails).
+		// We cannot check the device type later (seems this is missing functionality in D3D).
+
 		if(SUCCEEDED(D3D11CreateDevice(
 			Adapter,
 			D3D_DRIVER_TYPE_UNKNOWN,
@@ -195,6 +199,23 @@ const TCHAR* GetFeatureLevelString(D3D_FEATURE_LEVEL FeatureLevel)
 	return TEXT("X_X");
 }
 
+static uint32 CountAdapterOutputs(TRefCountPtr<IDXGIAdapter>& Adapter)
+{
+	uint32 OutputCount = 0;
+	for(;;)
+	{
+		TRefCountPtr<IDXGIOutput> Output;
+		HRESULT hr = Adapter->EnumOutputs(OutputCount, Output.GetInitReference());
+		if(FAILED(hr))
+		{
+			break;
+		}
+		++OutputCount;
+	}
+
+	return OutputCount;
+}
+
 void FD3D11DynamicRHIModule::FindAdapter()
 {
 	// Once we chosen one we don't need to do it again.
@@ -224,6 +245,10 @@ void FD3D11DynamicRHIModule::FindAdapter()
 	FD3D11Adapter FirstWithoutIntegratedAdapter;
 	FD3D11Adapter FirstAdapter;
 
+	bool bIsAnyAMD = false;
+	bool bIsAnyIntel = false;
+	bool bIsAnyNVIDIA = false;
+
 	// Enumerate the DXGIFactory's adapters.
 	for(uint32 AdapterIndex = 0; DXGIFactory->EnumAdapters(AdapterIndex,TempAdapter.GetInitReference()) != DXGI_ERROR_NOT_FOUND; ++AdapterIndex)
 	{
@@ -236,6 +261,7 @@ void FD3D11DynamicRHIModule::FindAdapter()
 				// Log some information about the available D3D11 adapters.
 				DXGI_ADAPTER_DESC AdapterDesc;
 				VERIFYD3D11RESULT(TempAdapter->GetDesc(&AdapterDesc));
+				uint32 OutputCount = CountAdapterOutputs(TempAdapter);
 
 				UE_LOG(LogD3D11RHI, Log,
 					TEXT("Found D3D11 adapter %u: %s (Feature Level %s)"),
@@ -244,33 +270,38 @@ void FD3D11DynamicRHIModule::FindAdapter()
 					GetFeatureLevelString(ActualFeatureLevel)
 					);
 				UE_LOG(LogD3D11RHI, Log,
-					TEXT("Adapter has %uMB of dedicated video memory, %uMB of dedicated system memory, and %uMB of shared system memory"),
+					TEXT("Adapter has %uMB of dedicated video memory, %uMB of dedicated system memory, and %uMB of shared system memory, %d output[s]"),
 					(uint32)(AdapterDesc.DedicatedVideoMemory / (1024*1024)),
 					(uint32)(AdapterDesc.DedicatedSystemMemory / (1024*1024)),
-					(uint32)(AdapterDesc.SharedSystemMemory / (1024*1024))
+					(uint32)(AdapterDesc.SharedSystemMemory / (1024*1024)),
+					OutputCount
 					);
 
-				// We could refine this test but it only matters for systems with multiple graphic cards.
-				const bool bIsIntegrated = FCString::Stristr(AdapterDesc.Description,TEXT("Intel")) != 0;
+				bool bIsAMD = AdapterDesc.VendorId == 0x1002;
+				bool bIsIntel = AdapterDesc.VendorId == 0x8086;
+				bool bIsNVIDIA = AdapterDesc.VendorId == 0x10DE;
+
+				if(bIsAMD) bIsAnyAMD = true;
+				if(bIsIntel) bIsAnyIntel = true;
+				if(bIsNVIDIA) bIsAnyNVIDIA = true;
+
+				// Simple heuristic but without profiling it's hard to do better
+				const bool bIsIntegrated = bIsIntel;
 				// PerfHUD is for performance profiling
 				const bool bIsPerfHUD = !FCString::Stricmp(AdapterDesc.Description,TEXT("NVIDIA PerfHUD"));
-				// Software renderer
-				const bool bIsSoftware = !FCString::Stricmp(AdapterDesc.Description,TEXT("Software Adapter"))
-					|| !FCString::Stricmp(AdapterDesc.Description,TEXT("Microsoft Basic Render Driver"));
 
 				FD3D11Adapter CurrentAdapter(AdapterIndex, ActualFeatureLevel);
 
-				if(bIsPerfHUD && !bAllowPerfHUD)
+				if(!OutputCount)
 				{
-					// we don't allow the PerfHUD adapter
+					// This device has no outputs. Reject it, 
+					// http://msdn.microsoft.com/en-us/library/windows/desktop/bb205075%28v=vs.85%29.aspx#WARP_new_for_Win8
 					continue;
 				}
 
-				if(bIsSoftware)
+				if(bIsPerfHUD && !bAllowPerfHUD)
 				{
-					// We don't want software renderer. Ideally we specify D3D_DRIVER_TYPE_HARDWARE on creation but
-					// wen we specify an adapter we need to specify D3D_DRIVER_TYPE_UNKNOWN (otherwise the call fails).
-					// We cannot check the device type later (seems this is missing functionality in D3D).
+					// we don't allow the PerfHUD adapter
 					continue;
 				}
 
@@ -293,11 +324,11 @@ void FD3D11DynamicRHIModule::FindAdapter()
 		}
 	}
 
-	if(bFavorNonIntegrated)
+	if(bFavorNonIntegrated && (bIsAnyAMD || bIsAnyNVIDIA))
 	{
 		ChosenAdapter = FirstWithoutIntegratedAdapter;
 
-		// We assume Intel is integrated graphics (slower than discrete) and rather take a different
+		// We assume Intel is integrated graphics (slower than discrete) than NVIDIA or AMD cards and rather take a different one
 		if(!ChosenAdapter.IsValid())
 		{
 			ChosenAdapter = FirstAdapter;
