Fixed texture pool log in Win32 apps

diff --git a/Engine/Source/Runtime/Windows/D3D11RHI/Private/Windows/WindowsD3D11Device.cpp b/Engine/Source/Runtime/Windows/D3D11RHI/Private/Windows/WindowsD3D11Device.cpp
--- a/Engine/Source/Runtime/Windows/D3D11RHI/Private/Windows/WindowsD3D11Device.cpp
+++ b/Engine/Source/Runtime/Windows/D3D11RHI/Private/Windows/WindowsD3D11Device.cpp
@@ -505,7 +505,7 @@ void FD3D11DynamicRHI::InitD3DDevice()
 						// Truncate GTexturePoolSize to MB (but still counted in bytes)
 						GTexturePoolSize = int64(FGenericPlatformMath::TruncFloat(PoolSize / 1024.0f / 1024.0f)) * 1024 * 1024;
 
-						UE_LOG(LogRHI,Log,TEXT("Texture pool is %u MB (%d%% of %u MB)"),
+						UE_LOG(LogRHI,Log,TEXT("Texture pool is %llu MB (%d%% of %llu MB)"),
 							GTexturePoolSize / 1024 / 1024,
 							GPoolSizeVRAMPercentage,
 							GTotalGraphicsMemory / 1024 / 1024);
