* allow VRam allocation on Volume/3D Textures - it's actually supported, reenabled for LPV

diff --git a/Engine/Source/Runtime/Renderer/Private/LightPropagationVolume.cpp b/Engine/Source/Runtime/Renderer/Private/LightPropagationVolume.cpp
--- a/Engine/Source/Runtime/Renderer/Private/LightPropagationVolume.cpp
+++ b/Engine/Source/Runtime/Renderer/Private/LightPropagationVolume.cpp
@@ -802,7 +802,7 @@ void FLightPropagationVolume::InitSettings( const FSceneView& View )
 			LPV_GRIDRES,
 			PF_FloatRGBA,
 			TexCreate_HideInVisualizeTexture,
-			TexCreate_ShaderResource | TexCreate_UAV, // | TexCreate_FastVRAM not yet supported
+			TexCreate_ShaderResource | TexCreate_UAV | TexCreate_FastVRAM,
 			false,
 			1));
 
diff --git a/Engine/Source/Runtime/Windows/D3D11RHI/Private/D3D11Texture.cpp b/Engine/Source/Runtime/Windows/D3D11RHI/Private/D3D11Texture.cpp
--- a/Engine/Source/Runtime/Windows/D3D11RHI/Private/D3D11Texture.cpp
+++ b/Engine/Source/Runtime/Windows/D3D11RHI/Private/D3D11Texture.cpp
@@ -865,9 +865,6 @@ FD3D11Texture3D* FD3D11DynamicRHI::CreateD3D11Texture3D(uint32 SizeX,uint32 Size
 		bCreateRTV = true;
 	}
 
-	// not yet supported, would require FFastVRAMAllocator::GetFastVRAMAllocator()->AllocTexture3D
-	check(!(Flags & TexCreate_FastVRAM));
-
 	// Set up the texture bind flags.
 	check(!(Flags & TexCreate_DepthStencilTargetable));
 	check(!(Flags & TexCreate_ResolveTargetable));
