* XBoxOne: expose ESRam allocation to visualizetexture log printout and outside RHI
* XBoxOne: added RangeAllocators for ESRam (not used yet)

diff --git a/Engine/Source/Runtime/EmptyRHI/Private/EmptyTexture.cpp b/Engine/Source/Runtime/EmptyRHI/Private/EmptyTexture.cpp
--- a/Engine/Source/Runtime/EmptyRHI/Private/EmptyTexture.cpp
+++ b/Engine/Source/Runtime/EmptyRHI/Private/EmptyTexture.cpp
@@ -138,6 +138,10 @@ FTexture3DRHIRef FEmptyDynamicRHI::RHICreateTexture3D(uint32 SizeX, uint32 SizeY
 	return new FEmptyTexture3D((EPixelFormat)Format, SizeX, SizeY, SizeZ, NumMips, Flags, BulkData);
 }
 
+void FEmptyDynamicRHI::RHIGetResourceInfo(FTextureRHIParamRef Ref, FRHIResourceInfo& OutInfo)
+{
+}
+
 void FEmptyDynamicRHI::RHIGenerateMips(FTextureRHIParamRef SourceSurfaceRHI)
 {
 
diff --git a/Engine/Source/Runtime/OpenGLDrv/Private/OpenGLTexture.cpp b/Engine/Source/Runtime/OpenGLDrv/Private/OpenGLTexture.cpp
--- a/Engine/Source/Runtime/OpenGLDrv/Private/OpenGLTexture.cpp
+++ b/Engine/Source/Runtime/OpenGLDrv/Private/OpenGLTexture.cpp
@@ -1350,6 +1350,10 @@ FTexture3DRHIRef FOpenGLDynamicRHI::RHICreateTexture3D(uint32 SizeX,uint32 SizeY
 	return Texture;
 }
 
+void FOpenGLDynamicRHI::RHIGetResourceInfo(FTextureRHIParamRef Ref, FRHIResourceInfo& OutInfo)
+{
+}
+
 FShaderResourceViewRHIRef FOpenGLDynamicRHI::RHICreateShaderResourceView(FTexture2DRHIParamRef Texture2DRHI, uint8 MipLevel)
 {
 	DYNAMIC_CAST_OPENGLRESOURCE(Texture2D,Texture2D);
diff --git a/Engine/Source/Runtime/RHI/Public/RHI.h b/Engine/Source/Runtime/RHI/Public/RHI.h
--- a/Engine/Source/Runtime/RHI/Public/RHI.h
+++ b/Engine/Source/Runtime/RHI/Public/RHI.h
@@ -1208,6 +1208,27 @@ enum ETextureCreateFlags
 	TexCreate_Shared = 1 << 24,
 };
 
+struct FVRamAllocation
+{
+	FVRamAllocation(uint32 InAllocationStart = 0, uint32 InAllocationSize = 0)
+		: AllocationStart(InAllocationStart)
+		, AllocationSize(InAllocationSize)
+	{
+	}
+
+	bool IsValid() { return AllocationSize > 0; }
+
+	// in bytes
+	uint32 AllocationStart;
+	// in bytes
+	uint32 AllocationSize;
+};
+
+struct FRHIResourceInfo
+{
+	FVRamAllocation VRamAllocation;
+};
+
 // Forward-declaration.
 struct FResolveParams;
 
diff --git a/Engine/Source/Runtime/RHI/Public/RHIMethods.h b/Engine/Source/Runtime/RHI/Public/RHIMethods.h
--- a/Engine/Source/Runtime/RHI/Public/RHIMethods.h
+++ b/Engine/Source/Runtime/RHI/Public/RHIMethods.h
@@ -498,6 +498,16 @@ DEFINE_RHIMETHOD_7(
 	return,return new FRHITexture3D(SizeX,SizeY,SizeZ,NumMips,(EPixelFormat)Format,Flags);
 	);
 
+/**
+ * @param Ref may be 0
+ */
+DEFINE_RHIMETHOD_2(
+	void,RHIGetResourceInfo,
+	FTextureRHIParamRef,Ref,
+	FRHIResourceInfo&,OutInfo,
+	return,return;
+	);
+
 /**
 * Creates a shader resource view for a 2d texture, viewing only a single
 * mip level. Useful when rendering to one mip while sampling from another.
@@ -509,7 +519,7 @@ DEFINE_RHIMETHOD_2(
 	return,return new FRHIShaderResourceView();
 	);
 
-/**
+/**FRHIResourceInfo
 * Creates a shader resource view for a 2d texture, with a different
 * format from the original.  Useful when sampling stencil.
 */
@@ -520,7 +530,7 @@ DEFINE_RHIMETHOD_4(
 	uint8, NumMipLevels,
 	uint8, Format,		
 	return,return new FRHIShaderResourceView();
-);
+	);
 
 /**
 * Generates mip maps for a texture.
diff --git a/Engine/Source/Runtime/RHI/Public/RHIResources.h b/Engine/Source/Runtime/RHI/Public/RHIResources.h
--- a/Engine/Source/Runtime/RHI/Public/RHIResources.h
+++ b/Engine/Source/Runtime/RHI/Public/RHIResources.h
@@ -167,7 +167,9 @@ public:
 	uint32 GetNumSamples() const { return NumSamples; }
 
 	/** @return Whether the texture is multi sampled. */
-	bool IsMultisampled() const { return NumSamples > 1; }		
+	bool IsMultisampled() const { return NumSamples > 1; }
+
+	FRHIResourceInfo ResourceInfo;
 
 private:
 	uint32 NumMips;
diff --git a/Engine/Source/Runtime/Renderer/Private/PostProcess/RenderTargetPool.h b/Engine/Source/Runtime/Renderer/Private/PostProcess/RenderTargetPool.h
--- a/Engine/Source/Runtime/Renderer/Private/PostProcess/RenderTargetPool.h
+++ b/Engine/Source/Runtime/Renderer/Private/PostProcess/RenderTargetPool.h
@@ -6,6 +6,53 @@
 
 #pragma once
 
+/** The reference to a pooled render target, use like this: TRefCountPtr<IPooledRenderTarget> */
+struct FPooledRenderTarget : public IPooledRenderTarget
+{
+	/** constructor */
+	FPooledRenderTarget(const FPooledRenderTargetDesc& InDesc) 
+		: NumRefs(0)
+		, Desc(InDesc)
+		, UnusedForNFrames(0)
+	{
+	}
+
+	virtual ~FPooledRenderTarget()
+	{
+		check(!NumRefs);
+		RenderTargetItem.SafeRelease();
+	}
+
+	uint32 GetUnusedForNFrames() const { return UnusedForNFrames; }
+
+	// interface IPooledRenderTarget --------------
+
+	virtual uint32 AddRef() const;
+	virtual uint32 Release() const;
+	virtual uint32 GetRefCount() const;
+	virtual bool IsFree() const;
+	virtual void SetDebugName(const TCHAR *InName);
+	virtual const FPooledRenderTargetDesc& GetDesc() const;
+	virtual uint32 ComputeMemorySize() const;
+
+// todo private:
+	FVRamAllocation VRamAllocation;
+
+private:
+
+	/** For pool management (only if NumRef == 0 the element can be reused) */
+	mutable int32 NumRefs;
+	/** All necessary data to create the render target */
+	FPooledRenderTargetDesc Desc;
+	/** Allows to defer the release to save performance on some hardware (DirectX) */
+	uint32 UnusedForNFrames;
+
+	/** @return true:release this one, false otherwise */
+	bool OnFrameStart();
+
+	friend class FRenderTargetPool;
+};
+
 enum ERenderTargetPoolEventType
 {
 	ERTPE_Alloc,
@@ -16,7 +63,7 @@
 struct FRenderTargetPoolEvent
 {
 	// constructor for ERTPE_Alloc
-	FRenderTargetPoolEvent(uint32 InPoolEntryId, uint32 InTimeStep, IPooledRenderTarget* InPointer)
+	FRenderTargetPoolEvent(uint32 InPoolEntryId, uint32 InTimeStep, FPooledRenderTarget* InPointer)
 		: PoolEntryId(InPoolEntryId)
 		, TimeStep(InTimeStep)
 		, Pointer(InPointer)
@@ -27,6 +74,7 @@
 	{
 		Desc = Pointer->GetDesc();
 		SizeInBytes = Pointer->ComputeMemorySize();
+		VRamAllocation = Pointer->VRamAllocation;
 	}
 
 	// constructor for ERTPE_Alloc
@@ -84,7 +132,9 @@
 	//
 	uint32 TimeStep;
 	// valid EventType==ERTPE_Alloc, 0 if not set
-	IPooledRenderTarget* Pointer;
+	FPooledRenderTarget* Pointer;
+	//
+	FVRamAllocation VRamAllocation;
 	// valid if EventType==ERTPE_Phase TEXT("") if not set
 	FString PhaseName;
 	// valid if EventType==ERTPE_Alloc || EventType==ERTPE_Dealloc
@@ -105,50 +155,6 @@
 };
 
 
-/** The reference to a pooled render target, use like this: TRefCountPtr<IPooledRenderTarget> */
-struct FPooledRenderTarget : public IPooledRenderTarget
-{
-	/** constructor */
-	FPooledRenderTarget(const FPooledRenderTargetDesc& InDesc) 
-		: NumRefs(0)
-		, Desc(InDesc)
-		, UnusedForNFrames(0)
-	{
-	}
-
-	virtual ~FPooledRenderTarget()
-	{
-		check(!NumRefs);
-		RenderTargetItem.SafeRelease();
-	}
-
-	uint32 GetUnusedForNFrames() const { return UnusedForNFrames; }
-
-	// interface IPooledRenderTarget --------------
-
-	virtual uint32 AddRef() const;
-	virtual uint32 Release() const;
-	virtual uint32 GetRefCount() const;
-	virtual bool IsFree() const;
-	virtual void SetDebugName(const TCHAR *InName);
-	virtual const FPooledRenderTargetDesc& GetDesc() const;
-	virtual uint32 ComputeMemorySize() const;
-
-private:
-
-	/** For pool management (only if NumRef == 0 the element can be reused) */
-	mutable int32 NumRefs;
-	/** All necessary data to create the render target */
-	FPooledRenderTargetDesc Desc;
-	/** Allows to defer the release to save performance on some hardware (DirectX) */
-	uint32 UnusedForNFrames;
-
-	/** @return true:release this one, false otherwise */
-	bool OnFrameStart();
-
-	friend class FRenderTargetPool;
-};
-
 #include "VisualizeTexture.h"
 
 
diff --git a/Engine/Source/Runtime/Renderer/Private/PostProcess/VisualizeTexture.cpp b/Engine/Source/Runtime/Renderer/Private/PostProcess/VisualizeTexture.cpp
--- a/Engine/Source/Runtime/Renderer/Private/PostProcess/VisualizeTexture.cpp
+++ b/Engine/Source/Runtime/Renderer/Private/PostProcess/VisualizeTexture.cpp
@@ -727,6 +727,35 @@ void FVisualizeTexture::DebugLog(bool bExtended)
 					check(0);
 				}
 
+				if(Desc.Flags & TexCreate_FastVRAM)
+				{
+					FRHIResourceInfo Info;
+
+					FTextureRHIParamRef Texture = RT->GetRenderTargetItem().ShaderResourceTexture;
+
+					if(!IsValidRef(Texture))
+					{
+						RT->GetRenderTargetItem().TargetableTexture;
+					}
+
+					if(IsValidRef(Texture))
+					{
+						RHIGetResourceInfo(Texture, Info);
+					}
+
+					if(Info.VRamAllocation.AllocationSize)
+					{
+						// note we do KB for more readable numbers but this can cause quantization loss
+						Element.Line += FString::Printf(TEXT(" VRamInKB(Start/Size):%d/%d"), 
+							Info.VRamAllocation.AllocationStart / 1024, 
+							(Info.VRamAllocation.AllocationSize + 1023) / 1024);
+					}
+					else
+					{
+						Element.Line += TEXT(" VRamInKB(Start/Size):<NONE>");
+					}
+				}
+
 				SortedLines.Add(Element);
 			}
 		}
diff --git a/Engine/Source/Runtime/Windows/D3D11RHI/Private/D3D11RHIPrivate.h b/Engine/Source/Runtime/Windows/D3D11RHI/Private/D3D11RHIPrivate.h
--- a/Engine/Source/Runtime/Windows/D3D11RHI/Private/D3D11RHIPrivate.h
+++ b/Engine/Source/Runtime/Windows/D3D11RHI/Private/D3D11RHIPrivate.h
@@ -482,7 +482,8 @@ protected:
 	int32 ChosenAdapter;
 
 	template<typename BaseResourceType>
-	TD3D11Texture2D<BaseResourceType>* CreateD3D11Texture2D(uint32 SizeX,uint32 SizeY,uint32 SizeZ,bool bTextureArray,bool CubeTexture,uint8 Format,uint32 NumMips,uint32 NumSamples,uint32 Flags,FResourceBulkDataInterface* BulkData = NULL);
+	TD3D11Texture2D<BaseResourceType>* CreateD3D11Texture2D(uint32 SizeX,uint32 SizeY,uint32 SizeZ,bool bTextureArray,bool CubeTexture,uint8 Format,
+		uint32 NumMips,uint32 NumSamples,uint32 Flags,FResourceBulkDataInterface* BulkData = NULL);
 
 	FD3D11Texture3D* CreateD3D11Texture3D(uint32 SizeX,uint32 SizeY,uint32 SizeZ,uint8 Format,uint32 NumMips,uint32 Flags,FResourceBulkDataInterface* BulkData);
 
@@ -715,25 +716,25 @@ public:
 	/**
 	 *	IMPORTANT: This function CAN modify the TextureDesc!
 	 */
-	virtual bool AllocTexture2D(D3D11_TEXTURE2D_DESC& TextureDesc)
+	virtual FVRamAllocation AllocTexture2D(D3D11_TEXTURE2D_DESC& TextureDesc)
 	{
-		return false;
+		return FVRamAllocation();
 	}
 
 	/**
 	 *	IMPORTANT: This function CAN modify the TextureDesc!
 	 */
-	virtual bool AllocTexture3D(D3D11_TEXTURE3D_DESC& TextureDesc)
+	virtual FVRamAllocation AllocTexture3D(D3D11_TEXTURE3D_DESC& TextureDesc)
 	{
-		return false;
+		return FVRamAllocation();
 	}
 
 	/**
 	 *	IMPORTANT: This function CAN modify the BufferDesc!
 	 */
-	virtual bool AllocUAVBuffer(D3D11_BUFFER_DESC& BufferDesc)
+	virtual FVRamAllocation AllocUAVBuffer(D3D11_BUFFER_DESC& BufferDesc)
 	{
-		return false;
+		return FVRamAllocation();
 	}
 
 	template< typename t_A, typename t_B >
diff --git a/Engine/Source/Runtime/Windows/D3D11RHI/Private/D3D11Texture.cpp b/Engine/Source/Runtime/Windows/D3D11RHI/Private/D3D11Texture.cpp
--- a/Engine/Source/Runtime/Windows/D3D11RHI/Private/D3D11Texture.cpp
+++ b/Engine/Source/Runtime/Windows/D3D11RHI/Private/D3D11Texture.cpp
@@ -435,7 +435,8 @@ void SafeCreateTexture2D(ID3D11Device* Direct3DDevice, const D3D11_TEXTURE2D_DES
 }
 
 template<typename BaseResourceType>
-TD3D11Texture2D<BaseResourceType>* FD3D11DynamicRHI::CreateD3D11Texture2D(uint32 SizeX,uint32 SizeY,uint32 SizeZ,bool bTextureArray,bool bCubeTexture,uint8 Format,uint32 NumMips,uint32 NumSamples,uint32 Flags,FResourceBulkDataInterface* BulkData)
+TD3D11Texture2D<BaseResourceType>* FD3D11DynamicRHI::CreateD3D11Texture2D(uint32 SizeX,uint32 SizeY,uint32 SizeZ,bool bTextureArray,bool bCubeTexture,uint8 Format,
+	uint32 NumMips,uint32 NumSamples,uint32 Flags, FResourceBulkDataInterface* BulkData)
 {
 	check(SizeX > 0 && SizeY > 0 && NumMips > 0);
 
@@ -592,11 +593,13 @@ TD3D11Texture2D<BaseResourceType>* FD3D11DynamicRHI::CreateD3D11Texture2D(uint32
 		bPooledTexture = false;
 	}
 
+	FVRamAllocation VRamAllocation;
+
 	if (FPlatformProperties::SupportsFastVRAMMemory())
 	{
 		if (Flags & TexCreate_FastVRAM)
 		{
-			FFastVRAMAllocator::GetFastVRAMAllocator()->AllocTexture2D(TextureDesc);
+			VRamAllocation = FFastVRAMAllocator::GetFastVRAMAllocator()->AllocTexture2D(TextureDesc);
 		}
 	}
 
@@ -819,6 +822,8 @@ TD3D11Texture2D<BaseResourceType>* FD3D11DynamicRHI::CreateD3D11Texture2D(uint32
 #endif
 		);
 
+	Texture2D->ResourceInfo.VRamAllocation = VRamAllocation;
+
 	D3D11TextureAllocated(*Texture2D);
 
 	return Texture2D;
@@ -887,11 +892,13 @@ FD3D11Texture3D* FD3D11DynamicRHI::CreateD3D11Texture3D(uint32 SizeX,uint32 Size
 		}
 	}
 
+	FVRamAllocation VRamAllocation;
+
 	if (FPlatformProperties::SupportsFastVRAMMemory())
 	{
 		if (Flags & TexCreate_FastVRAM)
 		{
-			FFastVRAMAllocator::GetFastVRAMAllocator()->AllocTexture3D(TextureDesc);
+			VRamAllocation = FFastVRAMAllocator::GetFastVRAMAllocator()->AllocTexture3D(TextureDesc);
 		}
 	}
 
@@ -936,6 +943,8 @@ FD3D11Texture3D* FD3D11DynamicRHI::CreateD3D11Texture3D(uint32 SizeX,uint32 Size
 	RenderTargetViews.Add(RenderTargetView);
 	FD3D11Texture3D* Texture3D = new FD3D11Texture3D(this,TextureResource,ShaderResourceView,RenderTargetViews,SizeX,SizeY,SizeZ,NumMips,(EPixelFormat)Format,Flags);
 
+	Texture3D->ResourceInfo.VRamAllocation = VRamAllocation;
+
 	D3D11TextureAllocated(*Texture3D);
 
 	return Texture3D;
@@ -1091,6 +1100,14 @@ FTexture3DRHIRef FD3D11DynamicRHI::RHICreateTexture3D(uint32 SizeX,uint32 SizeY,
 	return CreateD3D11Texture3D(SizeX,SizeY,SizeZ,Format,NumMips,Flags,BulkData);
 }
 
+void FD3D11DynamicRHI::RHIGetResourceInfo(FTextureRHIParamRef Ref, FRHIResourceInfo& OutInfo)
+{
+	if(Ref)
+	{
+		OutInfo = Ref->ResourceInfo;
+	}
+}
+
 FShaderResourceViewRHIRef FD3D11DynamicRHI::RHICreateShaderResourceView(FTexture2DRHIParamRef Texture2DRHI, uint8 MipLevel)
 {
 	DYNAMIC_CAST_D3D11RESOURCE(Texture2D,Texture2D);
@@ -1419,12 +1436,12 @@ void FD3D11DynamicRHI::RHIUpdateTexture3D(FTexture3DRHIParamRef TextureRHI,uint3
 -----------------------------------------------------------------------------*/
 FTextureCubeRHIRef FD3D11DynamicRHI::RHICreateTextureCube(uint32 Size, uint8 Format, uint32 NumMips, uint32 Flags, FResourceBulkDataInterface* BulkData)
 {
-	return CreateD3D11Texture2D<FD3D11BaseTextureCube>(Size,Size,6,false,true,Format,NumMips,1,Flags);
+	return CreateD3D11Texture2D<FD3D11BaseTextureCube>(Size,Size,6,false,true,Format,NumMips,1,Flags,BulkData);
 }
 
 FTextureCubeRHIRef FD3D11DynamicRHI::RHICreateTextureCubeArray(uint32 Size, uint32 ArraySize, uint8 Format, uint32 NumMips, uint32 Flags, FResourceBulkDataInterface* BulkData)
 {
-	return CreateD3D11Texture2D<FD3D11BaseTextureCube>(Size,Size,6 * ArraySize,true,true,Format,NumMips,1,Flags);
+	return CreateD3D11Texture2D<FD3D11BaseTextureCube>(Size,Size,6 * ArraySize,true,true,Format,NumMips,1,Flags,BulkData);
 }
 
 void* FD3D11DynamicRHI::RHILockTextureCubeFace(FTextureCubeRHIParamRef TextureCubeRHI,uint32 FaceIndex,uint32 ArrayIndex,uint32 MipIndex,EResourceLockMode LockMode,uint32& DestStride,bool bLockWithinMiptail)
