Log out the UAV bind flag when texture creation fails

diff --git a/Engine/Source/Runtime/Windows/D3D11RHI/Private/D3D11Util.cpp b/Engine/Source/Runtime/Windows/D3D11RHI/Private/D3D11Util.cpp
--- a/Engine/Source/Runtime/Windows/D3D11RHI/Private/D3D11Util.cpp
+++ b/Engine/Source/Runtime/Windows/D3D11RHI/Private/D3D11Util.cpp
@@ -136,6 +136,11 @@ static FString GetD3D11TextureFlagString(uint32 TextureFlags)
 		TextureFormatText += TEXT("D3D11_BIND_SHADER_RESOURCE ");
 	}
 
+	if (TextureFlags & D3D11_BIND_UNORDERED_ACCESS)
+	{
+		TextureFormatText += TEXT("D3D11_BIND_UNORDERED_ACCESS ");
+	}
+
 	return TextureFormatText;
 }
 
