Make sure all D3D buffer desc structs are fully zeroed out, as some platforms have extra parameters which must be properly initialized.

diff --git a/Engine/Source/Runtime/Windows/D3D11RHI/Private/D3D11IndexBuffer.cpp b/Engine/Source/Runtime/Windows/D3D11RHI/Private/D3D11IndexBuffer.cpp
--- a/Engine/Source/Runtime/Windows/D3D11RHI/Private/D3D11IndexBuffer.cpp
+++ b/Engine/Source/Runtime/Windows/D3D11RHI/Private/D3D11IndexBuffer.cpp
@@ -13,6 +13,7 @@ FIndexBufferRHIRef FD3D11DynamicRHI::RHICreateIndexBuffer(uint32 Stride,uint32 S
 
 	// Describe the index buffer.
 	D3D11_BUFFER_DESC Desc;
+	ZeroMemory( &Desc, sizeof( D3D11_BUFFER_DESC ) );
 	Desc.ByteWidth = Size;
 	Desc.Usage = (InUsage & BUF_AnyDynamic) ? D3D11_USAGE_DYNAMIC : D3D11_USAGE_DEFAULT;
 	Desc.BindFlags = D3D11_BIND_INDEX_BUFFER;
@@ -88,6 +89,7 @@ void* FD3D11DynamicRHI::RHILockIndexBuffer(FIndexBufferRHIParamRef IndexBufferRH
 		{
 			// If the static buffer is being locked for reading, create a staging buffer.
 			D3D11_BUFFER_DESC StagingBufferDesc;
+			ZeroMemory( &StagingBufferDesc, sizeof( D3D11_BUFFER_DESC ) );
 			StagingBufferDesc.ByteWidth = Size;
 			StagingBufferDesc.Usage = D3D11_USAGE_STAGING;
 			StagingBufferDesc.BindFlags = 0;
diff --git a/Engine/Source/Runtime/Windows/D3D11RHI/Private/D3D11RenderTarget.cpp b/Engine/Source/Runtime/Windows/D3D11RHI/Private/D3D11RenderTarget.cpp
--- a/Engine/Source/Runtime/Windows/D3D11RHI/Private/D3D11RenderTarget.cpp
+++ b/Engine/Source/Runtime/Windows/D3D11RHI/Private/D3D11RenderTarget.cpp
@@ -592,6 +592,7 @@ TRefCountPtr<ID3D11Texture2D> FD3D11DynamicRHI::GetStagingTexture(FTextureRHIPar
 
 	// create a temp 2d texture to copy render target to
 	D3D11_TEXTURE2D_DESC Desc;
+	ZeroMemory( &Desc, sizeof( D3D11_TEXTURE2D_DESC ) );
 	Desc.Width = SizeX;
 	Desc.Height = SizeY;
 	Desc.MipLevels = 1;
@@ -1026,6 +1027,7 @@ void FD3D11DynamicRHI::ReadSurfaceDataMSAARaw(FTextureRHIParamRef TextureRHI,FIn
 
 	// Create a non-MSAA render target to resolve individual samples of the source surface to.
 	D3D11_TEXTURE2D_DESC NonMSAADesc;
+	ZeroMemory( &NonMSAADesc, sizeof( D3D11_TEXTURE2D_DESC ) );
 	NonMSAADesc.Width = SizeX;
 	NonMSAADesc.Height = SizeY;
 	NonMSAADesc.MipLevels = 1;
@@ -1054,6 +1056,7 @@ void FD3D11DynamicRHI::ReadSurfaceDataMSAARaw(FTextureRHIParamRef TextureRHI,FIn
 	// Create a CPU-accessible staging texture to copy the resolved sample data to.
 	TRefCountPtr<ID3D11Texture2D> StagingTexture2D;
 	D3D11_TEXTURE2D_DESC StagingDesc;
+	ZeroMemory( &StagingDesc, sizeof( D3D11_TEXTURE2D_DESC ) );
 	StagingDesc.Width = SizeX;
 	StagingDesc.Height = SizeY;
 	StagingDesc.MipLevels = 1;
@@ -1191,6 +1194,7 @@ void FD3D11DynamicRHI::RHIReadSurfaceFloatData(FTextureRHIParamRef TextureRHI,FI
 
 	// create a temp 2d texture to copy render target to
 	D3D11_TEXTURE2D_DESC Desc;
+	ZeroMemory( &Desc, sizeof( D3D11_TEXTURE2D_DESC ) );
 	Desc.Width = SizeX;
 	Desc.Height = SizeY;
 	Desc.MipLevels = 1;
@@ -1266,6 +1270,7 @@ void FD3D11DynamicRHI::RHIRead3DSurfaceFloatData(FTextureRHIParamRef TextureRHI,
 
 	// create a temp 2d texture to copy render target to
 	D3D11_TEXTURE3D_DESC Desc;
+	ZeroMemory( &Desc, sizeof( D3D11_TEXTURE3D_DESC ) );
 	Desc.Width = SizeX;
 	Desc.Height = SizeY;
 	Desc.Depth = SizeZ;
diff --git a/Engine/Source/Runtime/Windows/D3D11RHI/Private/D3D11StructuredBuffer.cpp b/Engine/Source/Runtime/Windows/D3D11RHI/Private/D3D11StructuredBuffer.cpp
--- a/Engine/Source/Runtime/Windows/D3D11RHI/Private/D3D11StructuredBuffer.cpp
+++ b/Engine/Source/Runtime/Windows/D3D11RHI/Private/D3D11StructuredBuffer.cpp
@@ -11,6 +11,7 @@ FStructuredBufferRHIRef FD3D11DynamicRHI::RHICreateStructuredBuffer(uint32 Strid
 	check(Size / Stride > 0 && Size % Stride == 0);
 
 	D3D11_BUFFER_DESC Desc;
+	ZeroMemory( &Desc, sizeof( D3D11_BUFFER_DESC ) );
 	Desc.ByteWidth = Size;
 	Desc.Usage = (InUsage & BUF_AnyDynamic) ? D3D11_USAGE_DYNAMIC : D3D11_USAGE_DEFAULT;
 	Desc.BindFlags = 0;
@@ -115,6 +116,7 @@ void* FD3D11DynamicRHI::RHILockStructuredBuffer(FStructuredBufferRHIParamRef Str
 		{
 			// If the static buffer is being locked for reading, create a staging buffer.
 			D3D11_BUFFER_DESC StagingBufferDesc;
+			ZeroMemory( &StagingBufferDesc, sizeof( D3D11_BUFFER_DESC ) );
 			StagingBufferDesc.ByteWidth = Size;
 			StagingBufferDesc.Usage = D3D11_USAGE_STAGING;
 			StagingBufferDesc.BindFlags = 0;
@@ -181,6 +183,9 @@ void FD3D11DynamicRHI::RHIUnlockStructuredBuffer(FStructuredBufferRHIParamRef St
 			// Copy the contents of the temporary memory buffer allocated for writing into the VB.
 			Direct3DDeviceIMContext->UpdateSubresource(StructuredBuffer->Resource,LockedKey.Subresource,NULL,LockedData->GetData(),LockedData->Pitch,0);
 
+			// Check the copy is finished before freeing...
+			Direct3DDeviceIMContext->Flush();
+
 			// Free the temporary memory buffer.
 			LockedData->FreeData();
 		}
diff --git a/Engine/Source/Runtime/Windows/D3D11RHI/Private/D3D11Texture.cpp b/Engine/Source/Runtime/Windows/D3D11RHI/Private/D3D11Texture.cpp
--- a/Engine/Source/Runtime/Windows/D3D11RHI/Private/D3D11Texture.cpp
+++ b/Engine/Source/Runtime/Windows/D3D11RHI/Private/D3D11Texture.cpp
@@ -523,6 +523,7 @@ TD3D11Texture2D<BaseResourceType>* FD3D11DynamicRHI::CreateD3D11Texture2D(uint32
 
 	// Describe the texture.
 	D3D11_TEXTURE2D_DESC TextureDesc;
+	ZeroMemory( &TextureDesc, sizeof( D3D11_TEXTURE2D_DESC ) );
 	TextureDesc.Width = SizeX;
 	TextureDesc.Height = SizeY;
 	TextureDesc.MipLevels = NumMips;
@@ -841,6 +842,7 @@ FD3D11Texture3D* FD3D11DynamicRHI::CreateD3D11Texture3D(uint32 SizeX,uint32 Size
 
 	// Describe the texture.
 	D3D11_TEXTURE3D_DESC TextureDesc;
+	ZeroMemory( &TextureDesc, sizeof( D3D11_TEXTURE3D_DESC ) );
 	TextureDesc.Width = SizeX;
 	TextureDesc.Height = SizeY;
 	TextureDesc.Depth = SizeZ;
diff --git a/Engine/Source/Runtime/Windows/D3D11RHI/Private/D3D11Util.cpp b/Engine/Source/Runtime/Windows/D3D11RHI/Private/D3D11Util.cpp
--- a/Engine/Source/Runtime/Windows/D3D11RHI/Private/D3D11Util.cpp
+++ b/Engine/Source/Runtime/Windows/D3D11RHI/Private/D3D11Util.cpp
@@ -292,6 +292,7 @@ FD3D11DynamicBuffer::~FD3D11DynamicBuffer()
 void FD3D11DynamicBuffer::InitRHI()
 {
 	D3D11_BUFFER_DESC Desc;
+	ZeroMemory( &Desc, sizeof( D3D11_BUFFER_DESC ) );
 	Desc.Usage = D3D11_USAGE_DYNAMIC;
 	Desc.BindFlags = BindFlags;
 	Desc.CPUAccessFlags = D3D11_CPU_ACCESS_WRITE;
@@ -332,6 +333,7 @@ void* FD3D11DynamicBuffer::Lock(uint32 Size)
 
 		TRefCountPtr<ID3D11Buffer> Buffer;
 		D3D11_BUFFER_DESC Desc;
+		ZeroMemory( &Desc, sizeof( D3D11_BUFFER_DESC ) );
 		Desc.Usage = D3D11_USAGE_DYNAMIC;
 		Desc.BindFlags = BindFlags;
 		Desc.CPUAccessFlags = D3D11_CPU_ACCESS_WRITE;
@@ -346,7 +348,7 @@ void* FD3D11DynamicBuffer::Lock(uint32 Size)
 
 	LockedBufferIndex = BufferIndex;
 	D3D11_MAPPED_SUBRESOURCE MappedSubresource;
-	D3DRHI->GetDeviceContext()->Map(Buffers[BufferIndex],0,D3D11_MAP_WRITE_DISCARD,0,&MappedSubresource);
+	VERIFYD3D11RESULT( D3DRHI->GetDeviceContext()->Map( Buffers[BufferIndex],0,D3D11_MAP_WRITE_DISCARD,0,&MappedSubresource ) );
 	return MappedSubresource.pData;
 }
 
diff --git a/Engine/Source/Runtime/Windows/D3D11RHI/Private/D3D11VertexBuffer.cpp b/Engine/Source/Runtime/Windows/D3D11RHI/Private/D3D11VertexBuffer.cpp
--- a/Engine/Source/Runtime/Windows/D3D11RHI/Private/D3D11VertexBuffer.cpp
+++ b/Engine/Source/Runtime/Windows/D3D11RHI/Private/D3D11VertexBuffer.cpp
@@ -12,6 +12,7 @@ FVertexBufferRHIRef FD3D11DynamicRHI::RHICreateVertexBuffer(uint32 Size,FResourc
 	check(Size > 0);
 
 	D3D11_BUFFER_DESC Desc;
+	ZeroMemory( &Desc, sizeof( D3D11_BUFFER_DESC ) );
 	Desc.ByteWidth = Size;
 	Desc.Usage = (InUsage & BUF_AnyDynamic) ? D3D11_USAGE_DYNAMIC : D3D11_USAGE_DEFAULT;
 	Desc.BindFlags = D3D11_BIND_VERTEX_BUFFER;
@@ -109,6 +110,7 @@ void* FD3D11DynamicRHI::RHILockVertexBuffer(FVertexBufferRHIParamRef VertexBuffe
 		{
 			// If the static buffer is being locked for reading, create a staging buffer.
 			D3D11_BUFFER_DESC StagingBufferDesc;
+			ZeroMemory( &StagingBufferDesc, sizeof( D3D11_BUFFER_DESC ) );
 			StagingBufferDesc.ByteWidth = Size;
 			StagingBufferDesc.Usage = D3D11_USAGE_STAGING;
 			StagingBufferDesc.BindFlags = 0;
diff --git a/Engine/Source/Runtime/Windows/D3D11RHI/Private/Windows/WindowsD3D11ConstantBuffer.cpp b/Engine/Source/Runtime/Windows/D3D11RHI/Private/Windows/WindowsD3D11ConstantBuffer.cpp
--- a/Engine/Source/Runtime/Windows/D3D11RHI/Private/Windows/WindowsD3D11ConstantBuffer.cpp
+++ b/Engine/Source/Runtime/Windows/D3D11RHI/Private/Windows/WindowsD3D11ConstantBuffer.cpp
@@ -11,6 +11,7 @@ void FWinD3D11ConstantBuffer::InitDynamicRHI()
 {
 	TRefCountPtr<ID3D11Buffer> CBuffer = NULL;
 	D3D11_BUFFER_DESC BufferDesc;
+	ZeroMemory( &BufferDesc, sizeof( D3D11_BUFFER_DESC ) );
 	// Verify constant buffer ByteWidth requirements
 	check(MaxSize <= D3D11_REQ_CONSTANT_BUFFER_ELEMENT_COUNT && (MaxSize % 16) == 0);
 	BufferDesc.ByteWidth = MaxSize;
