* Fixed crash on startup with new shaders, RHIClear was used in resource init
* It was picking the wrong clear method which is fixed now (D3D11 problem)
* added check to find such a case earlier

diff --git a/Engine/Source/Runtime/Windows/D3D11RHI/Private/D3D11Commands.cpp b/Engine/Source/Runtime/Windows/D3D11RHI/Private/D3D11Commands.cpp
--- a/Engine/Source/Runtime/Windows/D3D11RHI/Private/D3D11Commands.cpp
+++ b/Engine/Source/Runtime/Windows/D3D11RHI/Private/D3D11Commands.cpp
@@ -1565,6 +1565,9 @@ void FD3D11DynamicRHI::RHIClearMRT(bool bClearColor,int32 NumClearColors,const F
 		// So we can see when we are taking the slow path
 		SCOPED_DRAW_EVENT(DrawCallClear, DEC_SCENE_ITEMS);
 
+		// we don't support draw call clears before the RHI is initialized, reorder the code or make sure it's not a draw call clear
+		check(GIsRHIInitialized);
+
 		if (CurrentDepthTexture)
 		{
 			// Clear all texture references to this depth buffer
diff --git a/Engine/Source/Runtime/Windows/D3D11RHI/Private/D3D11Device.cpp b/Engine/Source/Runtime/Windows/D3D11RHI/Private/D3D11Device.cpp
--- a/Engine/Source/Runtime/Windows/D3D11RHI/Private/D3D11Device.cpp
+++ b/Engine/Source/Runtime/Windows/D3D11RHI/Private/D3D11Device.cpp
@@ -337,6 +337,24 @@ uint32 FD3D11DynamicRHI::GetMaxMSAAQuality(uint32 SampleCount)
 	return 0xffffffff;
 }
 
+void FD3D11DynamicRHI::SetupAfterDeviceCreation()
+{
+	// without that the first RHIClear would get a scissor rect of (0,0)-(0,0) which means we get a draw call clear 
+	RHISetScissorRect(false, 0, 0, 0, 0);
+
+	UpdateMSAASettings();
+
+	if (GRHISupportsAsyncTextureCreation)
+	{
+		UE_LOG(LogD3D11RHI,Log,TEXT("Async texture creation enabled"));
+	}
+	else
+	{
+		UE_LOG(LogD3D11RHI,Log,TEXT("Async texture creation disabled: %s"),
+			D3D11RHI_ShouldAllowAsyncResourceCreation() ? TEXT("no driver support") : TEXT("disabled by user"));
+	}
+}
+
 void FD3D11DynamicRHI::UpdateMSAASettings()
 {	
 	check(DX_MAX_MSAA_COUNT == 8);
@@ -355,7 +373,6 @@ void FD3D11DynamicRHI::UpdateMSAASettings()
 	AvailableMSAAQualities[8] = 0;
 }
 
-
 void FD3D11DynamicRHI::CleanupD3DDevice()
 {
 	if(GIsRHIInitialized)
diff --git a/Engine/Source/Runtime/Windows/D3D11RHI/Private/D3D11RHIPrivate.h b/Engine/Source/Runtime/Windows/D3D11RHI/Private/D3D11RHIPrivate.h
--- a/Engine/Source/Runtime/Windows/D3D11RHI/Private/D3D11RHIPrivate.h
+++ b/Engine/Source/Runtime/Windows/D3D11RHI/Private/D3D11RHIPrivate.h
@@ -506,7 +506,11 @@ protected:
 	 */
 	void GetBestSupportedMSAASetting( DXGI_FORMAT PlatformFormat, uint32 MSAACount, uint32& OutBestMSAACount, uint32& OutMSAAQualityLevels );
 
-	// called when the device gets initialized
+	// shared code for different D3D11 devices (e.g. PC DirectX11 and XboxOne) called
+	// after device creation and GRHISupportsAsyncTextureCreation was set and before resource init
+	void SetupAfterDeviceCreation();
+
+	// called by SetupAfterDeviceCreation() when the device gets initialized
 	void UpdateMSAASettings();
 
 	// @return 0xffffffff if not not supported
diff --git a/Engine/Source/Runtime/Windows/D3D11RHI/Private/Windows/WindowsD3D11Device.cpp b/Engine/Source/Runtime/Windows/D3D11RHI/Private/Windows/WindowsD3D11Device.cpp
--- a/Engine/Source/Runtime/Windows/D3D11RHI/Private/Windows/WindowsD3D11Device.cpp
+++ b/Engine/Source/Runtime/Windows/D3D11RHI/Private/Windows/WindowsD3D11Device.cpp
@@ -560,11 +560,6 @@ void FD3D11DynamicRHI::InitD3DDevice()
 		D3DPERF_SetOptions(1);
 #endif
 
-		// Determine PF_G8 usage as a render target
-		uint32 G8FormatSupport = 0;
-		VERIFYD3D11RESULT(Direct3DDevice->CheckFormatSupport(DXGI_FORMAT_R8_UNORM,&G8FormatSupport));
-		GSupportsRenderTargetFormat_PF_G8 = !!(G8FormatSupport & D3D11_FORMAT_SUPPORT_RENDER_TARGET);
-
 		// Check for async texture creation support.
 		D3D11_FEATURE_DATA_THREADING ThreadingSupport = {0};
 		VERIFYD3D11RESULT(Direct3DDevice->CheckFeatureSupport(D3D11_FEATURE_THREADING,&ThreadingSupport,sizeof(ThreadingSupport)));
@@ -581,17 +576,7 @@ void FD3D11DynamicRHI::InitD3DDevice()
 			GRHISupportsAsyncTextureCreation = false;
 		}
 
-		if (GRHISupportsAsyncTextureCreation)
-		{
-			UE_LOG(LogD3D11RHI,Log,TEXT("Async texture creation enabled"));
-		}
-		else
-		{
-			UE_LOG(LogD3D11RHI,Log,TEXT("Async texture creation disabled: %s"),
-				D3D11RHI_ShouldAllowAsyncResourceCreation() ? TEXT("no driver support") : TEXT("disabled by user"));
-		}
-
-		UpdateMSAASettings();
+		SetupAfterDeviceCreation();
 
 		// Notify all initialized FRenderResources that there's a valid RHI device to create their RHI resources for now.
 		for(TLinkedList<FRenderResource*>::TIterator ResourceIt(FRenderResource::GetResourceList());ResourceIt;ResourceIt.Next())
