don't try displaying a dialog box from the rendering thread

diff --git a/Engine/Source/Runtime/Windows/D3D11RHI/Private/D3D11Util.cpp b/Engine/Source/Runtime/Windows/D3D11RHI/Private/D3D11Util.cpp
--- a/Engine/Source/Runtime/Windows/D3D11RHI/Private/D3D11Util.cpp
+++ b/Engine/Source/Runtime/Windows/D3D11RHI/Private/D3D11Util.cpp
@@ -172,7 +172,10 @@ void VerifyD3D11Result(HRESULT D3DResult,const ANSICHAR* Code,const ANSICHAR* Fi
 	// this is to track down a rarely happening crash
 	if(D3DResult == E_OUTOFMEMORY)
 	{
-		FMessageDialog::Open(EAppMsgType::Ok, NSLOCTEXT("D3D11RHI", "OutOfMemory", "Out of video memory trying to allocate a rendering resource."));
+		if (IsInGameThread())
+		{
+			FMessageDialog::Open(EAppMsgType::Ok, NSLOCTEXT("D3D11RHI", "OutOfMemory", "Out of video memory trying to allocate a rendering resource."));
+		}
 #if STATS
 		GetRendererModule().DebugLogOnCrash();
 #endif
