* added more info to track down error on Intel HD3000

diff --git a/Engine/Source/Runtime/Renderer/Private/PostProcess/PostProcessSelectionOutline.cpp b/Engine/Source/Runtime/Renderer/Private/PostProcess/PostProcessSelectionOutline.cpp
--- a/Engine/Source/Runtime/Renderer/Private/PostProcess/PostProcessSelectionOutline.cpp
+++ b/Engine/Source/Runtime/Renderer/Private/PostProcess/PostProcessSelectionOutline.cpp
@@ -235,8 +235,9 @@ public:
 
 			SetTextureParameter(ShaderRHI, PostprocessInput1MS, TargetableTexture);
 
-			// cache the stencil SRV to avoid create calls each frame (the cache element is stored in the state)
+			if(EditorPrimitivesStencil.IsBound())
 			{
+				// cache the stencil SRV to avoid create calls each frame (the cache element is stored in the state)
 				if(ViewState->SelectionOutlineCacheKey != TargetableTexture)
 				{
 					// release if not the right one (as the internally SRV stores a pointer to the texture we cannot get a false positive)
@@ -250,9 +251,9 @@ public:
 					ViewState->SelectionOutlineCacheKey = TargetableTexture;
 					ViewState->SelectionOutlineCacheValue = RHICreateShaderResourceView(TargetableTexture, 0, 1, PF_X24_G8);
 				}
-			}
 
-			SetSRVParameter(ShaderRHI, EditorPrimitivesStencil, ViewState->SelectionOutlineCacheValue);
+				SetSRVParameter(ShaderRHI, EditorPrimitivesStencil, ViewState->SelectionOutlineCacheValue);
+			}
 		}
 
 #if WITH_EDITOR
diff --git a/Engine/Source/Runtime/Windows/D3D11RHI/Private/D3D11Texture.cpp b/Engine/Source/Runtime/Windows/D3D11RHI/Private/D3D11Texture.cpp
--- a/Engine/Source/Runtime/Windows/D3D11RHI/Private/D3D11Texture.cpp
+++ b/Engine/Source/Runtime/Windows/D3D11RHI/Private/D3D11Texture.cpp
@@ -1175,10 +1175,19 @@ FShaderResourceViewRHIRef FD3D11DynamicRHI::RHICreateShaderResourceView(FTexture
 		SRVDesc.Texture2D.MipLevels = NumMipLevels;
 	}
 	
-
 	SRVDesc.Format = PlatformShaderResourceFormat;
 	TRefCountPtr<ID3D11ShaderResourceView> ShaderResourceView;
-	VERIFYD3D11RESULT(Direct3DDevice->CreateShaderResourceView(Texture2D->GetResource(), &SRVDesc, (ID3D11ShaderResourceView**)ShaderResourceView.GetInitReference()));
+
+	HRESULT hResCreateShaderResourceView = Direct3DDevice->CreateShaderResourceView(Texture2D->GetResource(), &SRVDesc, (ID3D11ShaderResourceView**)ShaderResourceView.GetInitReference());
+
+	if(FAILED(hResCreateShaderResourceView))
+	{
+		// provide more input data to track down error
+		UE_LOG(LogD3D11RHI,Warning,TEXT("CreateShaderResourceView failed, input: ViewDim:%d MSAA:%d Format:%d/%d SRGB:%d hRes:%x"),
+			(int32)SRVDesc.ViewDimension, (int32)TextureDesc.SampleDesc.Count, (int32)PlatformResourceFormat, (int32)PlatformShaderResourceFormat, bSRGB ? 1 : 0, hResCreateShaderResourceView);
+	}
+
+	VERIFYD3D11RESULT(hResCreateShaderResourceView);
 
 	return new FD3D11ShaderResourceView(ShaderResourceView,Texture2D);
 }
