fix for uniform buffer pooling to not reuse buffers freed from the current frame

diff --git a/Engine/Source/Runtime/OpenGLDrv/Private/OpenGLUniformBuffer.cpp b/Engine/Source/Runtime/OpenGLDrv/Private/OpenGLUniformBuffer.cpp
--- a/Engine/Source/Runtime/OpenGLDrv/Private/OpenGLUniformBuffer.cpp
+++ b/Engine/Source/Runtime/OpenGLDrv/Private/OpenGLUniformBuffer.cpp
@@ -439,7 +439,7 @@ void AddNewlyFreedBufferToUniformBufferPool( GLuint Buffer, uint32 BufferSize, b
 		int StreamedIndex = bStreamDraw ? 1 : 0;
 
 		// Add to this frame's array of free uniform buffers
-		const int32 SafeFrameIndex = GFrameNumberRenderThread % NUM_SAFE_FRAMES;
+		const int32 SafeFrameIndex = (GFrameNumberRenderThread - 1) % NUM_SAFE_FRAMES;
 		const uint32 BucketIndex = GetPoolBucketIndex(BufferSize);
 
 		check(BufferSize == UniformBufferSizeBuckets[BucketIndex]);
diff --git a/Engine/Source/Runtime/Windows/D3D11RHI/Private/D3D11UniformBuffer.cpp b/Engine/Source/Runtime/Windows/D3D11RHI/Private/D3D11UniformBuffer.cpp
--- a/Engine/Source/Runtime/Windows/D3D11RHI/Private/D3D11UniformBuffer.cpp
+++ b/Engine/Source/Runtime/Windows/D3D11RHI/Private/D3D11UniformBuffer.cpp
@@ -212,7 +212,7 @@ FD3D11UniformBuffer::~FD3D11UniformBuffer()
 			NewEntry.CreatedSize = Desc.ByteWidth;
 
 			// Add to this frame's array of free uniform buffers
-			const int32 SafeFrameIndex = GFrameNumberRenderThread % NumSafeFrames;
+			const int32 SafeFrameIndex = (GFrameNumberRenderThread - 1) % NumSafeFrames;
 			const uint32 BucketIndex = GetPoolBucketIndex(Desc.ByteWidth);
 			check(Desc.ByteWidth <= GetPoolBucketSize(Desc.ByteWidth));
 			SafeUniformBufferPools[SafeFrameIndex][BucketIndex].Add(NewEntry);
