* added CreateInfo to RHICreateTexture functions to ease future changes

diff --git a/Engine/Source/Runtime/EmptyRHI/Private/EmptyTexture.cpp b/Engine/Source/Runtime/EmptyRHI/Private/EmptyTexture.cpp
--- a/Engine/Source/Runtime/EmptyRHI/Private/EmptyTexture.cpp
+++ b/Engine/Source/Runtime/EmptyRHI/Private/EmptyTexture.cpp
@@ -112,9 +112,9 @@ uint32 FEmptyDynamicRHI::RHIComputeMemorySize(FTextureRHIParamRef TextureRHI)
 	2D texture support.
 -----------------------------------------------------------------------------*/
 
-FTexture2DRHIRef FEmptyDynamicRHI::RHICreateTexture2D(uint32 SizeX, uint32 SizeY, uint8 Format, uint32 NumMips, uint32 NumSamples, uint32 Flags, FResourceBulkDataInterface* BulkData)
+FTexture2DRHIRef FEmptyDynamicRHI::RHICreateTexture2D(uint32 SizeX, uint32 SizeY, uint8 Format, uint32 NumMips, uint32 NumSamples, uint32 Flags, FRHIResourceCreateInfo& CreateInfo)
 {
-	return new FEmptyTexture2D((EPixelFormat)Format, SizeX, SizeY, NumMips, NumSamples, Flags, BulkData);
+	return new FEmptyTexture2D((EPixelFormat)Format, SizeX, SizeY, NumMips, NumSamples, Flags, CreateInfo.BulkData);
 }
 
 FTexture2DRHIRef FEmptyDynamicRHI::RHIAsyncCreateTexture2D(uint32 SizeX,uint32 SizeY,uint8 Format,uint32 NumMips,uint32 Flags,void** InitialMipData,uint32 NumInitialMips)
@@ -128,14 +128,14 @@ void FEmptyDynamicRHI::RHICopySharedMips(FTexture2DRHIParamRef DestTexture2D,FTe
 
 }
 
-FTexture2DArrayRHIRef FEmptyDynamicRHI::RHICreateTexture2DArray(uint32 SizeX, uint32 SizeY, uint32 SizeZ, uint8 Format, uint32 NumMips, uint32 Flags, FResourceBulkDataInterface* BulkData)
+FTexture2DArrayRHIRef FEmptyDynamicRHI::RHICreateTexture2DArray(uint32 SizeX, uint32 SizeY, uint32 SizeZ, uint8 Format, uint32 NumMips, uint32 Flags, FRHIResourceCreateInfo& CreateInfo)
 {
-	return new FEmptyTexture2DArray((EPixelFormat)Format, SizeX, SizeY, SizeZ, NumMips, Flags, BulkData);
+	return new FEmptyTexture2DArray((EPixelFormat)Format, SizeX, SizeY, SizeZ, NumMips, Flags, CreateInfo.BulkData);
 }
 
-FTexture3DRHIRef FEmptyDynamicRHI::RHICreateTexture3D(uint32 SizeX, uint32 SizeY, uint32 SizeZ, uint8 Format, uint32 NumMips, uint32 Flags, FResourceBulkDataInterface* BulkData)
+FTexture3DRHIRef FEmptyDynamicRHI::RHICreateTexture3D(uint32 SizeX, uint32 SizeY, uint32 SizeZ, uint8 Format, uint32 NumMips, uint32 Flags, FRHIResourceCreateInfo& CreateInfo)
 {
-	return new FEmptyTexture3D((EPixelFormat)Format, SizeX, SizeY, SizeZ, NumMips, Flags, BulkData);
+	return new FEmptyTexture3D((EPixelFormat)Format, SizeX, SizeY, SizeZ, NumMips, Flags, CreateInfo.BulkData);
 }
 
 void FEmptyDynamicRHI::RHIGetResourceInfo(FTextureRHIParamRef Ref, FRHIResourceInfo& OutInfo)
@@ -205,14 +205,14 @@ void FEmptyDynamicRHI::RHIUpdateTexture3D(FTexture3DRHIParamRef TextureRHI,uint3
 /*-----------------------------------------------------------------------------
 	Cubemap texture support.
 -----------------------------------------------------------------------------*/
-FTextureCubeRHIRef FEmptyDynamicRHI::RHICreateTextureCube(uint32 Size, uint8 Format, uint32 NumMips, uint32 Flags, FResourceBulkDataInterface* BulkData)
+FTextureCubeRHIRef FEmptyDynamicRHI::RHICreateTextureCube(uint32 Size, uint8 Format, uint32 NumMips, uint32 Flags, FRHIResourceCreateInfo& CreateInfo)
 {
-	return new FEmptyTextureCube((EPixelFormat)Format, Size, false, 1, NumMips, Flags, BulkData);
+	return new FEmptyTextureCube((EPixelFormat)Format, Size, false, 1, NumMips, Flags, CreateInfo.BulkData);
 }
 
-FTextureCubeRHIRef FEmptyDynamicRHI::RHICreateTextureCubeArray(uint32 Size, uint32 ArraySize, uint8 Format, uint32 NumMips, uint32 Flags, FResourceBulkDataInterface* BulkData)
+FTextureCubeRHIRef FEmptyDynamicRHI::RHICreateTextureCubeArray(uint32 Size, uint32 ArraySize, uint8 Format, uint32 NumMips, uint32 Flags, FRHIResourceCreateInfo& CreateInfo)
 {
-	return new FEmptyTextureCube((EPixelFormat)Format, Size, true, ArraySize, NumMips, Flags, BulkData);
+	return new FEmptyTextureCube((EPixelFormat)Format, Size, true, ArraySize, NumMips, Flags, CreateInfo.BulkData);
 }
 
 void* FEmptyDynamicRHI::RHILockTextureCubeFace(FTextureCubeRHIParamRef TextureCubeRHI,uint32 FaceIndex,uint32 ArrayIndex,uint32 MipIndex,EResourceLockMode LockMode,uint32& DestStride,bool bLockWithinMiptail)
diff --git a/Engine/Source/Runtime/Engine/Private/Atmosphere/Atmosphere.h b/Engine/Source/Runtime/Engine/Private/Atmosphere/Atmosphere.h
--- a/Engine/Source/Runtime/Engine/Private/Atmosphere/Atmosphere.h
+++ b/Engine/Source/Runtime/Engine/Private/Atmosphere/Atmosphere.h
@@ -142,35 +142,38 @@ public:
 				{
 					const uint32 DataSize = SizeX * SizeY * sizeof(FColor);
 					FAtmosphereResourceBulkDataInterface BulkDataInterface(TextureData, DataSize);
+					FRHIResourceCreateInfo CreateInfo(&BulkDataInterface);
 					TextureRHI = RHICreateTexture2D(
 						SizeX, SizeY, PF_B8G8R8A8,
 						/*NumMips=*/ 1,
 						/*NumSamples=*/ 1,
 						/*Flags=*/ TexCreate_ShaderResource,
-						/*BulkData=*/ &BulkDataInterface );
+						/*BulkData=*/ CreateInfo );
 				}
 				break;
 			case E_Irradiance:
 				{
 					const uint32 DataSize = SizeX * SizeY * sizeof(FColor);
 					FAtmosphereResourceBulkDataInterface BulkDataInterface(TextureData, DataSize);
+					FRHIResourceCreateInfo CreateInfo(&BulkDataInterface);
 					TextureRHI = RHICreateTexture2D(
 						SizeX, SizeY, PF_B8G8R8A8,
 						/*NumMips=*/ 1,
 						/*NumSamples=*/ 1,
 						/*Flags=*/ TexCreate_ShaderResource,
-						/*BulkData=*/ &BulkDataInterface );
+						/*BulkData=*/ CreateInfo );
 				}
 				break;
 			case E_Inscatter:
 				{
 					const uint32 DataSize = SizeX * SizeY * SizeZ * sizeof(FFloat16Color);
 					FAtmosphereResourceBulkDataInterface BulkDataInterface(TextureData, DataSize);
+					FRHIResourceCreateInfo CreateInfo(&BulkDataInterface);
 					TextureRHI = RHICreateTexture3D(
 						SizeX, SizeY, SizeZ, PF_FloatRGBA,
 						/*NumMips=*/ 1,
 						/*Flags=*/ TexCreate_ShaderResource,
-						/*BulkData=*/ &BulkDataInterface );
+						/*BulkData=*/ CreateInfo );
 				}
 				break;
 			}
diff --git a/Engine/Source/Runtime/Engine/Private/GPUSkinVertexFactory.cpp b/Engine/Source/Runtime/Engine/Private/GPUSkinVertexFactory.cpp
--- a/Engine/Source/Runtime/Engine/Private/GPUSkinVertexFactory.cpp
+++ b/Engine/Source/Runtime/Engine/Private/GPUSkinVertexFactory.cpp
@@ -115,7 +115,8 @@ TStatId FBoneBufferPool::GetStatId() const
 FTexture2DRHIRef FBoneTexturePoolPolicy::CreateResource(FSharedPoolPolicyData::CreationArguments Args)
 {
 	FSharedPoolPolicyData::CreationArguments ActualSize = GetPoolBucketSize(GetPoolBucketIndex(Args)) / sizeof(FVector4);
-	FTexture2DRHIRef Texture = RHICreateTexture2D(ActualSize, 1, PF_A32B32G32R32F, 1, 1, (TexCreate_ShaderResource|TexCreate_NoMipTail|TexCreate_Dynamic), NULL);
+	FRHIResourceCreateInfo CreateInfo;
+	FTexture2DRHIRef Texture = RHICreateTexture2D(ActualSize, 1, PF_A32B32G32R32F, 1, 1, (TexCreate_ShaderResource|TexCreate_NoMipTail|TexCreate_Dynamic), CreateInfo);
 	return Texture;
 }
 
diff --git a/Engine/Source/Runtime/Engine/Private/Particles/ParticleCurveTexture.cpp b/Engine/Source/Runtime/Engine/Private/Particles/ParticleCurveTexture.cpp
--- a/Engine/Source/Runtime/Engine/Private/Particles/ParticleCurveTexture.cpp
+++ b/Engine/Source/Runtime/Engine/Private/Particles/ParticleCurveTexture.cpp
@@ -494,6 +494,7 @@ FParticleCurveTexture::FParticleCurveTexture()
 void FParticleCurveTexture::InitRHI()
 {
 	// 8-bit per channel RGBA texture for curves.
+	FRHIResourceCreateInfo CreateInfo;
 	RHICreateTargetableShaderResource2D(
 		GParticleCurveTextureSizeX,
 		GParticleCurveTextureSizeY,
@@ -502,6 +503,7 @@ void FParticleCurveTexture::InitRHI()
 		TexCreate_None,
 		TexCreate_RenderTargetable,
 		/*bForceSeparateTargetAndShaderResource=*/ false,
+		CreateInfo,
 		CurveTextureTargetRHI,
 		CurveTextureRHI
 		);
diff --git a/Engine/Source/Runtime/Engine/Private/Particles/ParticleGpuSimulation.cpp b/Engine/Source/Runtime/Engine/Private/Particles/ParticleGpuSimulation.cpp
--- a/Engine/Source/Runtime/Engine/Private/Particles/ParticleGpuSimulation.cpp
+++ b/Engine/Source/Runtime/Engine/Private/Particles/ParticleGpuSimulation.cpp
@@ -170,6 +170,7 @@ public:
 		check( !IsValidRef( PositionTextureTargetRHI ) );
 		check( !IsValidRef( PositionTextureRHI ) );
 
+		FRHIResourceCreateInfo CreateInfo;
 		RHICreateTargetableShaderResource2D(
 			SizeX,
 			SizeY,
@@ -178,6 +179,7 @@ public:
 			TexCreate_None,
 			TexCreate_RenderTargetable,
 			/*bForceSeparateTargetAndShaderResource=*/ false,
+			CreateInfo,
 			PositionTextureTargetRHI,
 			PositionTextureRHI
 			);
@@ -194,6 +196,7 @@ public:
 			TexCreate_None,
 			TexCreate_RenderTargetable,
 			/*bForceSeparateTargetAndShaderResource=*/ false,
+			CreateInfo,
 			VelocityTextureTargetRHI,
 			VelocityTextureRHI
 			);
@@ -235,6 +238,7 @@ public:
 		const int32 SizeX = GParticleSimulationTextureSizeX;
 		const int32 SizeY = GParticleSimulationTextureSizeY;
 
+		FRHIResourceCreateInfo CreateInfo;
 		RHICreateTargetableShaderResource2D(
 			SizeX,
 			SizeY,
@@ -243,6 +247,7 @@ public:
 			TexCreate_None,
 			TexCreate_RenderTargetable,
 			/*bForceSeparateTargetAndShaderResource=*/ false,
+			CreateInfo,
 			TextureTargetRHI,
 			TextureRHI
 			);
diff --git a/Engine/Source/Runtime/Engine/Private/ReflectionCaptureComponent.cpp b/Engine/Source/Runtime/Engine/Private/ReflectionCaptureComponent.cpp
--- a/Engine/Source/Runtime/Engine/Private/ReflectionCaptureComponent.cpp
+++ b/Engine/Source/Runtime/Engine/Private/ReflectionCaptureComponent.cpp
@@ -622,7 +622,8 @@ public:
 
 	virtual void InitRHI()
 	{
-		TextureCubeRHI = RHICreateTextureCube(Size, Format, NumMips, 0, NULL);
+		FRHIResourceCreateInfo CreateInfo;
+		TextureCubeRHI = RHICreateTextureCube(Size, Format, NumMips, 0, CreateInfo);
 		TextureRHI = TextureCubeRHI;
 
 		if (SourceData)
diff --git a/Engine/Source/Runtime/Engine/Private/SkyLightComponent.cpp b/Engine/Source/Runtime/Engine/Private/SkyLightComponent.cpp
--- a/Engine/Source/Runtime/Engine/Private/SkyLightComponent.cpp
+++ b/Engine/Source/Runtime/Engine/Private/SkyLightComponent.cpp
@@ -20,7 +20,8 @@ void FSkyTextureCubeResource::InitRHI()
 {
 	if (GRHIFeatureLevel >= ERHIFeatureLevel::SM3)
 	{
-		TextureCubeRHI = RHICreateTextureCube(Size, Format, NumMips, 0, NULL);
+		FRHIResourceCreateInfo CreateInfo;
+		TextureCubeRHI = RHICreateTextureCube(Size, Format, NumMips, 0, CreateInfo);
 		TextureRHI = TextureCubeRHI;
 
 		// Create the sampler state RHI resource.
diff --git a/Engine/Source/Runtime/Engine/Private/Slate/SceneViewport.cpp b/Engine/Source/Runtime/Engine/Private/Slate/SceneViewport.cpp
--- a/Engine/Source/Runtime/Engine/Private/Slate/SceneViewport.cpp
+++ b/Engine/Source/Runtime/Engine/Private/Slate/SceneViewport.cpp
@@ -1103,7 +1103,8 @@ void FSceneViewport::InitDynamicRHI()
 	{
 		FTexture2DRHIRef ShaderResourceTextureRHI;
 
-		RHICreateTargetableShaderResource2D( SizeX, SizeY, PF_B8G8R8A8, 1, TexCreate_None, TexCreate_RenderTargetable, false, RenderTargetTextureRHI, ShaderResourceTextureRHI );
+		FRHIResourceCreateInfo CreateInfo;
+		RHICreateTargetableShaderResource2D( SizeX, SizeY, PF_B8G8R8A8, 1, TexCreate_None, TexCreate_RenderTargetable, false, CreateInfo, RenderTargetTextureRHI, ShaderResourceTextureRHI );
 
 		if( !SlateRenderTargetHandle )
 		{
diff --git a/Engine/Source/Runtime/Engine/Private/Slate/SlateTextures.cpp b/Engine/Source/Runtime/Engine/Private/Slate/SlateTextures.cpp
--- a/Engine/Source/Runtime/Engine/Private/Slate/SlateTextures.cpp
+++ b/Engine/Source/Runtime/Engine/Private/Slate/SlateTextures.cpp
@@ -42,7 +42,8 @@ void FSlateTexture2DRHIRef::InitDynamicRHI()
 		if( TextureData.IsValid() || bCreateEmptyTexture )
 		{
 			check( !IsValidRef( ShaderResource) );
-			ShaderResource = RHICreateTexture2D( Width, Height, PixelFormat, 1, 1, TexCreateFlags, NULL );
+			FRHIResourceCreateInfo CreateInfo;
+			ShaderResource = RHICreateTexture2D( Width, Height, PixelFormat, 1, 1, TexCreateFlags, CreateInfo );
 			check( IsValidRef( ShaderResource ) );
 		}
 
@@ -170,6 +171,7 @@ void FSlateTextureRenderTarget2DResource::InitDynamicRHI()
 	if( TargetSizeX > 0 && TargetSizeY > 0 )
 	{
 		// Create the RHI texture. Only one mip is used and the texture is targetable for resolve.
+		FRHIResourceCreateInfo CreateInfo;
 		RHICreateTargetableShaderResource2D(
 			TargetSizeX, 
 			TargetSizeY, 
@@ -178,6 +180,7 @@ void FSlateTextureRenderTarget2DResource::InitDynamicRHI()
 			/*TexCreateFlags=*/0,
 			TexCreate_RenderTargetable,
 			/*bNeedsTwoCopies=*/false,
+			CreateInfo,
 			RenderTargetTextureRHI,
 			Texture2DRHI
 			);
diff --git a/Engine/Source/Runtime/Engine/Private/Texture2D.cpp b/Engine/Source/Runtime/Engine/Private/Texture2D.cpp
--- a/Engine/Source/Runtime/Engine/Private/Texture2D.cpp
+++ b/Engine/Source/Runtime/Engine/Private/Texture2D.cpp
@@ -1271,7 +1271,8 @@ void FTexture2DResource::InitRHI()
 			{
 				TexCreateFlags |= TexCreate_Virtual;
 
-				Texture2DRHI	= RHICreateTexture2D( OwnerMips[0].SizeX, OwnerMips[0].SizeY, EffectiveFormat, OwnerMips.Num(), 1, TexCreateFlags, ResourceMem );
+				FRHIResourceCreateInfo CreateInfo(ResourceMem);
+				Texture2DRHI	= RHICreateTexture2D( OwnerMips[0].SizeX, OwnerMips[0].SizeY, EffectiveFormat, OwnerMips.Num(), 1, TexCreateFlags, CreateInfo);
 				RHIVirtualTextureSetFirstMipInMemory(Texture2DRHI, CurrentFirstMip);
 				RHIVirtualTextureSetFirstMipVisible(Texture2DRHI, CurrentFirstMip);
 
@@ -1303,7 +1304,8 @@ void FTexture2DResource::InitRHI()
 			}
 
 			// create texture with ResourceMem data when available
-			Texture2DRHI	= RHICreateTexture2D( SizeX, SizeY, EffectiveFormat, Owner->RequestedMips, 1, TexCreateFlags, ResourceMem );
+			FRHIResourceCreateInfo CreateInfo(ResourceMem);
+			Texture2DRHI	= RHICreateTexture2D( SizeX, SizeY, EffectiveFormat, Owner->RequestedMips, 1, TexCreateFlags, CreateInfo);
 			TextureRHI		= Texture2DRHI;
 			RHIBindDebugLabelName(TextureRHI, *Owner->GetName());
 
@@ -1351,7 +1353,8 @@ void FTexture2DResource::InitRHI()
 		bool bSkipRHITextureCreation = false; //Owner->bIsCompositingSource;
 		if (GIsEditor || (!bSkipRHITextureCreation))
 		{
-			Texture2DRHI	= RHICreateTexture2D( SizeX, SizeY, EffectiveFormat, Owner->RequestedMips, 1, TexCreateFlags, NULL );
+			FRHIResourceCreateInfo CreateInfo;
+			Texture2DRHI	= RHICreateTexture2D( SizeX, SizeY, EffectiveFormat, Owner->RequestedMips, 1, TexCreateFlags, CreateInfo );
 			TextureRHI		= Texture2DRHI;
 			RHIBindDebugLabelName(TextureRHI, *Owner->GetName());
 			for( int32 MipIndex=CurrentFirstMip; MipIndex<OwnerMips.Num(); MipIndex++ )
@@ -1634,7 +1637,8 @@ void FTexture2DResource::UpdateMipCount()
 				check(PendingFirstMip>=0);
 
 				const uint32 TexCreateFlags = Texture2DRHI->GetFlags() & ~TexCreate_Virtual;
-				IntermediateTextureRHI = RHICreateTexture2D( OwnerMips[PendingFirstMip].SizeX, OwnerMips[PendingFirstMip].SizeY, OldTexture->GetFormat(), OwnerMips.Num() - PendingFirstMip, 1, TexCreateFlags, ResourceMem );
+				FRHIResourceCreateInfo CreateInfo(ResourceMem);
+				IntermediateTextureRHI = RHICreateTexture2D( OwnerMips[PendingFirstMip].SizeX, OwnerMips[PendingFirstMip].SizeY, OldTexture->GetFormat(), OwnerMips.Num() - PendingFirstMip, 1, TexCreateFlags, CreateInfo );
 
 				check(Owner->bIsStreamable);
 				check(Owner->PendingMipChangeRequestStatus.GetValue() == TexState_InProgress_Allocation);
@@ -1654,7 +1658,8 @@ void FTexture2DResource::UpdateMipCount()
 			else
 			{
 				const uint32 TexCreateFlags = Texture2DRHI->GetFlags() | TexCreate_Virtual;
-				Texture2DRHI = RHICreateTexture2D( OwnerMips[0].SizeX, OwnerMips[0].SizeY, OldTexture->GetFormat(), OwnerMips.Num(), 1, TexCreateFlags, ResourceMem );
+				FRHIResourceCreateInfo CreateInfo(ResourceMem);
+				Texture2DRHI = RHICreateTexture2D( OwnerMips[0].SizeX, OwnerMips[0].SizeY, OldTexture->GetFormat(), OwnerMips.Num(), 1, TexCreateFlags, CreateInfo );
 				RHIVirtualTextureSetFirstMipInMemory(Texture2DRHI, CurrentFirstMip);
 				RHIVirtualTextureSetFirstMipVisible(Texture2DRHI, CurrentFirstMip);
 				RHICopySharedMips(Texture2DRHI, OldTexture);
@@ -2306,7 +2311,8 @@ void FTexture2DArrayResource::InitRHI()
 {
 	// Create the RHI texture.
 	const uint32 TexCreateFlags = (bSRGB ? TexCreate_SRGB : 0) | TexCreate_OfflineProcessed;
-	FTexture2DArrayRHIRef TextureArray = RHICreateTexture2DArray(SizeX, SizeY, GetNumValidTextures(), Format, NumMips, TexCreateFlags, NULL);
+	FRHIResourceCreateInfo CreateInfo;
+	FTexture2DArrayRHIRef TextureArray = RHICreateTexture2DArray(SizeX, SizeY, GetNumValidTextures(), Format, NumMips, TexCreateFlags, CreateInfo);
 	TextureRHI = TextureArray;
 
 	// Read the mip-levels into the RHI texture.
diff --git a/Engine/Source/Runtime/Engine/Private/Texture2DDynamic.cpp b/Engine/Source/Runtime/Engine/Private/Texture2DDynamic.cpp
--- a/Engine/Source/Runtime/Engine/Private/Texture2DDynamic.cpp
+++ b/Engine/Source/Runtime/Engine/Private/Texture2DDynamic.cpp
@@ -56,7 +56,8 @@ void FTexture2DDynamicResource::InitRHI()
 	{
 		Flags |= TexCreate_NoTiling;
 	}
-	Texture2DRHI = RHICreateTexture2D( Owner->SizeX, Owner->SizeY, Owner->Format, Owner->NumMips, 1, Flags, NULL );
+	FRHIResourceCreateInfo CreateInfo;
+	Texture2DRHI = RHICreateTexture2D( Owner->SizeX, Owner->SizeY, Owner->Format, Owner->NumMips, 1, Flags, CreateInfo );
 	TextureRHI = Texture2DRHI;
 }
 
diff --git a/Engine/Source/Runtime/Engine/Private/TextureCube.cpp b/Engine/Source/Runtime/Engine/Private/TextureCube.cpp
--- a/Engine/Source/Runtime/Engine/Private/TextureCube.cpp
+++ b/Engine/Source/Runtime/Engine/Private/TextureCube.cpp
@@ -224,7 +224,8 @@ public:
 
 		// Create the RHI texture.
 		uint32 TexCreateFlags = (Owner->SRGB ? TexCreate_SRGB : 0)  | TexCreate_OfflineProcessed;
-		TextureCubeRHI = RHICreateTextureCube( Owner->GetSizeX(), Owner->GetPixelFormat(), Owner->GetNumMips(), TexCreateFlags, NULL );
+		FRHIResourceCreateInfo CreateInfo;
+		TextureCubeRHI = RHICreateTextureCube( Owner->GetSizeX(), Owner->GetPixelFormat(), Owner->GetNumMips(), TexCreateFlags, CreateInfo );
 		TextureRHI = TextureCubeRHI;
 		RHIBindDebugLabelName(TextureRHI, *Owner->GetName());
 
diff --git a/Engine/Source/Runtime/Engine/Private/TextureMovie.cpp b/Engine/Source/Runtime/Engine/Private/TextureMovie.cpp
--- a/Engine/Source/Runtime/Engine/Private/TextureMovie.cpp
+++ b/Engine/Source/Runtime/Engine/Private/TextureMovie.cpp
@@ -352,6 +352,7 @@ void FTextureMovieResource::InitDynamicRHI()
 	{
 		// Create the RHI texture. Only one mip is used and the texture is targetable or resolve.
 		uint32 TexCreateFlags = Owner->SRGB ? TexCreate_SRGB : 0;
+		FRHIResourceCreateInfo CreateInfo;
 		RHICreateTargetableShaderResource2D(
 			Owner->SizeX, 
 			Owner->SizeY, 
@@ -360,6 +361,7 @@ void FTextureMovieResource::InitDynamicRHI()
 			TexCreateFlags,
 			TexCreate_RenderTargetable,
 			false,
+			CreateInfo,
 			RenderTargetTextureRHI,
 			Texture2DRHI
 			);
diff --git a/Engine/Source/Runtime/Engine/Private/TextureRenderTarget2D.cpp b/Engine/Source/Runtime/Engine/Private/TextureRenderTarget2D.cpp
--- a/Engine/Source/Runtime/Engine/Private/TextureRenderTarget2D.cpp
+++ b/Engine/Source/Runtime/Engine/Private/TextureRenderTarget2D.cpp
@@ -312,6 +312,7 @@ void FTextureRenderTarget2DResource::InitDynamicRHI()
 
 		// Create the RHI texture. Only one mip is used and the texture is targetable for resolve.
 		uint32 TexCreateFlags = bSRGB ? TexCreate_SRGB : 0;
+		FRHIResourceCreateInfo CreateInfo;
 		RHICreateTargetableShaderResource2D(
 			Owner->SizeX, 
 			Owner->SizeY, 
@@ -320,6 +321,7 @@ void FTextureRenderTarget2DResource::InitDynamicRHI()
 			TexCreateFlags,
 			TexCreate_RenderTargetable,
 			Owner->bNeedsTwoCopies,
+			CreateInfo,
 			RenderTargetTextureRHI,
 			Texture2DRHI
 			);
diff --git a/Engine/Source/Runtime/Engine/Private/TextureRenderTargetCube.cpp b/Engine/Source/Runtime/Engine/Private/TextureRenderTargetCube.cpp
--- a/Engine/Source/Runtime/Engine/Private/TextureRenderTargetCube.cpp
+++ b/Engine/Source/Runtime/Engine/Private/TextureRenderTargetCube.cpp
@@ -209,28 +209,35 @@ void FTextureRenderTargetCubeResource::InitDynamicRHI()
 
 		// Create the RHI texture. Only one mip is used and the texture is targetable for resolve.
 		uint32 TexCreateFlags = bSRGB ? TexCreate_SRGB : 0;
-		RHICreateTargetableShaderResourceCube(
-			Owner->SizeX,
-			Owner->GetFormat(), 
-			Owner->GetNumMips(),
-			TexCreateFlags, 
-			TexCreate_RenderTargetable,
-			false,
-			RenderTargetCubeRHI,
-			TextureCubeRHI );
+		{
+			FRHIResourceCreateInfo CreateInfo;
+			RHICreateTargetableShaderResourceCube(
+				Owner->SizeX,
+				Owner->GetFormat(), 
+				Owner->GetNumMips(),
+				TexCreateFlags, 
+				TexCreate_RenderTargetable,
+				false,
+				CreateInfo,
+				RenderTargetCubeRHI,
+				TextureCubeRHI );
+		}
 
 		TextureRHI = TextureCubeRHI;
 
 		// Create the RHI target surface used for rendering to
-		CubeFaceSurfaceRHI = RHICreateTexture2D(
-			Owner->SizeX, 
-			Owner->SizeX, 
-			Owner->GetFormat(),
-			Owner->GetNumMips(), 
-			/* NumSamples =*/ 1,
-			TexCreate_RenderTargetable|TexCreateFlags,
-			NULL
-			);
+		{
+			FRHIResourceCreateInfo CreateInfo;
+			CubeFaceSurfaceRHI = RHICreateTexture2D(
+				Owner->SizeX, 
+				Owner->SizeX, 
+				Owner->GetFormat(),
+				Owner->GetNumMips(), 
+				/* NumSamples =*/ 1,
+				TexCreate_RenderTargetable|TexCreateFlags,
+				CreateInfo
+				);
+		}
 
 		// Set render target to 2D surface.
 		RenderTargetTextureRHI = CubeFaceSurfaceRHI;
diff --git a/Engine/Source/Runtime/Engine/Private/UnrealClient.cpp b/Engine/Source/Runtime/Engine/Private/UnrealClient.cpp
--- a/Engine/Source/Runtime/Engine/Private/UnrealClient.cpp
+++ b/Engine/Source/Runtime/Engine/Private/UnrealClient.cpp
@@ -723,7 +723,8 @@ public:
 	{
 		FTexture2DRHIRef ShaderResourceTextureRHI;
 
-		RHICreateTargetableShaderResource2D( SizeX, SizeY, PF_B8G8R8A8, 1, TexCreate_None, TexCreate_RenderTargetable, false, RenderTargetTextureRHI, ShaderResourceTextureRHI );
+		FRHIResourceCreateInfo CreateInfo;
+		RHICreateTargetableShaderResource2D( SizeX, SizeY, PF_B8G8R8A8, 1, TexCreate_None, TexCreate_RenderTargetable, false, CreateInfo, RenderTargetTextureRHI, ShaderResourceTextureRHI );
 	}
 
 	// @todo UE4 DLL: Without these functions we get unresolved linker errors with FRenderResource
@@ -1481,8 +1482,14 @@ void FViewport::FHitProxyMap::Init(uint32 NewSizeX,uint32 NewSizeY)
 	SizeY = NewSizeY;
 
 	// Create a render target to store the hit proxy map.
-	RHICreateTargetableShaderResource2D(SizeX,SizeY,PF_B8G8R8A8,1,TexCreate_None,TexCreate_RenderTargetable,false,RenderTargetTextureRHI,HitProxyTexture);
-	HitProxyCPUTexture = RHICreateTexture2D(SizeX, SizeY, PF_B8G8R8A8,1,1,TexCreate_CPUReadback,NULL);
+	{
+		FRHIResourceCreateInfo CreateInfo;
+		RHICreateTargetableShaderResource2D(SizeX,SizeY,PF_B8G8R8A8,1,TexCreate_None,TexCreate_RenderTargetable,false,CreateInfo,RenderTargetTextureRHI,HitProxyTexture);
+	}
+	{
+		FRHIResourceCreateInfo CreateInfo;
+		HitProxyCPUTexture = RHICreateTexture2D(SizeX, SizeY, PF_B8G8R8A8,1,1,TexCreate_CPUReadback,CreateInfo);
+	}
 }
 
 void FViewport::FHitProxyMap::Release()
diff --git a/Engine/Source/Runtime/Engine/Private/VectorField.cpp b/Engine/Source/Runtime/Engine/Private/VectorField.cpp
--- a/Engine/Source/Runtime/Engine/Private/VectorField.cpp
+++ b/Engine/Source/Runtime/Engine/Private/VectorField.cpp
@@ -184,11 +184,12 @@ public:
 		{
 			const uint32 DataSize = SizeX * SizeY * SizeZ * sizeof(FFloat16Color);
 			FVectorFieldStaticResourceBulkDataInterface BulkDataInterface(VolumeData, DataSize);
+			FRHIResourceCreateInfo CreateInfo(&BulkDataInterface);
 			VolumeTextureRHI = RHICreateTexture3D(
 				SizeX, SizeY, SizeZ, PF_FloatRGBA,
 				/*NumMips=*/ 1,
 				/*Flags=*/ TexCreate_ShaderResource,
-				/*BulkData=*/ &BulkDataInterface );
+				/*BulkData=*/ CreateInfo );
 			FMemory::Free(VolumeData);
 			VolumeData = NULL;
 		}
@@ -723,12 +724,13 @@ public:
 				TexCreateFlags = TexCreate_ShaderResource | TexCreate_UAV;
 			}
 
+			FRHIResourceCreateInfo CreateInfo;
 			VolumeTextureRHI = RHICreateTexture3D(
 				SizeX, SizeY, SizeZ,
 				PF_FloatRGBA,
 				/*NumMips=*/ 1,
 				TexCreateFlags,
-				/*BulkData=*/ NULL );
+				CreateInfo);
 
 			if (GRHIFeatureLevel == ERHIFeatureLevel::SM5)
 			{
diff --git a/Engine/Source/Runtime/Engine/Public/GPUSkinVertexFactory.h b/Engine/Source/Runtime/Engine/Public/GPUSkinVertexFactory.h
--- a/Engine/Source/Runtime/Engine/Public/GPUSkinVertexFactory.h
+++ b/Engine/Source/Runtime/Engine/Public/GPUSkinVertexFactory.h
@@ -307,7 +307,8 @@ public:
 				uint32 Y = (SizeX / GPUSKIN_TEXTURE_STRIDE)+1;
 				check((uint32)GMaxTextureDimensions >= X);
 				check((uint32)GMaxTextureDimensions >= Y);
-				BoneBuffer = RHICreateTexture2D(X, Y, PF_A32B32G32R32F, 1, 1, (TexCreate_ShaderResource|TexCreate_NoMipTail|TexCreate_Dynamic), NULL);
+				FRHIResourceCreateInfo CreateInfo;
+				BoneBuffer = RHICreateTexture2D(X, Y, PF_A32B32G32R32F, 1, 1, (TexCreate_ShaderResource|TexCreate_NoMipTail|TexCreate_Dynamic), CreateInfo);
 			}
 #endif
 		}
diff --git a/Engine/Source/Runtime/OpenGLDrv/Private/OpenGLTexture.cpp b/Engine/Source/Runtime/OpenGLDrv/Private/OpenGLTexture.cpp
--- a/Engine/Source/Runtime/OpenGLDrv/Private/OpenGLTexture.cpp
+++ b/Engine/Source/Runtime/OpenGLDrv/Private/OpenGLTexture.cpp
@@ -1150,9 +1150,9 @@ TOpenGLTexture<RHIResourceType>::~TOpenGLTexture()
 * @param NumMips - number of mips to generate or 0 for full mip pyramid
 * @param Flags - ETextureCreateFlags creation flags
 */
-FTexture2DRHIRef FOpenGLDynamicRHI::RHICreateTexture2D(uint32 SizeX,uint32 SizeY,uint8 Format,uint32 NumMips,uint32 NumSamples,uint32 Flags,FResourceBulkDataInterface* BulkData)
+FTexture2DRHIRef FOpenGLDynamicRHI::RHICreateTexture2D(uint32 SizeX,uint32 SizeY,uint8 Format,uint32 NumMips,uint32 NumSamples,uint32 Flags,FRHIResourceCreateInfo& Info)
 {
-	return (FRHITexture2D*)CreateOpenGLTexture(SizeX,SizeY,false,false,Format,NumMips,NumSamples,1, Flags,BulkData);
+	return (FRHITexture2D*)CreateOpenGLTexture(SizeX,SizeY,false,false,Format,NumMips,NumSamples,1,Flags,Info.BulkData);
 }
 
 FTexture2DRHIRef FOpenGLDynamicRHI::RHIAsyncCreateTexture2D(uint32 SizeX,uint32 SizeY,uint8 Format,uint32 NumMips,uint32 Flags,void** InitialMipData,uint32 NumInitialMips)
@@ -1166,7 +1166,7 @@ void FOpenGLDynamicRHI::RHICopySharedMips(FTexture2DRHIParamRef DestTexture2D,FT
 	check(0);
 }
 
-FTexture2DArrayRHIRef FOpenGLDynamicRHI::RHICreateTexture2DArray(uint32 SizeX,uint32 SizeY,uint32 SizeZ,uint8 Format,uint32 NumMips,uint32 Flags,FResourceBulkDataInterface* BulkData)
+FTexture2DArrayRHIRef FOpenGLDynamicRHI::RHICreateTexture2DArray(uint32 SizeX,uint32 SizeY,uint32 SizeZ,uint8 Format,uint32 NumMips,uint32 Flags, FRHIResourceCreateInfo& Info)
 {
 	VERIFY_GL_SCOPE();
 
@@ -1217,7 +1217,7 @@ FTexture2DArrayRHIRef FOpenGLDynamicRHI::RHICreateTexture2DArray(uint32 SizeX,ui
 	// Make sure PBO is disabled
 	CachedBindPixelUnpackBuffer(ContextState, 0);
 
-	uint8* Data = BulkData ? (uint8*)BulkData->GetResourceBulkData() : NULL;
+	uint8* Data = Info.BulkData ? (uint8*)Info.BulkData->GetResourceBulkData() : NULL;
 	uint32 MipOffset = 0;
 
 	FOpenGL::TexStorage3D( Target, NumMips, GLFormat.InternalFormat[bSRGB], SizeX, SizeY, SizeZ, GLFormat.Format, GLFormat.Type );
@@ -1275,7 +1275,7 @@ FTexture2DArrayRHIRef FOpenGLDynamicRHI::RHICreateTexture2DArray(uint32 SizeX,ui
 	return Texture;
 }
 
-FTexture3DRHIRef FOpenGLDynamicRHI::RHICreateTexture3D(uint32 SizeX,uint32 SizeY,uint32 SizeZ,uint8 Format,uint32 NumMips,uint32 Flags,FResourceBulkDataInterface* BulkData)
+FTexture3DRHIRef FOpenGLDynamicRHI::RHICreateTexture3D(uint32 SizeX,uint32 SizeY,uint32 SizeZ,uint8 Format,uint32 NumMips,uint32 Flags,FRHIResourceCreateInfo& CreateInfo)
 {
 	VERIFY_GL_SCOPE();
 
@@ -1325,7 +1325,7 @@ FTexture3DRHIRef FOpenGLDynamicRHI::RHICreateTexture3D(uint32 SizeX,uint32 SizeY
 	// Make sure PBO is disabled
 	CachedBindPixelUnpackBuffer(ContextState,0);
 
-	uint8* Data = BulkData ? (uint8*)BulkData->GetResourceBulkData() : NULL;
+	uint8* Data = CreateInfo.BulkData ? (uint8*)CreateInfo.BulkData->GetResourceBulkData() : NULL;
 	uint32 MipOffset = 0;
 
 	FOpenGL::TexStorage3D( Target, NumMips, GLFormat.InternalFormat[bSRGB], SizeX, SizeY, SizeZ, GLFormat.Format, GLFormat.Type );
@@ -1749,13 +1749,19 @@ void FOpenGLDynamicRHI::InvalidateUAVResourceInCache(GLuint Resource)
 /*-----------------------------------------------------------------------------
 	Cubemap texture support.
 -----------------------------------------------------------------------------*/
-FTextureCubeRHIRef FOpenGLDynamicRHI::RHICreateTextureCube( uint32 Size, uint8 Format, uint32 NumMips, uint32 Flags, FResourceBulkDataInterface* BulkData )
+FTextureCubeRHIRef FOpenGLDynamicRHI::RHICreateTextureCube( uint32 Size, uint8 Format, uint32 NumMips, uint32 Flags, FRHIResourceCreateInfo& CreateInfo )
 {
+	// not yet supported
+	check(!CreateInfo.BulkData);
+
 	return (FRHITextureCube*)CreateOpenGLTexture(Size,Size,true, false, Format, NumMips, 1, 1, Flags);
 }
 
-FTextureCubeRHIRef FOpenGLDynamicRHI::RHICreateTextureCubeArray( uint32 Size, uint32 ArraySize, uint8 Format, uint32 NumMips, uint32 Flags, FResourceBulkDataInterface* BulkData )
+FTextureCubeRHIRef FOpenGLDynamicRHI::RHICreateTextureCubeArray( uint32 Size, uint32 ArraySize, uint8 Format, uint32 NumMips, uint32 Flags, FRHIResourceCreateInfo& CreateInfo )
 {
+	// not yet supported
+	check(!CreateInfo.BulkData);
+
 	return (FRHITextureCube*)CreateOpenGLTexture(Size,Size,true, true, Format, NumMips, 1, 6 * ArraySize, Flags);
 }
 
diff --git a/Engine/Source/Runtime/RHI/Public/RHI.h b/Engine/Source/Runtime/RHI/Public/RHI.h
--- a/Engine/Source/Runtime/RHI/Public/RHI.h
+++ b/Engine/Source/Runtime/RHI/Public/RHI.h
@@ -1240,6 +1240,15 @@ struct FRHIResourceInfo
 	FVRamAllocation VRamAllocation;
 };
 
+struct FRHIResourceCreateInfo
+{
+	FRHIResourceCreateInfo(FResourceBulkDataInterface* InBulkData = 0)
+		: BulkData(InBulkData)
+	{}
+
+	FResourceBulkDataInterface* BulkData;
+};
+
 // Forward-declaration.
 struct FResolveParams;
 
diff --git a/Engine/Source/Runtime/RHI/Public/RHIMethods.h b/Engine/Source/Runtime/RHI/Public/RHIMethods.h
--- a/Engine/Source/Runtime/RHI/Public/RHIMethods.h
+++ b/Engine/Source/Runtime/RHI/Public/RHIMethods.h
@@ -472,7 +472,7 @@ DEFINE_RHIMETHOD_7(
 	uint32,NumMips,
 	uint32,NumSamples,
 	uint32,Flags,
-	FResourceBulkDataInterface*,BulkData,
+	FRHIResourceCreateInfo&,Info,
 	return,return new FRHITexture2D(SizeX,SizeY,NumMips,NumSamples,(EPixelFormat)Format,Flags);
 	);
 
@@ -532,7 +532,7 @@ DEFINE_RHIMETHOD_7(
 	uint8,Format,
 	uint32,NumMips,
 	uint32,Flags,
-	FResourceBulkDataInterface*,BulkData,
+	FRHIResourceCreateInfo&,Info,
 	return,return new FRHITexture2DArray(SizeX,SizeY,SizeZ,NumMips,(EPixelFormat)Format,Flags);
 	);
 
@@ -544,7 +544,6 @@ DEFINE_RHIMETHOD_7(
 * @param Format - EPixelFormat texture format
 * @param NumMips - number of mips to generate or 0 for full mip pyramid
 * @param Flags - ETextureCreateFlags creation flags
-* @param Data - Data to initialize the texture with
 */
 DEFINE_RHIMETHOD_7(
 	FTexture3DRHIRef,RHICreateTexture3D,
@@ -554,7 +553,7 @@ DEFINE_RHIMETHOD_7(
 	uint8,Format,
 	uint32,NumMips,
 	uint32,Flags,
-	FResourceBulkDataInterface*,BulkData,
+	FRHIResourceCreateInfo&,Info,
 	return,return new FRHITexture3D(SizeX,SizeY,SizeZ,NumMips,(EPixelFormat)Format,Flags);
 	);
 
@@ -790,29 +789,29 @@ DEFINE_RHIMETHOD_5(
 	uint8,Format,
 	uint32,NumMips,
 	uint32,Flags,
-	FResourceBulkDataInterface*,BulkData,
+	FRHIResourceCreateInfo&,Info,
 	return,return new FRHITextureCube(Size,NumMips,(EPixelFormat)Format,Flags);
 	);
 
 
-	/**
-	* Creates a Cube Array RHI texture resource
-	* @param Size - width/height of the texture to create
-	* @param ArraySize - number of array elements of the texture to create
-	* @param Format - EPixelFormat texture format
-	* @param NumMips - number of mips to generate or 0 for full mip pyramid
-	* @param Flags - ETextureCreateFlags creation flags
-	*/
-	DEFINE_RHIMETHOD_6(
-		FTextureCubeRHIRef,RHICreateTextureCubeArray,
-		uint32,Size,
-		uint32,ArraySize,
-		uint8,Format,
-		uint32,NumMips,
-		uint32,Flags,
-		FResourceBulkDataInterface*,BulkData,
-		return,return new FRHITextureCube(Size,NumMips,(EPixelFormat)Format,Flags);
-	);
+/**
+* Creates a Cube Array RHI texture resource
+* @param Size - width/height of the texture to create
+* @param ArraySize - number of array elements of the texture to create
+* @param Format - EPixelFormat texture format
+* @param NumMips - number of mips to generate or 0 for full mip pyramid
+* @param Flags - ETextureCreateFlags creation flags
+*/
+DEFINE_RHIMETHOD_6(
+	FTextureCubeRHIRef,RHICreateTextureCubeArray,
+	uint32,Size,
+	uint32,ArraySize,
+	uint8,Format,
+	uint32,NumMips,
+	uint32,Flags,
+	FRHIResourceCreateInfo&,Info,
+	return,return new FRHITextureCube(Size,NumMips,(EPixelFormat)Format,Flags);
+);
 
 
 /**
diff --git a/Engine/Source/Runtime/RHI/Public/RHIUtilities.h b/Engine/Source/Runtime/RHI/Public/RHIUtilities.h
--- a/Engine/Source/Runtime/RHI/Public/RHIUtilities.h
+++ b/Engine/Source/Runtime/RHI/Public/RHIUtilities.h
@@ -183,6 +183,7 @@ inline void RHICreateTargetableShaderResource2D(
 	uint32 Flags,
 	uint32 TargetableTextureFlags,
 	bool bForceSeparateTargetAndShaderResource,
+	FRHIResourceCreateInfo& CreateInfo,
 	FTexture2DRHIRef& OutTargetableTexture,
 	FTexture2DRHIRef& OutShaderResourceTexture,
 	uint32 NumSamples=1
@@ -207,13 +208,13 @@ inline void RHICreateTargetableShaderResource2D(
 	if(!bForceSeparateTargetAndShaderResource/* && GSupportsRenderDepthTargetableShaderResources*/)
 	{
 		// Create a single texture that has both TargetableTextureFlags and TexCreate_ShaderResource set.
-		OutTargetableTexture = OutShaderResourceTexture = RHICreateTexture2D(SizeX,SizeY,Format,NumMips,NumSamples,Flags | TargetableTextureFlags | TexCreate_ShaderResource,NULL);
+		OutTargetableTexture = OutShaderResourceTexture = RHICreateTexture2D(SizeX,SizeY,Format,NumMips,NumSamples,Flags | TargetableTextureFlags | TexCreate_ShaderResource,CreateInfo);
 	}
 	else
 	{
 		// Create a texture that has TargetableTextureFlags set, and a second texture that has TexCreate_ResolveTargetable and TexCreate_ShaderResource set.
-		OutTargetableTexture = RHICreateTexture2D(SizeX,SizeY,Format,NumMips,NumSamples,Flags | TargetableTextureFlags,NULL);
-		OutShaderResourceTexture = RHICreateTexture2D(SizeX,SizeY,Format,NumMips,1,Flags | TexCreate_ResolveTargetable | TexCreate_ShaderResource,NULL);
+		OutTargetableTexture = RHICreateTexture2D(SizeX,SizeY,Format,NumMips,NumSamples,Flags | TargetableTextureFlags,CreateInfo);
+		OutShaderResourceTexture = RHICreateTexture2D(SizeX,SizeY,Format,NumMips,1,Flags | TexCreate_ResolveTargetable | TexCreate_ShaderResource,CreateInfo);
 	}
 }
 
@@ -232,6 +233,7 @@ inline void RHICreateTargetableShaderResourceCube(
 	uint32 Flags,
 	uint32 TargetableTextureFlags,
 	bool bForceSeparateTargetAndShaderResource,
+	FRHIResourceCreateInfo& CreateInfo,
 	FTextureCubeRHIRef& OutTargetableTexture,
 	FTextureCubeRHIRef& OutShaderResourceTexture
 	)
@@ -256,13 +258,13 @@ inline void RHICreateTargetableShaderResourceCube(
 	if(!bForceSeparateTargetAndShaderResource/* && GSupportsRenderDepthTargetableShaderResources*/)
 	{
 		// Create a single texture that has both TargetableTextureFlags and TexCreate_ShaderResource set.
-		OutTargetableTexture = OutShaderResourceTexture = RHICreateTextureCube(LinearSize,Format,NumMips,Flags | TargetableTextureFlags | TexCreate_ShaderResource,NULL);
+		OutTargetableTexture = OutShaderResourceTexture = RHICreateTextureCube(LinearSize,Format,NumMips,Flags | TargetableTextureFlags | TexCreate_ShaderResource,CreateInfo);
 	}
 	else
 	{
 		// Create a texture that has TargetableTextureFlags set, and a second texture that has TexCreate_ResolveTargetable and TexCreate_ShaderResource set.
-		OutTargetableTexture = RHICreateTextureCube(LinearSize,Format,NumMips,Flags | TargetableTextureFlags,NULL);
-		OutShaderResourceTexture = RHICreateTextureCube(LinearSize,Format,NumMips,Flags | TexCreate_ResolveTargetable | TexCreate_ShaderResource,NULL);
+		OutTargetableTexture = RHICreateTextureCube(LinearSize,Format,NumMips,Flags | TargetableTextureFlags,CreateInfo);
+		OutShaderResourceTexture = RHICreateTextureCube(LinearSize,Format,NumMips,Flags | TexCreate_ResolveTargetable | TexCreate_ShaderResource,CreateInfo);
 	}
 }
 
@@ -282,6 +284,7 @@ inline void RHICreateTargetableShaderResourceCubeArray(
 	uint32 Flags,
 	uint32 TargetableTextureFlags,
 	bool bForceSeparateTargetAndShaderResource,
+	FRHIResourceCreateInfo& CreateInfo,
 	FTextureCubeRHIRef& OutTargetableTexture,
 	FTextureCubeRHIRef& OutShaderResourceTexture
 	)
@@ -300,13 +303,13 @@ inline void RHICreateTargetableShaderResourceCubeArray(
 	if(!bForceSeparateTargetAndShaderResource/* && GSupportsRenderDepthTargetableShaderResources*/)
 	{
 		// Create a single texture that has both TargetableTextureFlags and TexCreate_ShaderResource set.
-		OutTargetableTexture = OutShaderResourceTexture = RHICreateTextureCubeArray(LinearSize,ArraySize,Format,NumMips,Flags | TargetableTextureFlags | TexCreate_ShaderResource,NULL);
+		OutTargetableTexture = OutShaderResourceTexture = RHICreateTextureCubeArray(LinearSize,ArraySize,Format,NumMips,Flags | TargetableTextureFlags | TexCreate_ShaderResource,CreateInfo);
 	}
 	else
 	{
 		// Create a texture that has TargetableTextureFlags set, and a second texture that has TexCreate_ResolveTargetable and TexCreate_ShaderResource set.
-		OutTargetableTexture = RHICreateTextureCubeArray(LinearSize,ArraySize,Format,NumMips,Flags | TargetableTextureFlags,NULL);
-		OutShaderResourceTexture = RHICreateTextureCubeArray(LinearSize,ArraySize,Format,NumMips,Flags | TexCreate_ResolveTargetable | TexCreate_ShaderResource,NULL);
+		OutTargetableTexture = RHICreateTextureCubeArray(LinearSize,ArraySize,Format,NumMips,Flags | TargetableTextureFlags,CreateInfo);
+		OutShaderResourceTexture = RHICreateTextureCubeArray(LinearSize,ArraySize,Format,NumMips,Flags | TexCreate_ResolveTargetable | TexCreate_ShaderResource,CreateInfo);
 	}
 }
 /**
diff --git a/Engine/Source/Runtime/RenderCore/Private/RenderUtils.cpp b/Engine/Source/Runtime/RenderCore/Private/RenderUtils.cpp
--- a/Engine/Source/Runtime/RenderCore/Private/RenderUtils.cpp
+++ b/Engine/Source/Runtime/RenderCore/Private/RenderUtils.cpp
@@ -151,7 +151,8 @@ public:
 	virtual void InitRHI()
 	{
 		// Create the texture RHI.  		
-		FTexture2DRHIRef Texture2D = RHICreateTexture2D(1,1,PF_B8G8R8A8,1,1,TexCreate_ShaderResource,NULL);
+		FRHIResourceCreateInfo CreateInfo;
+		FTexture2DRHIRef Texture2D = RHICreateTexture2D(1,1,PF_B8G8R8A8,1,1,TexCreate_ShaderResource,CreateInfo);
 		TextureRHI = Texture2D;
 
 		// Write the contents of the texture.
@@ -240,7 +241,8 @@ public:
 		{
 			// Create the texture.
 			FBlackVolumeTextureResourceBulkDataInterface BlackTextureBulkData;
-			FTexture3DRHIRef Texture3D = RHICreateTexture3D(1,1,1,PF_B8G8R8A8,1,TexCreate_ShaderResource,&BlackTextureBulkData);
+			FRHIResourceCreateInfo CreateInfo(&BlackTextureBulkData);
+			FTexture3DRHIRef Texture3D = RHICreateTexture3D(1,1,1,PF_B8G8R8A8,1,TexCreate_ShaderResource,CreateInfo);
 			TextureRHI = Texture3D;	
 
 			// Create the sampler state.
@@ -278,7 +280,8 @@ public:
 		if (IsFeatureLevelSupported(GRHIShaderPlatform, ERHIFeatureLevel::SM4))
 		{
 			// Create the texture RHI.  		
-			FTexture2DArrayRHIRef TextureArray = RHICreateTexture2DArray(1,1,1,PF_B8G8R8A8,1,TexCreate_ShaderResource,NULL);
+			FRHIResourceCreateInfo CreateInfo;
+			FTexture2DArrayRHIRef TextureArray = RHICreateTexture2DArray(1,1,1,PF_B8G8R8A8,1,TexCreate_ShaderResource,CreateInfo);
 			TextureRHI = TextureArray;
 
 			uint32 DestStride;
@@ -328,7 +331,8 @@ public:
 	{
 		// Create the texture RHI.
 		int32 TextureSize = 1 << (NumMips - 1);
-		FTexture2DRHIRef Texture2D = RHICreateTexture2D(TextureSize,TextureSize,PF_B8G8R8A8,NumMips,1,TexCreate_ShaderResource,NULL);
+		FRHIResourceCreateInfo CreateInfo;
+		FTexture2DRHIRef Texture2D = RHICreateTexture2D(TextureSize,TextureSize,PF_B8G8R8A8,NumMips,1,TexCreate_ShaderResource,CreateInfo);
 		TextureRHI = Texture2D;
 
 		// Write the contents of the texture.
@@ -408,7 +412,8 @@ public:
 	virtual void InitRHI()
 	{
 		// Create the texture RHI.
-		FTextureCubeRHIRef TextureCube = RHICreateTextureCube(1,PF_B8G8R8A8,1,0,NULL);
+		FRHIResourceCreateInfo CreateInfo;
+		FTextureCubeRHIRef TextureCube = RHICreateTextureCube(1,PF_B8G8R8A8,1,0,CreateInfo);
 		TextureRHI = TextureCube;
 
 		// Write the contents of the texture.
@@ -467,7 +472,8 @@ public:
 		if (IsFeatureLevelSupported(GRHIShaderPlatform, ERHIFeatureLevel::SM5))
 		{
 			// Create the texture RHI.
-			FTextureCubeRHIRef TextureCubeArray = RHICreateTextureCubeArray(1,1,PF_B8G8R8A8,1,TexCreate_ShaderResource,NULL);
+			FRHIResourceCreateInfo CreateInfo;
+			FTextureCubeRHIRef TextureCubeArray = RHICreateTextureCubeArray(1,1,PF_B8G8R8A8,1,TexCreate_ShaderResource,CreateInfo);
 			TextureRHI = TextureCubeArray;
 
 			for(uint32 FaceIndex = 0;FaceIndex < 6;FaceIndex++)
diff --git a/Engine/Source/Runtime/Renderer/Private/PostProcess/RenderTargetPool.cpp b/Engine/Source/Runtime/Renderer/Private/PostProcess/RenderTargetPool.cpp
--- a/Engine/Source/Runtime/Renderer/Private/PostProcess/RenderTargetPool.cpp
+++ b/Engine/Source/Runtime/Renderer/Private/PostProcess/RenderTargetPool.cpp
@@ -176,6 +176,8 @@ bool FRenderTargetPool::FindFreeElement(const FPooledRenderTargetDesc& Desc, TRe
 		// TexCreate_UAV should be used on Desc.TargetableFlags
 		check(!(Desc.Flags & TexCreate_UAV));
 
+		FRHIResourceCreateInfo CreateInfo;
+
 		if(Desc.TargetableFlags & (TexCreate_RenderTargetable | TexCreate_DepthStencilTargetable | TexCreate_UAV))
 		{
 			if(Desc.Is2DTexture())
@@ -188,6 +190,7 @@ bool FRenderTargetPool::FindFreeElement(const FPooledRenderTargetDesc& Desc, TRe
 					Desc.Flags,
 					Desc.TargetableFlags,
 					Desc.bForceSeparateTargetAndShaderResource,
+					CreateInfo,
 					(FTexture2DRHIRef&)Found->RenderTargetItem.TargetableTexture,
 					(FTexture2DRHIRef&)Found->RenderTargetItem.ShaderResourceTexture,
 					Desc.NumSamples
@@ -202,7 +205,7 @@ bool FRenderTargetPool::FindFreeElement(const FPooledRenderTargetDesc& Desc, TRe
 					Desc.Format,
 					Desc.NumMips,
 					Desc.TargetableFlags,
-					NULL);
+					CreateInfo);
 
 				// similar to RHICreateTargetableShaderResource2D
 				Found->RenderTargetItem.TargetableTexture = Found->RenderTargetItem.ShaderResourceTexture;
@@ -220,6 +223,7 @@ bool FRenderTargetPool::FindFreeElement(const FPooledRenderTargetDesc& Desc, TRe
 						Desc.Flags,
 						Desc.TargetableFlags,
 						false,
+						CreateInfo,
 						(FTextureCubeRHIRef&)Found->RenderTargetItem.TargetableTexture,
 						(FTextureCubeRHIRef&)Found->RenderTargetItem.ShaderResourceTexture
 						);
@@ -233,6 +237,7 @@ bool FRenderTargetPool::FindFreeElement(const FPooledRenderTargetDesc& Desc, TRe
 						Desc.Flags,
 						Desc.TargetableFlags,
 						false,
+						CreateInfo,
 						(FTextureCubeRHIRef&)Found->RenderTargetItem.TargetableTexture,
 						(FTextureCubeRHIRef&)Found->RenderTargetItem.ShaderResourceTexture
 						);
@@ -254,7 +259,7 @@ bool FRenderTargetPool::FindFreeElement(const FPooledRenderTargetDesc& Desc, TRe
 					Desc.NumMips,
 					Desc.NumSamples,
 					Desc.Flags,
-					NULL);
+					CreateInfo);
 			}
 			else if(Desc.Is3DTexture())
 			{
@@ -265,19 +270,19 @@ bool FRenderTargetPool::FindFreeElement(const FPooledRenderTargetDesc& Desc, TRe
 					Desc.Format,
 					Desc.NumMips,
 					Desc.Flags,
-					NULL);
+					CreateInfo);
 			}
 			else 
 			{
 				check(Desc.IsCubemap());
 				if(Desc.IsArray())
 				{
-					FTextureCubeRHIRef CubeTexture = RHICreateTextureCubeArray(Desc.Extent.X,Desc.ArraySize,Desc.Format,Desc.NumMips,Desc.Flags | Desc.TargetableFlags | TexCreate_ShaderResource,NULL);
+					FTextureCubeRHIRef CubeTexture = RHICreateTextureCubeArray(Desc.Extent.X,Desc.ArraySize,Desc.Format,Desc.NumMips,Desc.Flags | Desc.TargetableFlags | TexCreate_ShaderResource,CreateInfo);
 					Found->RenderTargetItem.TargetableTexture = Found->RenderTargetItem.ShaderResourceTexture = CubeTexture;
 				}
 				else
 				{
-					FTextureCubeRHIRef CubeTexture = RHICreateTextureCube(Desc.Extent.X,Desc.Format,Desc.NumMips,Desc.Flags | Desc.TargetableFlags | TexCreate_ShaderResource,NULL);
+					FTextureCubeRHIRef CubeTexture = RHICreateTextureCube(Desc.Extent.X,Desc.Format,Desc.NumMips,Desc.Flags | Desc.TargetableFlags | TexCreate_ShaderResource,CreateInfo);
 					Found->RenderTargetItem.TargetableTexture = Found->RenderTargetItem.ShaderResourceTexture = CubeTexture;
 				}
 			}
diff --git a/Engine/Source/Runtime/Renderer/Private/PostProcess/RenderTargetPool.h b/Engine/Source/Runtime/Renderer/Private/PostProcess/RenderTargetPool.h
--- a/Engine/Source/Runtime/Renderer/Private/PostProcess/RenderTargetPool.h
+++ b/Engine/Source/Runtime/Renderer/Private/PostProcess/RenderTargetPool.h
@@ -291,6 +291,9 @@ private:
 	// sorted by size
 	SMemoryStats ComputeView();
 
+	// can be optimized to not be generated each time
+	void GenerateVRamAllocationUsage(TArray<FVRamAllocation>& Out);
+
 	friend struct FPooledRenderTarget;
 };
 
diff --git a/Engine/Source/Runtime/SlateRHIRenderer/Private/SlateRHIRenderer.cpp b/Engine/Source/Runtime/SlateRHIRenderer/Private/SlateRHIRenderer.cpp
--- a/Engine/Source/Runtime/SlateRHIRenderer/Private/SlateRHIRenderer.cpp
+++ b/Engine/Source/Runtime/SlateRHIRenderer/Private/SlateRHIRenderer.cpp
@@ -49,6 +49,7 @@ static FMatrix CreateProjectionMatrix( uint32 Width, uint32 Height )
 
 void FSlateCrashReportResource::InitDynamicRHI()
 {
+	FRHIResourceCreateInfo CreateInfo;
 	CrashReportBuffer = RHICreateTexture2D(
 		VirtualScreen.Width() * CrashTrackerConstants::ScreenScaling,
 		VirtualScreen.Height() * CrashTrackerConstants::ScreenScaling,
@@ -56,7 +57,7 @@ void FSlateCrashReportResource::InitDynamicRHI()
 		1,
 		1,
 		TexCreate_RenderTargetable,
-		NULL
+		CreateInfo
 		);
 
 	for (int32 i = 0; i < 2; ++i)
@@ -68,7 +69,7 @@ void FSlateCrashReportResource::InitDynamicRHI()
 			1,
 			1,
 			TexCreate_CPUReadback,
-			NULL
+			CreateInfo
 			);
 	}
 	
@@ -124,7 +125,8 @@ void FSlateRHIRenderer::FViewportInfo::RecreateDepthBuffer_RenderThread()
 	if (bRequiresStencilTest)
 	{		
 		FTexture2DRHIRef ShaderResourceUnused;
-		RHICreateTargetableShaderResource2D( Width, Height, PF_DepthStencil, 1, TexCreate_None, TexCreate_DepthStencilTargetable, false, DepthStencil, ShaderResourceUnused );
+		FRHIResourceCreateInfo CreateInfo;
+		RHICreateTargetableShaderResource2D( Width, Height, PF_DepthStencil, 1, TexCreate_None, TexCreate_DepthStencilTargetable, false, CreateInfo, DepthStencil, ShaderResourceUnused );
 		check( IsValidRef(DepthStencil) );
 	}
 }
diff --git a/Engine/Source/Runtime/Windows/D3D11RHI/Private/D3D11RHIPrivate.h b/Engine/Source/Runtime/Windows/D3D11RHI/Private/D3D11RHIPrivate.h
--- a/Engine/Source/Runtime/Windows/D3D11RHI/Private/D3D11RHIPrivate.h
+++ b/Engine/Source/Runtime/Windows/D3D11RHI/Private/D3D11RHIPrivate.h
@@ -483,9 +483,9 @@ protected:
 
 	template<typename BaseResourceType>
 	TD3D11Texture2D<BaseResourceType>* CreateD3D11Texture2D(uint32 SizeX,uint32 SizeY,uint32 SizeZ,bool bTextureArray,bool CubeTexture,uint8 Format,
-		uint32 NumMips,uint32 NumSamples,uint32 Flags,FResourceBulkDataInterface* BulkData = NULL);
+		uint32 NumMips,uint32 NumSamples,uint32 Flags,FRHIResourceCreateInfo& CreateInfo);
 
-	FD3D11Texture3D* CreateD3D11Texture3D(uint32 SizeX,uint32 SizeY,uint32 SizeZ,uint8 Format,uint32 NumMips,uint32 Flags,FResourceBulkDataInterface* BulkData);
+	FD3D11Texture3D* CreateD3D11Texture3D(uint32 SizeX,uint32 SizeY,uint32 SizeZ,uint8 Format,uint32 NumMips,uint32 Flags,FRHIResourceCreateInfo& CreateInfo);
 
 	/** Initializes the constant buffers.  Called once at RHI initialization time. */
 	void InitConstantBuffers();
@@ -750,4 +750,353 @@ public:
 	static FFastVRAMAllocator* GetFastVRAMAllocator();
 };
 
+
+
+
+
+// 1d, 31 bit (uses the sign bit for internal use), O(n) where n is the amount of elements stored
+// does not enforce any alignment
+// unoccupied regions get compacted but occupied don't get compacted
+class FRangeAllocator
+{
+public:
+
+	struct FRange
+	{
+		// not valid
+		FRange()
+			: Start(0)
+			, Size(0)
+		{
+			check(!IsValid());
+		}
+
+		void SetOccupied(int32 InStart, int32 InSize)
+		{
+			check(InStart >= 0);
+			check(InSize > 0);
+
+			Start = InStart;
+			Size = InSize;
+			check(IsOccupied());
+		}
+
+		void SetUnOccupied(int32 InStart, int32 InSize)
+		{
+			check(InStart >= 0);
+			check(InSize > 0);
+
+			Start = InStart;
+			Size = -InSize;
+			check(!IsOccupied());
+		}
+
+		bool IsValid() { return Size != 0; }
+
+		bool IsOccupied() const { return Size > 0; }
+		uint32 ComputeSize() const { return (Size > 0) ? Size : -Size; }
+
+		// @apram InSize can be <0 to remove from the size
+		void ExtendUnoccupied(int32 InSize) { check(!IsOccupied()); Size -= InSize; }
+
+		void MakeOccupied(int32 InSize) { check(InSize > 0); check(!IsOccupied()); Size = InSize; }
+		void MakeUnOccupied() { check(IsOccupied()); Size = -Size; }
+
+		bool operator==(const FRange& rhs) const { return Start == rhs.Start && Size == rhs.Size; }
+
+		int32 GetStart() { return Start; }
+		int32 GetEnd() { return Start + ComputeSize(); }
+
+	private:
+		// in bytes
+		int32 Start;
+		// in bytes, 0:not valid, <0:unoccupied, >0:occupied
+		int32 Size;
+	};
+public:
+
+	// constructor
+	FRangeAllocator(uint32 TotalSize)
+	{
+		FRange NewRange;
+
+		NewRange.SetUnOccupied(0, TotalSize);
+
+		Entries.Add(NewRange);
+	}
+
+	// specified range must be non occupied
+	void OccupyRange(FRange InRange)
+	{
+		check(InRange.IsValid());
+		check(InRange.IsOccupied());
+
+		for(uint32 i = 0, Num = Entries.Num(); i < Num; ++i)
+		{
+			FRange& ref = Entries[i];
+
+			if(!ref.IsOccupied())
+			{
+				int32 OverlapSize = ref.GetEnd() - InRange.GetStart();
+
+				if(OverlapSize > 0)
+				{
+					int32 FrontCutSize = InRange.GetStart() - ref.GetStart();
+
+					// there is some front part we cut off
+					if(FrontCutSize > 0)
+					{
+						FRange NewFrontRange;
+
+						NewFrontRange.SetUnOccupied(InRange.GetStart(), ref.ComputeSize() - FrontCutSize);
+
+						ref.SetUnOccupied(ref.GetStart(), FrontCutSize);
+
+						++i;
+
+						// remaining is added behind the found element
+						Entries.Insert(NewFrontRange, i);
+
+						// don't access ref or Num any more - Entries[] might be reallocated
+					}
+
+					check(Entries[i].GetStart() == InRange.GetStart());
+
+					int32 BackCutSize = Entries[i].ComputeSize() - InRange.ComputeSize();
+
+					// otherwise the range was already occupied or not enough space was left (internal error)
+					check(BackCutSize >= 0);
+
+					// there is some back part we cut off
+					if(BackCutSize > 0)
+					{
+						FRange NewBackRange;
+
+						NewBackRange.SetUnOccupied(Entries[i].GetStart() + InRange.ComputeSize(), BackCutSize);
+
+						Entries.Insert(NewBackRange, i + 1);
+					}
+
+					Entries[i] = InRange;
+					return;
+				}
+			}
+		}
+	}
+
+
+	// All resources in ESRAM must be 64KiB aligned
+//	uint32 AlignedByteOffset = FFastVRAMAllocator::RoundUpToNextMultiple(ESRAMByteOffset, ESRAMMinumumAlignment );
+
+
+	// @param InSize >0
+	FRange AllocRange(uint32 InSize)//, uint32 Alignment)
+	{
+		check(InSize > 0);
+
+		for(uint32 i = 0, Num = Entries.Num(); i < Num; ++i)
+		{
+			FRange& ref = Entries[i];
+
+			if(!ref.IsOccupied())
+			{
+				uint32 RefSize = ref.ComputeSize();
+
+				// take the first fitting one - later we could optimize for minimal fragmentation
+				if(RefSize >= InSize)
+				{
+					ref.MakeOccupied(InSize);
+
+					FRange Ret = ref;
+
+					if(RefSize > InSize)
+					{
+						FRange NewRange;
+
+						NewRange.SetUnOccupied(ref.GetEnd(), RefSize - InSize);
+
+						// remaining is added behind the found element
+						Entries.Insert(NewRange, i + 1);
+					}
+					return Ret;
+				}
+			}
+		}
+
+		// nothing found
+		return FRange();
+	}
+
+	// @param In needs to be what was returned by AllocRange()
+	void ReleaseRange(FRange In)
+	{
+		int32 Index = Entries.Find(In);
+
+		check(Index != INDEX_NONE);
+
+		FRange& refIndex = Entries[Index];
+
+		refIndex.MakeUnOccupied();
+
+		Compacten(Index);
+	}
+
+	// for debugging
+	uint32 GetNumEntries() const { return Entries.Num(); }
+
+	// for debugging
+	uint32 ComputeUnoccupiedSize() const
+	{
+		uint32 Ret = 0;
+
+		for(uint32 i = 0, Num = Entries.Num(); i < Num; ++i)
+		{
+			const FRange& ref = Entries[i];
+
+			if(!ref.IsOccupied())
+			{
+				uint32 RefSize = ref.ComputeSize();
+
+				Ret += RefSize;
+			}
+		}
+
+		return Ret;
+	}
+
+private:
+	// compact unoccupied ranges
+	void Compacten(uint32 StartIndex)
+	{
+		check(!Entries[StartIndex].IsOccupied());
+
+		if(StartIndex && !Entries[StartIndex-1].IsOccupied())
+		{
+			// Seems we can combine with the element before,
+			// searching further is not needed as we assume the buffer was compact before the last change.
+			--StartIndex;
+		}
+
+		uint32 ElementsToRemove = 0;
+		uint32 SizeGained = 0;
+
+		for(uint32 i = StartIndex + 1, Num = Entries.Num(); i < Num; ++i)
+		{
+			FRange& ref = Entries[i];
+
+			if(!ref.IsOccupied())
+			{
+				++ElementsToRemove;
+				SizeGained += ref.ComputeSize();
+			}
+			else
+			{
+				break;
+			}
+		}
+
+		if(ElementsToRemove)
+		{
+			Entries.RemoveAt(StartIndex + 1, ElementsToRemove, false);
+			Entries[StartIndex].ExtendUnoccupied(SizeGained);
+		}
+	}
+
+	static void Test()
+	{
+		// testing code
+#if !UE_BUILD_SHIPPING
+		{
+			// create
+			FRangeAllocator A(10);
+			check(A.GetNumEntries() == 1);
+			check(A.ComputeUnoccupiedSize() == 10);
+
+			// successfully alloc
+			FRangeAllocator::FRange a = A.AllocRange(3);
+			check(a.GetStart() == 0);
+			check(a.GetEnd() == 3);
+			check(a.IsOccupied());
+			check(A.GetNumEntries() == 2);
+			check(A.ComputeUnoccupiedSize() == 7);
+
+			// successfully alloc
+			FRangeAllocator::FRange b = A.AllocRange(4);
+			check(b.GetStart() == 3);
+			check(b.GetEnd() == 7);
+			check(b.IsOccupied());
+			check(A.GetNumEntries() == 3);
+			check(A.ComputeUnoccupiedSize() == 3);
+
+			// failed alloc
+			FRangeAllocator::FRange c = A.AllocRange(4);
+			check(!c.IsValid());
+			check(!c.IsOccupied());
+			check(A.GetNumEntries() == 3);
+			check(A.ComputeUnoccupiedSize() == 3);
+
+			// successfully alloc
+			FRangeAllocator::FRange d = A.AllocRange(3);
+			check(d.GetStart() == 7);
+			check(d.GetEnd() == 10);
+			check(d.IsOccupied());
+			check(A.GetNumEntries() == 3);
+			check(A.ComputeUnoccupiedSize() == 0);
+
+			A.ReleaseRange(b);
+			check(A.GetNumEntries() == 3);
+			check(A.ComputeUnoccupiedSize() == 4);
+
+			A.ReleaseRange(a);
+			check(A.GetNumEntries() == 2);
+			check(A.ComputeUnoccupiedSize() == 7);
+
+			A.ReleaseRange(d);
+			check(A.GetNumEntries() == 1);
+			check(A.ComputeUnoccupiedSize() == 10);
+
+			// we are back to a clean start
+
+			FRangeAllocator::FRange e = A.AllocRange(10);
+			check(e.GetStart() == 0);
+			check(e.GetEnd() == 10);
+			check(e.IsOccupied());
+			check(A.GetNumEntries() == 1);
+			check(A.ComputeUnoccupiedSize() == 0);
+
+			A.ReleaseRange(e);
+			check(A.GetNumEntries() == 1);
+			check(A.ComputeUnoccupiedSize() == 10);
+
+			// we are back to a clean start
+
+			// create define range we want to block out
+			FRangeAllocator::FRange f;
+			f.SetOccupied(2, 4);
+			A.OccupyRange(f);
+			check(A.GetNumEntries() == 3);
+			check(A.ComputeUnoccupiedSize() == 6);
+
+			FRangeAllocator::FRange g = A.AllocRange(2);
+			check(g.GetStart() == 0);
+			check(g.GetEnd() == 2);
+			check(g.IsOccupied());
+			check(A.GetNumEntries() == 3);
+			check(A.ComputeUnoccupiedSize() == 4);
+
+			FRangeAllocator::FRange h = A.AllocRange(4);
+			check(h.GetStart() == 6);
+			check(h.GetEnd() == 10);
+			check(h.IsOccupied());
+			check(A.GetNumEntries() == 3);
+			check(A.ComputeUnoccupiedSize() == 0);
+		}
+#endif // !UE_BUILD_SHIPPING
+	}
+
+	// ordered from small to large (for efficient compactening)
+	TArray<FRange> Entries;
+};
+
+
 #endif
diff --git a/Engine/Source/Runtime/Windows/D3D11RHI/Private/D3D11Texture.cpp b/Engine/Source/Runtime/Windows/D3D11RHI/Private/D3D11Texture.cpp
--- a/Engine/Source/Runtime/Windows/D3D11RHI/Private/D3D11Texture.cpp
+++ b/Engine/Source/Runtime/Windows/D3D11RHI/Private/D3D11Texture.cpp
@@ -454,7 +454,7 @@ void SafeCreateTexture2D(ID3D11Device* Direct3DDevice, const D3D11_TEXTURE2D_DES
 
 template<typename BaseResourceType>
 TD3D11Texture2D<BaseResourceType>* FD3D11DynamicRHI::CreateD3D11Texture2D(uint32 SizeX,uint32 SizeY,uint32 SizeZ,bool bTextureArray,bool bCubeTexture,uint8 Format,
-	uint32 NumMips,uint32 NumSamples,uint32 Flags, FResourceBulkDataInterface* BulkData)
+	uint32 NumMips,uint32 NumSamples,uint32 Flags, FRHIResourceCreateInfo& CreateInfo)
 {
 	check(SizeX > 0 && SizeY > 0 && NumMips > 0);
 
@@ -651,9 +651,9 @@ TD3D11Texture2D<BaseResourceType>* FD3D11DynamicRHI::CreateD3D11Texture2D(uint32
 	{
 		TArray<D3D11_SUBRESOURCE_DATA> SubResourceData;
 
-		if (BulkData)
+		if (CreateInfo.BulkData)
 		{
-			uint8* Data = (uint8*)BulkData->GetResourceBulkData();
+			uint8* Data = (uint8*)CreateInfo.BulkData->GetResourceBulkData();
 
 			// each mip of each array slice counts as a subresource
 			SubResourceData.AddZeroed(NumMips * SizeZ);
@@ -688,7 +688,7 @@ TD3D11Texture2D<BaseResourceType>* FD3D11DynamicRHI::CreateD3D11Texture2D(uint32
 		else
 #endif
 		{
-			SafeCreateTexture2D(Direct3DDevice, &TextureDesc, BulkData != NULL ? (const D3D11_SUBRESOURCE_DATA*)SubResourceData.GetData() : NULL, TextureResource.GetInitReference());
+			SafeCreateTexture2D(Direct3DDevice, &TextureDesc, CreateInfo.BulkData != NULL ? (const D3D11_SUBRESOURCE_DATA*)SubResourceData.GetData() : NULL, TextureResource.GetInitReference());
 		}
 
 		if(bCreateRTV)
@@ -848,7 +848,7 @@ TD3D11Texture2D<BaseResourceType>* FD3D11DynamicRHI::CreateD3D11Texture2D(uint32
 	return Texture2D;
 }
 
-FD3D11Texture3D* FD3D11DynamicRHI::CreateD3D11Texture3D(uint32 SizeX,uint32 SizeY,uint32 SizeZ,uint8 Format,uint32 NumMips,uint32 Flags,FResourceBulkDataInterface* BulkData)
+FD3D11Texture3D* FD3D11DynamicRHI::CreateD3D11Texture3D(uint32 SizeX,uint32 SizeY,uint32 SizeZ,uint8 Format,uint32 NumMips,uint32 Flags,FRHIResourceCreateInfo& CreateInfo)
 {
 	SCOPE_CYCLE_COUNTER(STAT_D3D11CreateTextureTime);
 	
@@ -898,9 +898,9 @@ FD3D11Texture3D* FD3D11DynamicRHI::CreateD3D11Texture3D(uint32 SizeX,uint32 Size
 
 	TArray<D3D11_SUBRESOURCE_DATA> SubResourceData;
 
-	if (BulkData)
+	if (CreateInfo.BulkData)
 	{
-		uint8* Data = (uint8*)BulkData->GetResourceBulkData();
+		uint8* Data = (uint8*)CreateInfo.BulkData->GetResourceBulkData();
 		SubResourceData.AddZeroed(NumMips);
 		uint32 MipOffset = 0;
 		for(uint32 MipIndex = 0;MipIndex < NumMips;++MipIndex)
@@ -924,7 +924,7 @@ FD3D11Texture3D* FD3D11DynamicRHI::CreateD3D11Texture3D(uint32 SizeX,uint32 Size
 
 	TRefCountPtr<ID3D11Texture3D> TextureResource;
 	VERIFYD3D11CREATETEXTURERESULT(
-		Direct3DDevice->CreateTexture3D(&TextureDesc,BulkData != NULL ? (const D3D11_SUBRESOURCE_DATA*)SubResourceData.GetData() : NULL,TextureResource.GetInitReference()),
+		Direct3DDevice->CreateTexture3D(&TextureDesc,CreateInfo.BulkData != NULL ? (const D3D11_SUBRESOURCE_DATA*)SubResourceData.GetData() : NULL,TextureResource.GetInitReference()),
 		SizeX,
 		SizeY,
 		SizeZ,
@@ -974,9 +974,9 @@ FD3D11Texture3D* FD3D11DynamicRHI::CreateD3D11Texture3D(uint32 SizeX,uint32 Size
 	2D texture support.
 -----------------------------------------------------------------------------*/
 
-FTexture2DRHIRef FD3D11DynamicRHI::RHICreateTexture2D(uint32 SizeX,uint32 SizeY,uint8 Format,uint32 NumMips,uint32 NumSamples,uint32 Flags,FResourceBulkDataInterface* BulkData)
+FTexture2DRHIRef FD3D11DynamicRHI::RHICreateTexture2D(uint32 SizeX,uint32 SizeY,uint8 Format,uint32 NumMips,uint32 NumSamples,uint32 Flags,FRHIResourceCreateInfo& CreateInfo)
 {
-	return CreateD3D11Texture2D<FD3D11BaseTexture2D>(SizeX,SizeY,1,false,false,Format,NumMips,NumSamples,Flags,BulkData);
+	return CreateD3D11Texture2D<FD3D11BaseTexture2D>(SizeX,SizeY,1,false,false,Format,NumMips,NumSamples,Flags,CreateInfo);
 }
 
 FTexture2DRHIRef FD3D11DynamicRHI::RHIAsyncCreateTexture2D(uint32 SizeX,uint32 SizeY,uint8 Format,uint32 NumMips,uint32 Flags,void** InitialMipData,uint32 NumInitialMips)
@@ -1108,16 +1108,16 @@ void FD3D11DynamicRHI::RHICopySharedMips(FTexture2DRHIParamRef DestTexture2DRHI,
 	}
 }
 
-FTexture2DArrayRHIRef FD3D11DynamicRHI::RHICreateTexture2DArray(uint32 SizeX,uint32 SizeY,uint32 SizeZ,uint8 Format,uint32 NumMips,uint32 Flags,FResourceBulkDataInterface* BulkData)
+FTexture2DArrayRHIRef FD3D11DynamicRHI::RHICreateTexture2DArray(uint32 SizeX,uint32 SizeY,uint32 SizeZ,uint8 Format,uint32 NumMips,uint32 Flags,FRHIResourceCreateInfo& CreateInfo)
 {
 	check(SizeZ >= 1);
-	return CreateD3D11Texture2D<FD3D11BaseTexture2DArray>(SizeX,SizeY,SizeZ,true,false,Format,NumMips,1,Flags,BulkData);
+	return CreateD3D11Texture2D<FD3D11BaseTexture2DArray>(SizeX,SizeY,SizeZ,true,false,Format,NumMips,1,Flags,CreateInfo);
 }
 
-FTexture3DRHIRef FD3D11DynamicRHI::RHICreateTexture3D(uint32 SizeX,uint32 SizeY,uint32 SizeZ,uint8 Format,uint32 NumMips,uint32 Flags,FResourceBulkDataInterface* BulkData)
+FTexture3DRHIRef FD3D11DynamicRHI::RHICreateTexture3D(uint32 SizeX,uint32 SizeY,uint32 SizeZ,uint8 Format,uint32 NumMips,uint32 Flags,FRHIResourceCreateInfo& CreateInfo)
 {
 	check(SizeZ >= 1);
-	return CreateD3D11Texture3D(SizeX,SizeY,SizeZ,Format,NumMips,Flags,BulkData);
+	return CreateD3D11Texture3D(SizeX,SizeY,SizeZ,Format,NumMips,Flags,CreateInfo);
 }
 
 void FD3D11DynamicRHI::RHIGetResourceInfo(FTextureRHIParamRef Ref, FRHIResourceInfo& OutInfo)
@@ -1240,7 +1240,8 @@ FTexture2DRHIRef FD3D11DynamicRHI::RHIAsyncReallocateTexture2D(FTexture2DRHIPara
 	DYNAMIC_CAST_D3D11RESOURCE(Texture2D,Texture2D);
 
 	// Allocate a new texture.
-	FD3D11Texture2D* NewTexture2D = CreateD3D11Texture2D<FD3D11BaseTexture2D>(NewSizeX,NewSizeY,1,false,false,Texture2D->GetFormat(),NewMipCount,1,Texture2D->GetFlags());
+	FRHIResourceCreateInfo CreateInfo;
+	FD3D11Texture2D* NewTexture2D = CreateD3D11Texture2D<FD3D11BaseTexture2D>(NewSizeX,NewSizeY,1,false,false,Texture2D->GetFormat(),NewMipCount,1,Texture2D->GetFlags(),CreateInfo);
 	
 	// Use the GPU to asynchronously copy the old mip-maps into the new texture.
 	const uint32 NumSharedMips = FMath::Min(Texture2D->GetNumMips(),NewTexture2D->GetNumMips());
@@ -1463,14 +1464,14 @@ void FD3D11DynamicRHI::RHIUpdateTexture3D(FTexture3DRHIParamRef TextureRHI,uint3
 /*-----------------------------------------------------------------------------
 	Cubemap texture support.
 -----------------------------------------------------------------------------*/
-FTextureCubeRHIRef FD3D11DynamicRHI::RHICreateTextureCube(uint32 Size, uint8 Format, uint32 NumMips, uint32 Flags, FResourceBulkDataInterface* BulkData)
+FTextureCubeRHIRef FD3D11DynamicRHI::RHICreateTextureCube(uint32 Size, uint8 Format, uint32 NumMips, uint32 Flags, FRHIResourceCreateInfo& CreateInfo)
 {
-	return CreateD3D11Texture2D<FD3D11BaseTextureCube>(Size,Size,6,false,true,Format,NumMips,1,Flags,BulkData);
+	return CreateD3D11Texture2D<FD3D11BaseTextureCube>(Size,Size,6,false,true,Format,NumMips,1,Flags,CreateInfo);
 }
 
-FTextureCubeRHIRef FD3D11DynamicRHI::RHICreateTextureCubeArray(uint32 Size, uint32 ArraySize, uint8 Format, uint32 NumMips, uint32 Flags, FResourceBulkDataInterface* BulkData)
+FTextureCubeRHIRef FD3D11DynamicRHI::RHICreateTextureCubeArray(uint32 Size, uint32 ArraySize, uint8 Format, uint32 NumMips, uint32 Flags, FRHIResourceCreateInfo& CreateInfo)
 {
-	return CreateD3D11Texture2D<FD3D11BaseTextureCube>(Size,Size,6 * ArraySize,true,true,Format,NumMips,1,Flags,BulkData);
+	return CreateD3D11Texture2D<FD3D11BaseTextureCube>(Size,Size,6 * ArraySize,true,true,Format,NumMips,1,Flags,CreateInfo);
 }
 
 void* FD3D11DynamicRHI::RHILockTextureCubeFace(FTextureCubeRHIParamRef TextureCubeRHI,uint32 FaceIndex,uint32 ArrayIndex,uint32 MipIndex,EResourceLockMode LockMode,uint32& DestStride,bool bLockWithinMiptail)
