* fixed build

diff --git a/Engine/Source/Runtime/Windows/D3D11RHI/Private/D3D11RHIPrivate.h b/Engine/Source/Runtime/Windows/D3D11RHI/Private/D3D11RHIPrivate.h
--- a/Engine/Source/Runtime/Windows/D3D11RHI/Private/D3D11RHIPrivate.h
+++ b/Engine/Source/Runtime/Windows/D3D11RHI/Private/D3D11RHIPrivate.h
@@ -1002,6 +1002,7 @@ private:
 		}
 	}
 
+public:
 	static void Test()
 	{
 		// testing code
@@ -1094,6 +1095,8 @@ private:
 #endif // !UE_BUILD_SHIPPING
 	}
 
+private:
+
 	// ordered from small to large (for efficient compactening)
 	TArray<FRange> Entries;
 };
