Ensure that precaching of a sound completes instead of failing to start.
TTP #335967

diff --git a/Engine/Source/Runtime/Windows/XAudio2/Private/XAudio2Buffer.cpp b/Engine/Source/Runtime/Windows/XAudio2/Private/XAudio2Buffer.cpp
--- a/Engine/Source/Runtime/Windows/XAudio2/Private/XAudio2Buffer.cpp
+++ b/Engine/Source/Runtime/Windows/XAudio2/Private/XAudio2Buffer.cpp
@@ -425,12 +425,7 @@ FXAudio2SoundBuffer* FXAudio2SoundBuffer::CreateNativeBuffer( FXAudio2Device* XA
 	// Check to see if thread has finished decompressing on the other thread
 	if( Wave->AudioDecompressor != NULL )
 	{
-		if (!Wave->AudioDecompressor->IsDone())
-		{
-			// Don't play this sound just yet
-			UE_LOG(LogAudio, Log, TEXT( "Waiting for sound to decompress: %s" ), *Wave->GetName() );
-			return( NULL );
-		}
+		Wave->AudioDecompressor->EnsureCompletion();
 
 		// Remove the decompressor
 		delete Wave->AudioDecompressor;
