UE4: Fixed a minor issue in error logging in XAudio.

diff --git a/Engine/Source/Runtime/Windows/XAudio2/Private/XAudio2Device.cpp b/Engine/Source/Runtime/Windows/XAudio2/Private/XAudio2Device.cpp
--- a/Engine/Source/Runtime/Windows/XAudio2/Private/XAudio2Device.cpp
+++ b/Engine/Source/Runtime/Windows/XAudio2/Private/XAudio2Device.cpp
@@ -268,9 +268,11 @@ bool FXAudio2Device::ValidateAPICall( const TCHAR* Function, int32 ErrorCode )
 		case XAUDIO2_E_DEVICE_INVALIDATED:
 			UE_LOG(LogAudio, Warning, TEXT( "%s error: Device Invalidated" ), Function );
 			break;
-		};
 
-		UE_LOG(LogAudio, Warning, TEXT( "%s error: Unhandled error code %d" ), Function, ErrorCode );
+		default:
+			UE_LOG(LogAudio, Warning, TEXT( "%s error: Unhandled error code %d" ), Function, ErrorCode );
+			break;
+		};
 
 		return( false );
 	}
