UE4 - merge first pass at parallel rendering to main

diff --git a/Engine/Source/Runtime/ShaderCore/Public/ShaderParameters.h b/Engine/Source/Runtime/ShaderCore/Public/ShaderParameters.h
--- a/Engine/Source/Runtime/ShaderCore/Public/ShaderParameters.h
+++ b/Engine/Source/Runtime/ShaderCore/Public/ShaderParameters.h
@@ -9,7 +9,6 @@
 #include "ShaderCore.h"
 #include "RenderResource.h"
 #include "UniformBuffer.h"
-#include "RHIStaticStates.h"
 
 enum EShaderParameterFlags
 {
@@ -96,299 +95,6 @@ private:
 #endif
 };
 
-/**
- * Sets the value of a  shader parameter.  Template'd on shader type
- * A template parameter specified the type of the parameter value.
- * NOTE: Shader should be the param ref type, NOT the param type, since Shader is passed by value. 
- * Otherwise AddRef/ReleaseRef will be called many times.
- */
-template<typename ShaderRHIParamRef, class ParameterType>
-void SetShaderValue(
-	ShaderRHIParamRef Shader,
-	const FShaderParameter& Parameter,
-	const ParameterType& Value,
-	uint32 ElementIndex = 0
-	)
-{
-	checkAtCompileTime(!TIsPointerType<ParameterType>::Value, ErrorPassByValue);
-
-	const uint32 AlignedTypeSize = Align(sizeof(ParameterType),ShaderArrayElementAlignBytes);
-	const int32 NumBytesToSet = FMath::Min<int32>(sizeof(ParameterType),Parameter.GetNumBytes() - ElementIndex * AlignedTypeSize);
-
-	// This will trigger if the parameter was not serialized
-	checkSlow(Parameter.IsInitialized());
-
-	if(NumBytesToSet > 0)
-	{
-		RHISetShaderParameter(
-			Shader,
-			Parameter.GetBufferIndex(),
-			Parameter.GetBaseIndex() + ElementIndex * AlignedTypeSize,
-			(uint32)NumBytesToSet,
-			&Value
-			);
-	}
-}
-
-/** Specialization of the above for C++ bool type. */
-template<typename ShaderRHIParamRef>
-void SetShaderValue(
-	ShaderRHIParamRef Shader,
-	const FShaderParameter& Parameter,
-	const bool& Value,
-	uint32 ElementIndex = 0
-	)
-{
-	const uint32 BoolValue = Value;
-	SetShaderValue(Shader, Parameter, BoolValue, ElementIndex);
-}
-
-/**
- * Sets the value of a shader parameter array.  Template'd on shader type
- * A template parameter specified the type of the parameter value.
- * NOTE: Shader should be the param ref type, NOT the param type, since Shader is passed by value. 
- * Otherwise AddRef/ReleaseRef will be called many times.
- */
-template<typename ShaderRHIParamRef,class ParameterType>
-void SetShaderValueArray(
-	ShaderRHIParamRef Shader,
-	const FShaderParameter& Parameter,
-	const ParameterType* Values,
-	uint32 NumElements,
-	uint32 BaseElementIndex = 0
-	)
-{
-	const uint32 AlignedTypeSize = Align(sizeof(ParameterType),ShaderArrayElementAlignBytes);
-	const int32 NumBytesToSet = FMath::Min<int32>(NumElements * AlignedTypeSize,Parameter.GetNumBytes() - BaseElementIndex * AlignedTypeSize);
-
-	// This will trigger if the parameter was not serialized
-	checkSlow(Parameter.IsInitialized());
-
-	if(NumBytesToSet > 0)
-	{
-		RHISetShaderParameter(
-			Shader,
-			Parameter.GetBufferIndex(),
-			Parameter.GetBaseIndex() + BaseElementIndex * AlignedTypeSize,
-			(uint32)NumBytesToSet,
-			Values
-			);
-	}
-}
-
-/** Specialization of the above for C++ bool type. */
-template<typename ShaderRHIParamRef>
-void SetShaderValueArray(
-	ShaderRHIParamRef Shader,
-	const FShaderParameter& Parameter,
-	const bool* Values,
-	uint32 NumElements,
-	uint32 BaseElementIndex = 0
-	)
-{
-	UE_LOG(LogShaders, Fatal, TEXT("SetShaderValueArray does not support bool arrays."));
-}
-
-/**
- * Sets the value of a pixel shader bool parameter.
- */
-inline void SetPixelShaderBool(
-	FPixelShaderRHIParamRef PixelShader,
-	const FShaderParameter& Parameter,
-	bool Value
-	)
-{
-	// This will trigger if the parameter was not serialized
-	checkSlow(Parameter.IsInitialized());
-
-	if (Parameter.GetNumBytes() > 0)
-	{
-		// Convert to uint32 before passing to RHI
-		uint32 BoolValue = Value;
-		RHISetShaderParameter(
-			PixelShader,
-			Parameter.GetBufferIndex(),
-			Parameter.GetBaseIndex(),
-			sizeof(BoolValue),
-			&BoolValue
-			);
-	}
-}
-
-/**
- * Sets the value of a shader texture parameter.  Template'd on shader type
- */
-template<typename ShaderTypeRHIParamRef>
-FORCEINLINE void SetTextureParameter(
-	ShaderTypeRHIParamRef Shader,
-	const FShaderResourceParameter& TextureParameter,
-	const FShaderResourceParameter& SamplerParameter,
-	const FTexture* Texture,
-	uint32 ElementIndex = 0
-	)
-{
-	// This will trigger if the parameter was not serialized
-	checkSlow(TextureParameter.IsInitialized());
-	checkSlow(SamplerParameter.IsInitialized());
-	if(TextureParameter.IsBound())
-	{
-		Texture->LastRenderTime = FApp::GetCurrentTime();
-
-		if (ElementIndex < TextureParameter.GetNumResources())
-		{
-			RHISetShaderTexture(Shader,TextureParameter.GetBaseIndex() + ElementIndex,Texture->TextureRHI);
-		}
-	}
-	
-	// @todo ue4 samplerstate Should we maybe pass in two separate values? SamplerElement and TextureElement? Or never allow an array of samplers? Unsure best
-	// if there is a matching sampler for this texture array index (ElementIndex), then set it. This will help with this case:
-	//			Texture2D LightMapTextures[NUM_LIGHTMAP_COEFFICIENTS];
-	//			SamplerState LightMapTexturesSampler;
-	// In this case, we only set LightMapTexturesSampler when ElementIndex is 0, we don't set the sampler state for all 4 textures
-	// This assumes that the all textures want to use the same sampler state
-	if(SamplerParameter.IsBound())
-	{
-		if (ElementIndex < SamplerParameter.GetNumResources())
-		{
-			RHISetShaderSampler(Shader,SamplerParameter.GetBaseIndex() + ElementIndex,Texture->SamplerStateRHI);
-		}
-	}
-}
-
-/**
- * Sets the value of a shader texture parameter. Template'd on shader type.
- */
-template<typename ShaderTypeRHIParamRef>
-FORCEINLINE void SetTextureParameter(
-	ShaderTypeRHIParamRef Shader,
-	const FShaderResourceParameter& TextureParameter,
-	const FShaderResourceParameter& SamplerParameter,
-	FSamplerStateRHIParamRef SamplerStateRHI,
-	FTextureRHIParamRef TextureRHI,
-	uint32 ElementIndex = 0
-	)
-{
-	// This will trigger if the parameter was not serialized
-	checkSlow(TextureParameter.IsInitialized());
-	checkSlow(SamplerParameter.IsInitialized());
-	if(TextureParameter.IsBound())
-	{
-		if (ElementIndex < TextureParameter.GetNumResources())
-		{
-			RHISetShaderTexture(Shader,TextureParameter.GetBaseIndex() + ElementIndex,TextureRHI);
-		}
-	}
-	// @todo ue4 samplerstate Should we maybe pass in two separate values? SamplerElement and TextureElement? Or never allow an array of samplers? Unsure best
-	// if there is a matching sampler for this texture array index (ElementIndex), then set it. This will help with this case:
-	//			Texture2D LightMapTextures[NUM_LIGHTMAP_COEFFICIENTS];
-	//			SamplerState LightMapTexturesSampler;
-	// In this case, we only set LightMapTexturesSampler when ElementIndex is 0, we don't set the sampler state for all 4 textures
-	// This assumes that the all textures want to use the same sampler state
-	if(SamplerParameter.IsBound())
-	{
-		if (ElementIndex < SamplerParameter.GetNumResources())
-		{
-			RHISetShaderSampler(Shader,SamplerParameter.GetBaseIndex() + ElementIndex,SamplerStateRHI);
-		}
-	}
-}
-
-/**
- * Sets the value of a shader surface parameter (e.g. to access MSAA samples).
- * Template'd on shader type (e.g. pixel shader or compute shader).
- */
-template<typename ShaderTypeRHIParamRef>
-FORCEINLINE void SetTextureParameter(
-	ShaderTypeRHIParamRef Shader,
-	const FShaderResourceParameter& Parameter,
-	FTextureRHIParamRef NewTextureRHI
-	)
-{
-	if(Parameter.IsBound())
-	{
-		RHISetShaderTexture(
-			Shader,
-			Parameter.GetBaseIndex(),
-			NewTextureRHI
-			);
-	}
-}
-
-/**
- * Sets the value of a shader sampler parameter. Template'd on shader type.
- */
-template<typename ShaderTypeRHIParamRef>
-FORCEINLINE void SetSamplerParameter(
-	ShaderTypeRHIParamRef Shader,
-	const FShaderResourceParameter& Parameter,
-	FSamplerStateRHIParamRef SamplerStateRHI
-	)
-{
-	if(Parameter.IsBound())
-	{
-		RHISetShaderSampler(
-			Shader,
-			Parameter.GetBaseIndex(),
-			SamplerStateRHI
-			);
-	}
-}
-
-/**
- * Sets the value of a shader resource view parameter
- * Template'd on shader type (e.g. pixel shader or compute shader).
- */
-template<typename ShaderTypeRHIParamRef>
-FORCEINLINE void SetSRVParameter(
-	ShaderTypeRHIParamRef Shader,
-	const FShaderResourceParameter& Parameter,
-	FShaderResourceViewRHIParamRef NewShaderResourceViewRHI
-	)
-{
-	if(Parameter.IsBound())
-	{
-		RHISetShaderResourceViewParameter(
-			Shader,
-			Parameter.GetBaseIndex(),
-			NewShaderResourceViewRHI
-			);
-	}
-}
-
-/**
- * Sets the value of a unordered access view parameter
- */
-FORCEINLINE void SetUAVParameter(
-	FComputeShaderRHIParamRef ComputeShader,
-	const FShaderResourceParameter& Parameter,
-	FUnorderedAccessViewRHIParamRef NewUnorderedAccessViewRHI
-	)
-{
-	if(Parameter.IsBound())
-	{
-		RHISetUAVParameter(
-			ComputeShader,
-			Parameter.GetBaseIndex(),
-			NewUnorderedAccessViewRHI
-			);
-	}
-}
-
-template<typename TShaderRHIRef>
-inline void SetUAVParameterIfCS(TShaderRHIRef,const FShaderResourceParameter& UAVParameter,FUnorderedAccessViewRHIParamRef UAV)
-{
-}
-
-template<> inline void SetUAVParameterIfCS<FComputeShaderRHIParamRef>(FComputeShaderRHIParamRef ComputeShader,const FShaderResourceParameter& UAVParameter,FUnorderedAccessViewRHIParamRef UAV)
-{
-	SetUAVParameter(ComputeShader,UAVParameter,UAV);
-}
-
-template<> inline void SetUAVParameterIfCS<FComputeShaderRHIRef>(FComputeShaderRHIRef ComputeShader,const FShaderResourceParameter& UAVParameter,FUnorderedAccessViewRHIParamRef UAV)
-{
-	SetUAVParameter(ComputeShader,UAVParameter,UAV);
-}
-
 /** A class that binds either a UAV or SRV of a resource. */
 class FRWShaderParameter
 {
@@ -422,23 +128,12 @@ public:
 	}
 
 	template<typename TShaderRHIRef>
-	void SetBuffer(TShaderRHIRef Shader,const FRWBuffer& RWBuffer) const
-	{
-		SetSRVParameter(Shader,SRVParameter,RWBuffer.SRV);
-		SetUAVParameterIfCS(Shader,UAVParameter,RWBuffer.UAV);
-	}
+	inline void SetBuffer(FRHICommandList* RHICmdList, TShaderRHIRef Shader,const FRWBuffer& RWBuffer) const;
 
 	template<typename TShaderRHIRef>
-	void SetTexture(TShaderRHIRef Shader,const FTextureRHIParamRef Texture,FUnorderedAccessViewRHIParamRef UAV) const
-	{
-		SetTextureParameter(Shader,SRVParameter,Texture);
-		SetUAVParameterIfCS(Shader,UAVParameter,UAV);
-	}
+	inline void SetTexture(FRHICommandList* RHICmdList, TShaderRHIRef Shader,const FTextureRHIParamRef Texture,FUnorderedAccessViewRHIParamRef UAV) const;
 
-	void UnsetUAV(FComputeShaderRHIParamRef ComputeShader) const
-	{
-		SetUAVParameter(ComputeShader,UAVParameter,FUnorderedAccessViewRHIRef());
-	}
+	inline void UnsetUAV(FRHICommandList* RHICmdList, FComputeShaderRHIParamRef ComputeShader) const;
 
 private:
 
@@ -514,22 +209,6 @@ private:
 #endif
 };
 
-/** Sets the value of a shader uniform buffer parameter to a uniform buffer containing the struct. */
-template<typename TShaderRHIRef>
-void SetUniformBufferParameter(
-	TShaderRHIRef Shader,
-	const FShaderUniformBufferParameter& Parameter,
-	FUniformBufferRHIParamRef UniformBufferRHI
-	)
-{
-	// This will trigger if the parameter was not serialized
-	checkSlow(Parameter.IsInitialized());
-	if(Parameter.IsBound())
-	{
-		RHISetShaderUniformBuffer(Shader,Parameter.GetBaseIndex(),UniformBufferRHI);
-	}
-}
-
 /** A shader uniform buffer binding with a specific structure. */
 template<typename TBufferStruct>
 class TShaderUniformBufferParameter : public FShaderUniformBufferParameter
@@ -546,55 +225,3 @@ public:
 		return Ar;
 	}
 };
-
-/** Sets the value of a shader uniform buffer parameter to a uniform buffer containing the struct. */
-template<typename TShaderRHIRef,typename TBufferStruct>
-void SetUniformBufferParameter(
-	TShaderRHIRef Shader,
-	const TShaderUniformBufferParameter<TBufferStruct>& Parameter,
-	const TUniformBufferRef<TBufferStruct>& UniformBufferRef
-	)
-{
-	// This will trigger if the parameter was not serialized
-	checkSlow(Parameter.IsInitialized());
-	if(Parameter.IsBound())
-	{
-		RHISetShaderUniformBuffer(Shader,Parameter.GetBaseIndex(),UniformBufferRef);
-	}
-}
-
-/** Sets the value of a shader uniform buffer parameter to a uniform buffer containing the struct. */
-template<typename TShaderRHIRef,typename TBufferStruct>
-void SetUniformBufferParameter(
-	TShaderRHIRef Shader,
-	const TShaderUniformBufferParameter<TBufferStruct>& Parameter,
-	const TUniformBuffer<TBufferStruct>& UniformBuffer
-	)
-{
-	// This will trigger if the parameter was not serialized
-	checkSlow(Parameter.IsInitialized());
-	if(Parameter.IsBound())
-	{
-		RHISetShaderUniformBuffer(Shader,Parameter.GetBaseIndex(),UniformBuffer.GetUniformBufferRHI());
-	}
-}
-
-/** Sets the value of a shader uniform buffer parameter to a value of the struct. */
-template<typename TShaderRHIRef,typename TBufferStruct>
-void SetUniformBufferParameterImmediate(
-	TShaderRHIRef Shader,
-	const TShaderUniformBufferParameter<TBufferStruct>& Parameter,
-	const TBufferStruct& UniformBufferValue
-	)
-{
-	// This will trigger if the parameter was not serialized
-	checkSlow(Parameter.IsInitialized());
-	if(Parameter.IsBound())
-	{
-		RHISetShaderUniformBuffer(
-			Shader,
-			Parameter.GetBaseIndex(),
-			RHICreateUniformBuffer(&UniformBufferValue,sizeof(UniformBufferValue),UniformBuffer_SingleUse)
-			);
-	}
-}
diff --git a/Engine/Source/Runtime/ShaderCore/Public/VertexFactory.h b/Engine/Source/Runtime/ShaderCore/Public/VertexFactory.h
--- a/Engine/Source/Runtime/ShaderCore/Public/VertexFactory.h
+++ b/Engine/Source/Runtime/ShaderCore/Public/VertexFactory.h
@@ -72,7 +72,7 @@ public:
 	virtual ~FVertexFactoryShaderParameters() {}
 	virtual void Bind(const class FShaderParameterMap& ParameterMap) = 0;
 	virtual void Serialize(FArchive& Ar) = 0;
-	virtual void SetMesh(FShader* VertexShader,const class FVertexFactory* VertexFactory,const class FSceneView& View,const struct FMeshBatchElement& BatchElement,uint32 DataFlags) const = 0;
+	virtual void SetMesh(FRHICommandList* RHICmdList, FShader* VertexShader,const class FVertexFactory* VertexFactory,const class FSceneView& View,const struct FMeshBatchElement& BatchElement,uint32 DataFlags) const = 0;
 	virtual uint32 GetSize() const { return sizeof(*this); }
 };
 
@@ -339,7 +339,7 @@ public:
 	/**
 	 * Activates the vertex factory.
 	 */
-	void Set() const;
+	void Set(FRHICommandList* RHICmdList) const;
 
 	/**
 	* Sets the position stream as the current stream source.
@@ -472,11 +472,11 @@ public:
 		delete Parameters;
 	}
 
-	void SetMesh(FShader* Shader,const FVertexFactory* VertexFactory,const class FSceneView& View,const FMeshBatchElement& BatchElement,uint32 DataFlags) const
+	void SetMesh(FRHICommandList* RHICmdList, FShader* Shader,const FVertexFactory* VertexFactory,const class FSceneView& View,const FMeshBatchElement& BatchElement,uint32 DataFlags) const
 	{
 		if(Parameters)
 		{
-			Parameters->SetMesh(Shader,VertexFactory,View,BatchElement,DataFlags);
+			Parameters->SetMesh(RHICmdList, Shader,VertexFactory,View,BatchElement,DataFlags);
 		}
 	}
 
diff --git a/Engine/Source/Runtime/ShaderCore/ShaderCore.Build.cs b/Engine/Source/Runtime/ShaderCore/ShaderCore.Build.cs
--- a/Engine/Source/Runtime/ShaderCore/ShaderCore.Build.cs
+++ b/Engine/Source/Runtime/ShaderCore/ShaderCore.Build.cs
@@ -6,9 +6,9 @@ public class ShaderCore : ModuleRules
 {
 	public ShaderCore(TargetInfo Target)
 	{
-		PublicDependencyModuleNames.Add("RHI");
+		PublicDependencyModuleNames.AddRange(new string[] { "RHI", "RenderCore" });
 
-		PrivateDependencyModuleNames.AddRange(new string[] { "Core", "RenderCore" });
+		PrivateDependencyModuleNames.AddRange(new string[] { "Core" });
 
 		PrivateIncludePathModuleNames.AddRange(new string[] { "DerivedDataCache", "TargetPlatform" });
 	}
diff --git a/Engine/Source/Runtime/SlateRHIRenderer/Private/SlateRHIRenderer.cpp b/Engine/Source/Runtime/SlateRHIRenderer/Private/SlateRHIRenderer.cpp
--- a/Engine/Source/Runtime/SlateRHIRenderer/Private/SlateRHIRenderer.cpp
+++ b/Engine/Source/Runtime/SlateRHIRenderer/Private/SlateRHIRenderer.cpp
@@ -6,7 +6,9 @@
 #include "SlateRHIFontTexture.h"
 #include "SlateRHIResourceManager.h"
 #include "SlateCore.h"
+#include "RHIStaticStates.h"
 #include "SlateRHIRenderingPolicy.h"
+#include "ShaderParameterUtils.h"
 #include "Runtime/Engine/Public/ScreenRendering.h"
 #include "Runtime/Engine/Public/ShaderCompiler.h"
 #include "SlateShaders.h"
@@ -198,7 +200,7 @@ void FSlateRHIRenderer::Initialize()
 		const FIntPoint VirtualScreenOrigin = FIntPoint(DisplayMetrics.VirtualDisplayRect.Left, DisplayMetrics.VirtualDisplayRect.Top);
 		const FIntPoint VirtualScreenLowerRight = FIntPoint(DisplayMetrics.VirtualDisplayRect.Right, DisplayMetrics.VirtualDisplayRect.Bottom);
 		const FIntRect VirtualScreen = FIntRect(VirtualScreenOrigin, VirtualScreenLowerRight);
-
+	
 		CrashTrackerResource = new FSlateCrashReportResource(VirtualScreen);
 		BeginInitResource(CrashTrackerResource);
 	}
diff --git a/Engine/Source/Runtime/SlateRHIRenderer/Private/SlateRHIRenderingPolicy.cpp b/Engine/Source/Runtime/SlateRHIRenderer/Private/SlateRHIRenderingPolicy.cpp
--- a/Engine/Source/Runtime/SlateRHIRenderer/Private/SlateRHIRenderingPolicy.cpp
+++ b/Engine/Source/Runtime/SlateRHIRenderer/Private/SlateRHIRenderingPolicy.cpp
@@ -5,6 +5,7 @@
 #include "SlateRHIRenderingPolicy.h"
 #include "SlateShaders.h"
 #include "SlateRHIResourceManager.h"
+#include "RHIStaticStates.h"
 
 static EPrimitiveType GetRHIPrimitiveType( ESlateDrawPrimitive::Type SlateType )
 {
@@ -287,6 +288,9 @@ void FSlateRHIRenderingPolicy::DrawElements( const FIntPoint& InViewportSize, FS
 	// Should only be called by the rendering thread
 	check(IsInRenderingThread());
 
+	//@todo-rco: RHIPacketList
+	FRHICommandList* RHICmdList = nullptr;
+
 	const float DisplayGamma = GEngine ? GEngine->GetDisplayGamma() : 2.2f;
 
 	FSlateElementVertexBuffer& VertexBuffer = VertexBuffers[CurrentBufferIndex];
@@ -339,8 +343,8 @@ void FSlateRHIRenderingPolicy::DrawElements( const FIntPoint& InViewportSize, FS
 				SetGlobalBoundShaderState( NormalShaderStates[ShaderType][bUseTextureAlpha], GSlateVertexDeclaration.VertexDeclarationRHI, *VertexShader, PixelShader );
 			}
 
-			VertexShader->SetViewProjection( ViewProjectionMatrix );
-			VertexShader->SetShaderParameters( ShaderParams.VertexParams );
+			VertexShader->SetViewProjection(RHICmdList, ViewProjectionMatrix );
+			VertexShader->SetShaderParameters(RHICmdList, ShaderParams.VertexParams );
 
 			RHISetBlendState(
 				(RenderBatch.DrawFlags & ESlateBatchDrawFlag::NoBlending)
@@ -384,7 +388,7 @@ void FSlateRHIRenderingPolicy::DrawElements( const FIntPoint& InViewportSize, FS
 #endif
 
 				// The lookup table used for splines should clamp when sampling otherwise the results are wrong
-				PixelShader->SetTexture( ((FSlateTexture2DRHIRef*)ShaderResource)->GetTypedResource(), BilinearClamp );
+				PixelShader->SetTexture(RHICmdList, ((FSlateTexture2DRHIRef*)ShaderResource)->GetTypedResource(), BilinearClamp );
 			}
 			else if( ShaderResource && ShaderResource->GetType() == ESlateShaderResource::Texture && IsValidRef( ((FSlateTexture2DRHIRef*)ShaderResource)->GetTypedResource() ) )
 			{
@@ -406,17 +410,17 @@ void FSlateRHIRenderingPolicy::DrawElements( const FIntPoint& InViewportSize, FS
 				{
 					SamplerState = BilinearClamp;
 				}
-				PixelShader->SetTexture( ((FSlateTexture2DRHIRef*)ShaderResource)->GetTypedResource(), SamplerState );
+				PixelShader->SetTexture(RHICmdList, ((FSlateTexture2DRHIRef*)ShaderResource)->GetTypedResource(), SamplerState );
 			}
 			else
 			{
-				PixelShader->SetTexture( GWhiteTexture->TextureRHI, BilinearClamp );
+				PixelShader->SetTexture(RHICmdList, GWhiteTexture->TextureRHI, BilinearClamp );
 			}
 
-			PixelShader->SetShaderParams( ShaderParams.PixelParams );
-			PixelShader->SetViewportSize( InViewportSize );
-			PixelShader->SetDrawEffects( DrawEffects );
-			PixelShader->SetDisplayGamma( ( DrawFlags & ESlateBatchDrawFlag::NoGamma ) ? 1.0f : DisplayGamma );
+			PixelShader->SetShaderParams(RHICmdList, ShaderParams.PixelParams );
+			PixelShader->SetViewportSize(RHICmdList, InViewportSize );
+			PixelShader->SetDrawEffects(RHICmdList, DrawEffects );
+			PixelShader->SetDisplayGamma(RHICmdList, ( DrawFlags & ESlateBatchDrawFlag::NoGamma ) ? 1.0f : DisplayGamma );
 
 			check( RenderBatch.NumIndices > 0 );
 
diff --git a/Engine/Source/Runtime/SlateRHIRenderer/Private/SlateShaders.cpp b/Engine/Source/Runtime/SlateRHIRenderer/Private/SlateShaders.cpp
--- a/Engine/Source/Runtime/SlateRHIRenderer/Private/SlateShaders.cpp
+++ b/Engine/Source/Runtime/SlateRHIRenderer/Private/SlateShaders.cpp
@@ -4,6 +4,7 @@
 #include "SlateRHIRendererPrivatePCH.h"
 #include "SlateRHIRenderingPolicy.h"
 #include "SlateShaders.h"
+#include "ShaderParameterUtils.h"
 
 /** Flag to determine if we are running with a color vision deficiency shader on */
 uint32 GSlateShaderColorVisionDeficiencyType = 0;
@@ -66,14 +67,14 @@ FSlateElementVS::FSlateElementVS( const ShaderMetaType::CompiledShaderInitialize
 	VertexShaderParams.Bind( Initializer.ParameterMap, TEXT("VertexShaderParams"));
 }
 
-void FSlateElementVS::SetViewProjection( const FMatrix& InViewProjection )
+void FSlateElementVS::SetViewProjection(FRHICommandList* RHICmdList, const FMatrix& InViewProjection )
 {
-	SetShaderValue( GetVertexShader(), ViewProjection, InViewProjection );
+	SetShaderValue(RHICmdList, GetVertexShader(), ViewProjection, InViewProjection );
 }
 
-void FSlateElementVS::SetShaderParameters( const FVector4& ShaderParams )
+void FSlateElementVS::SetShaderParameters(FRHICommandList* RHICmdList, const FVector4& ShaderParams )
 {
-	SetShaderValue( GetVertexShader(), VertexShaderParams, ShaderParams );
+	SetShaderValue(RHICmdList, GetVertexShader(), VertexShaderParams, ShaderParams );
 }
 
 /** Serializes the shader data */
diff --git a/Engine/Source/Runtime/SlateRHIRenderer/Private/SlateShaders.h b/Engine/Source/Runtime/SlateRHIRenderer/Private/SlateShaders.h
--- a/Engine/Source/Runtime/SlateRHIRenderer/Private/SlateShaders.h
+++ b/Engine/Source/Runtime/SlateRHIRenderer/Private/SlateShaders.h
@@ -45,14 +45,14 @@ public:
 	 *
 	 * @param InViewProjection	The ViewProjection matrix to use when this shader is bound 
 	 */
-	void SetViewProjection( const FMatrix& InViewProjection );
+	void SetViewProjection(FRHICommandList* RHICmdList, const FMatrix& InViewProjection );
 
 	/** 
 	 * Sets shader parameters for use in this shader
 	 *
 	 * @param ShaderParams	The shader params to be used
 	 */
-	void SetShaderParameters( const FVector4& ShaderParams );
+	void SetShaderParameters(FRHICommandList* RHICmdList, const FVector4& ShaderParams );
 
 	/** Serializes the shader data */
 	virtual bool Serialize( FArchive& Ar );
@@ -99,9 +99,9 @@ public:
 	 * @param Texture	Texture resource to use when this pixel shader is bound
 	 * @param SamplerState	Sampler state to use when sampling this texture
 	 */
-	void SetTexture( const FTextureRHIParamRef InTexture, const FSamplerStateRHIRef SamplerState )
+	void SetTexture(FRHICommandList* RHICmdList, const FTextureRHIParamRef InTexture, const FSamplerStateRHIRef SamplerState )
 	{
-		SetTextureParameter( GetPixelShader(), TextureParameter, TextureParameterSampler, SamplerState, InTexture );
+		SetTextureParameter(RHICmdList, GetPixelShader(), TextureParameter, TextureParameterSampler, SamplerState, InTexture );
 	}
 
 	/**
@@ -109,9 +109,9 @@ public:
 	 *
 	 * @param InDrawEffects		A mask of draw effects that the pixel shader should use.
 	 */
-	void SetDrawEffects( ESlateDrawEffect::Type InDrawEffects )
+	void SetDrawEffects(FRHICommandList* RHICmdList, ESlateDrawEffect::Type InDrawEffects )
 	{
-		SetShaderValue( GetPixelShader(), DrawEffectsParameter, InDrawEffects );
+		SetShaderValue(RHICmdList, GetPixelShader(), DrawEffectsParameter, InDrawEffects );
 	}
 
 	/**
@@ -119,17 +119,17 @@ public:
 	 * 
 	 * @param InShaderParams Shader params to use
 	 */
-	void SetShaderParams( const FVector4& InShaderParams )
+	void SetShaderParams(FRHICommandList* RHICmdList, const FVector4& InShaderParams )
 	{
-		SetShaderValue( GetPixelShader(), ShaderParams, InShaderParams );
+		SetShaderValue(RHICmdList, GetPixelShader(), ShaderParams, InShaderParams );
 	}
 
 	/**
 	 * Set the viewport size.
 	 */
-	void SetViewportSize(const FVector2D& InViewportSize)
+	void SetViewportSize(FRHICommandList* RHICmdList, const FVector2D& InViewportSize)
 	{
-		SetShaderValue(GetPixelShader(), ViewportSize, InViewportSize);
+		SetShaderValue(RHICmdList, GetPixelShader(), ViewportSize, InViewportSize);
 	}
 
 	/**
@@ -137,9 +137,9 @@ public:
  	 *
 	 * @param DisplayGamma The display gamma to use
 	 */
-	void SetDisplayGamma(float InDisplayGamma)
+	void SetDisplayGamma(FRHICommandList* RHICmdList, float InDisplayGamma)
 	{
-		SetShaderValue(GetPixelShader(),DisplayGamma,InDisplayGamma);
+		SetShaderValue(RHICmdList, GetPixelShader(),DisplayGamma,InDisplayGamma);
 	}
 
 	virtual bool Serialize( FArchive& Ar )
diff --git a/Engine/Source/Runtime/Windows/D3D11RHI/D3D11RHI.Build.cs b/Engine/Source/Runtime/Windows/D3D11RHI/D3D11RHI.Build.cs
--- a/Engine/Source/Runtime/Windows/D3D11RHI/D3D11RHI.Build.cs
+++ b/Engine/Source/Runtime/Windows/D3D11RHI/D3D11RHI.Build.cs
@@ -11,7 +11,6 @@ public class D3D11RHI : ModuleRules
 		PrivateDependencyModuleNames.AddRange(
 			new string[] {
 				"Core",
-				"CoreUObject",
 				"Engine",
 				"RHI",
 				"RenderCore",
diff --git a/Engine/Source/Runtime/Windows/D3D11RHI/Private/D3D11Commands.cpp b/Engine/Source/Runtime/Windows/D3D11RHI/Private/D3D11Commands.cpp
--- a/Engine/Source/Runtime/Windows/D3D11RHI/Private/D3D11Commands.cpp
+++ b/Engine/Source/Runtime/Windows/D3D11RHI/Private/D3D11Commands.cpp
@@ -14,6 +14,9 @@
 #include "StaticBoundShaderState.h"
 #include "GlobalShader.h"
 #include "OneColorShader.h"
+#include "RHICommandList.h"
+#include "RHIStaticStates.h"
+#include "ShaderParameterUtils.h"
 
 #if PLATFORM_DESKTOP
 // For Depth Bounds Test interface
@@ -73,6 +76,8 @@ void FD3D11DynamicRHI::RHIGpuTimeEnd(uint32 Hash, bool bCompute)
 // Vertex state.
 void FD3D11DynamicRHI::RHISetStreamSource(uint32 StreamIndex,FVertexBufferRHIParamRef VertexBufferRHI,uint32 Stride,uint32 Offset)
 {
+	GRHICommandList.Verify();
+
 	DYNAMIC_CAST_D3D11RESOURCE(VertexBuffer,VertexBuffer);
 
 	ID3D11Buffer* D3DBuffer = VertexBuffer ? VertexBuffer->Resource : NULL;
@@ -81,6 +86,8 @@ void FD3D11DynamicRHI::RHISetStreamSource(uint32 StreamIndex,FVertexBufferRHIPar
 
 void FD3D11DynamicRHI::RHISetStreamOutTargets(uint32 NumTargets, const FVertexBufferRHIParamRef* VertexBuffers, const uint32* Offsets)
 {
+	GRHICommandList.Verify();
+
 	ID3D11Buffer* D3DVertexBuffers[D3D11_SO_BUFFER_SLOT_COUNT] = {0};
 
 	if (VertexBuffers)
@@ -97,18 +104,24 @@ void FD3D11DynamicRHI::RHISetStreamOutTargets(uint32 NumTargets, const FVertexBu
 // Rasterizer state.
 void FD3D11DynamicRHI::RHISetRasterizerState(FRasterizerStateRHIParamRef NewStateRHI)
 {
+	GRHICommandList.Verify();
+
 	DYNAMIC_CAST_D3D11RESOURCE(RasterizerState,NewState);
 	StateCache.SetRasterizerState(NewState->Resource);
 }
 
 void FD3D11DynamicRHI::RHISetComputeShader(FComputeShaderRHIParamRef ComputeShaderRHI)
 {
+	GRHICommandList.Verify();
+
 	DYNAMIC_CAST_D3D11RESOURCE(ComputeShader,ComputeShader);
 	SetCurrentComputeShader(ComputeShaderRHI);
 }
 
 void FD3D11DynamicRHI::RHIDispatchComputeShader(uint32 ThreadGroupCountX, uint32 ThreadGroupCountY, uint32 ThreadGroupCountZ) 
 { 
+	GRHICommandList.Verify();
+
 	FComputeShaderRHIParamRef ComputeShaderRHI = GetCurrentComputeShader();
 	DYNAMIC_CAST_D3D11RESOURCE(ComputeShader,ComputeShader);
 
@@ -120,6 +133,7 @@ void FD3D11DynamicRHI::RHIDispatchComputeShader(uint32 ThreadGroupCountX, uint32
 	{
 		CommitComputeShaderConstants();
 	}
+	CommitComputeResourceTables(ComputeShader);
 	
 	Direct3DDeviceIMContext->Dispatch(ThreadGroupCountX, ThreadGroupCountY, ThreadGroupCountZ);	
 	StateCache.SetComputeShader(nullptr);
@@ -127,6 +141,8 @@ void FD3D11DynamicRHI::RHIDispatchComputeShader(uint32 ThreadGroupCountX, uint32
 
 void FD3D11DynamicRHI::RHIDispatchIndirectComputeShader(FVertexBufferRHIParamRef ArgumentBufferRHI, uint32 ArgumentOffset) 
 { 
+	GRHICommandList.Verify();
+
 	FComputeShaderRHIParamRef ComputeShaderRHI = GetCurrentComputeShader();
 	DYNAMIC_CAST_D3D11RESOURCE(ComputeShader,ComputeShader);
 	DYNAMIC_CAST_D3D11RESOURCE(VertexBuffer,ArgumentBuffer);
@@ -139,6 +155,7 @@ void FD3D11DynamicRHI::RHIDispatchIndirectComputeShader(FVertexBufferRHIParamRef
 	{
 		CommitComputeShaderConstants();
 	}
+	CommitComputeResourceTables(ComputeShader);
 
 	Direct3DDeviceIMContext->DispatchIndirect(ArgumentBuffer->Resource,ArgumentOffset);
 	StateCache.SetComputeShader(nullptr);
@@ -146,6 +163,8 @@ void FD3D11DynamicRHI::RHIDispatchIndirectComputeShader(FVertexBufferRHIParamRef
 
 void FD3D11DynamicRHI::RHISetViewport(uint32 MinX,uint32 MinY,float MinZ,uint32 MaxX,uint32 MaxY,float MaxZ)
 {
+	GRHICommandList.Verify();
+
 	// These are the maximum viewport extents for D3D11. Exceeding them leads to badness.
 	check(MinX <= (uint32)D3D11_VIEWPORT_BOUNDS_MAX);
 	check(MinY <= (uint32)D3D11_VIEWPORT_BOUNDS_MAX);
@@ -163,6 +182,8 @@ void FD3D11DynamicRHI::RHISetViewport(uint32 MinX,uint32 MinY,float MinZ,uint32
 
 void FD3D11DynamicRHI::RHISetScissorRect(bool bEnable,uint32 MinX,uint32 MinY,uint32 MaxX,uint32 MaxY)
 {
+	GRHICommandList.Verify();
+
 	if(bEnable)
 	{
 		D3D11_RECT ScissorRect;
@@ -189,6 +210,8 @@ void FD3D11DynamicRHI::RHISetScissorRect(bool bEnable,uint32 MinX,uint32 MinY,ui
 */
 void FD3D11DynamicRHI::RHISetBoundShaderState( FBoundShaderStateRHIParamRef BoundShaderStateRHI)
 {
+	GRHICommandList.Verify();
+
 	DYNAMIC_CAST_D3D11RESOURCE(BoundShaderState,BoundShaderState);
 
 	StateCache.SetInputLayout(BoundShaderState->InputLayout);
@@ -214,23 +237,41 @@ void FD3D11DynamicRHI::RHISetBoundShaderState( FBoundShaderStateRHIParamRef Boun
 	// Prevent transient bound shader states from being recreated for each use by keeping a history of the most recently used bound shader states.
 	// The history keeps them alive, and the bound shader state cache allows them to be reused if needed.
 	BoundShaderStateHistory.Add(BoundShaderState);
+
+	// Shader changed so all resource tables are dirty
+	DirtyUniformBuffers[SF_Vertex] = 0xffff;
+	DirtyUniformBuffers[SF_Pixel] = 0xffff;
+	DirtyUniformBuffers[SF_Hull] = 0xffff;
+	DirtyUniformBuffers[SF_Domain] = 0xffff;
+	DirtyUniformBuffers[SF_Geometry] = 0xffff;
 }
 
 void FD3D11DynamicRHI::RHISetShaderTexture(FVertexShaderRHIParamRef VertexShaderRHI,uint32 TextureIndex,FTextureRHIParamRef NewTextureRHI)
 {
+	uint32 Start = FPlatformTime::Cycles();
+
+	GRHICommandList.Verify();
+
 	VALIDATE_BOUND_SHADER(VertexShaderRHI);
 
 	FD3D11TextureBase* NewTexture = GetD3D11TextureFromRHITexture(NewTextureRHI);
 	ID3D11ShaderResourceView* ShaderResourceView = NewTexture ? NewTexture->GetShaderResourceView() : NULL;
-	
+
 	if (  ( NewTexture == NULL) || ( NewTexture->GetRenderTargetView( 0, 0 ) !=NULL) || ( NewTexture->HasDepthStencilView()) )
 		SetShaderResourceView<SF_Vertex>(NewTexture, ShaderResourceView, TextureIndex, FD3D11StateCache::SRV_Dynamic);
 	else
 		SetShaderResourceView<SF_Vertex>(NewTexture, ShaderResourceView, TextureIndex, FD3D11StateCache::SRV_Static);
+
+	SetShaderTextureCycles += (FPlatformTime::Cycles() - Start);
+	SetShaderTextureCalls++;
 }
 
 void FD3D11DynamicRHI::RHISetShaderTexture(FHullShaderRHIParamRef HullShaderRHI,uint32 TextureIndex,FTextureRHIParamRef NewTextureRHI)
 {
+	uint32 Start = FPlatformTime::Cycles();
+
+	GRHICommandList.Verify();
+
 	VALIDATE_BOUND_SHADER(HullShaderRHI);
 
 	FD3D11TextureBase* NewTexture = GetD3D11TextureFromRHITexture(NewTextureRHI);
@@ -241,10 +282,16 @@ void FD3D11DynamicRHI::RHISetShaderTexture(FHullShaderRHIParamRef HullShaderRHI,
 	else
 		SetShaderResourceView<SF_Hull>(NewTexture, ShaderResourceView, TextureIndex, FD3D11StateCache::SRV_Static);
 
+	SetShaderTextureCycles += (FPlatformTime::Cycles() - Start);
+	SetShaderTextureCalls++;
 }
 
 void FD3D11DynamicRHI::RHISetShaderTexture(FDomainShaderRHIParamRef DomainShaderRHI,uint32 TextureIndex,FTextureRHIParamRef NewTextureRHI)
 {
+	uint32 Start = FPlatformTime::Cycles();
+
+	GRHICommandList.Verify();
+
 	VALIDATE_BOUND_SHADER(DomainShaderRHI);
 
 	FD3D11TextureBase* NewTexture = GetD3D11TextureFromRHITexture(NewTextureRHI);
@@ -254,10 +301,17 @@ void FD3D11DynamicRHI::RHISetShaderTexture(FDomainShaderRHIParamRef DomainShader
 		SetShaderResourceView<SF_Domain>(NewTexture, ShaderResourceView, TextureIndex, FD3D11StateCache::SRV_Dynamic);
 	else
 		SetShaderResourceView<SF_Domain>(NewTexture, ShaderResourceView, TextureIndex, FD3D11StateCache::SRV_Static);
+
+	SetShaderTextureCycles += (FPlatformTime::Cycles() - Start);
+	SetShaderTextureCalls++;
 }
 
 void FD3D11DynamicRHI::RHISetShaderTexture(FGeometryShaderRHIParamRef GeometryShaderRHI,uint32 TextureIndex,FTextureRHIParamRef NewTextureRHI)
 {
+	uint32 Start = FPlatformTime::Cycles();
+
+	GRHICommandList.Verify();
+
 	VALIDATE_BOUND_SHADER(GeometryShaderRHI);
 
 	FD3D11TextureBase* NewTexture = GetD3D11TextureFromRHITexture(NewTextureRHI);
@@ -267,10 +321,17 @@ void FD3D11DynamicRHI::RHISetShaderTexture(FGeometryShaderRHIParamRef GeometrySh
 		SetShaderResourceView<SF_Geometry>(NewTexture, ShaderResourceView, TextureIndex, FD3D11StateCache::SRV_Dynamic);
 	else
 		SetShaderResourceView<SF_Geometry>(NewTexture, ShaderResourceView, TextureIndex, FD3D11StateCache::SRV_Static);
+
+	SetShaderTextureCycles += (FPlatformTime::Cycles() - Start);
+	SetShaderTextureCalls++;
 }
 
 void FD3D11DynamicRHI::RHISetShaderTexture(FPixelShaderRHIParamRef PixelShaderRHI,uint32 TextureIndex,FTextureRHIParamRef NewTextureRHI)
 {
+	uint32 Start = FPlatformTime::Cycles();
+
+	GRHICommandList.Verify();
+
 	VALIDATE_BOUND_SHADER(PixelShaderRHI);
 
 	FD3D11TextureBase* NewTexture = GetD3D11TextureFromRHITexture(NewTextureRHI);
@@ -280,10 +341,15 @@ void FD3D11DynamicRHI::RHISetShaderTexture(FPixelShaderRHIParamRef PixelShaderRH
 	else
 		SetShaderResourceView<SF_Pixel>(NewTexture, ShaderResourceView, TextureIndex, FD3D11StateCache::SRV_Static);
 
+	SetShaderTextureCycles += (FPlatformTime::Cycles() - Start);
+	SetShaderTextureCalls++;
 }
 
 void FD3D11DynamicRHI::RHISetShaderTexture(FComputeShaderRHIParamRef ComputeShaderRHI,uint32 TextureIndex,FTextureRHIParamRef NewTextureRHI)
 {
+	uint32 Start = FPlatformTime::Cycles();
+
+	GRHICommandList.Verify();
 	//VALIDATE_BOUND_SHADER(ComputeShaderRHI);
 
 	FD3D11TextureBase* NewTexture = GetD3D11TextureFromRHITexture(NewTextureRHI);
@@ -293,14 +359,18 @@ void FD3D11DynamicRHI::RHISetShaderTexture(FComputeShaderRHIParamRef ComputeShad
 		SetShaderResourceView<SF_Compute>(NewTexture, ShaderResourceView, TextureIndex, FD3D11StateCache::SRV_Dynamic);
 	else
 		SetShaderResourceView<SF_Compute>(NewTexture, ShaderResourceView, TextureIndex, FD3D11StateCache::SRV_Static);
+
+	SetShaderTextureCycles += (FPlatformTime::Cycles() - Start);
+	SetShaderTextureCalls++;
 }
 
 void FD3D11DynamicRHI::RHISetUAVParameter(FComputeShaderRHIParamRef ComputeShaderRHI,uint32 UAVIndex,FUnorderedAccessViewRHIParamRef UAVRHI)
 {
+	GRHICommandList.Verify();
 	//VALIDATE_BOUND_SHADER(ComputeShaderRHI);
 
 	DYNAMIC_CAST_D3D11RESOURCE(UnorderedAccessView,UAV);
-	
+
 	if(UAV)
 	{
 		ConditionalClearShaderResource(UAV->Resource);
@@ -314,10 +384,11 @@ void FD3D11DynamicRHI::RHISetUAVParameter(FComputeShaderRHIParamRef ComputeShade
 
 void FD3D11DynamicRHI::RHISetUAVParameter(FComputeShaderRHIParamRef ComputeShaderRHI,uint32 UAVIndex,FUnorderedAccessViewRHIParamRef UAVRHI, uint32 InitialCount )
 {
+	GRHICommandList.Verify();
 	//VALIDATE_BOUND_SHADER(ComputeShaderRHI);
 
 	DYNAMIC_CAST_D3D11RESOURCE(UnorderedAccessView,UAV);
-
+	
 	if(UAV)
 	{
 		ConditionalClearShaderResource(UAV->Resource);
@@ -329,6 +400,8 @@ void FD3D11DynamicRHI::RHISetUAVParameter(FComputeShaderRHIParamRef ComputeShade
 
 void FD3D11DynamicRHI::RHISetShaderResourceViewParameter(FPixelShaderRHIParamRef PixelShaderRHI,uint32 TextureIndex,FShaderResourceViewRHIParamRef SRVRHI)
 {
+	GRHICommandList.Verify();
+
 	VALIDATE_BOUND_SHADER(PixelShaderRHI);
 
 	DYNAMIC_CAST_D3D11RESOURCE(ShaderResourceView,SRV);
@@ -347,13 +420,15 @@ void FD3D11DynamicRHI::RHISetShaderResourceViewParameter(FPixelShaderRHIParamRef
 
 void FD3D11DynamicRHI::RHISetShaderResourceViewParameter(FVertexShaderRHIParamRef VertexShaderRHI,uint32 TextureIndex,FShaderResourceViewRHIParamRef SRVRHI)
 {
+	GRHICommandList.Verify();
+
 	VALIDATE_BOUND_SHADER(VertexShaderRHI);
 
 	DYNAMIC_CAST_D3D11RESOURCE(ShaderResourceView,SRV);
 
 	FD3D11BaseShaderResource* Resource = nullptr;
 	ID3D11ShaderResourceView* D3D11SRV = nullptr;
-
+	
 	if (SRV)
 	{
 		Resource = SRV->Resource;
@@ -365,13 +440,15 @@ void FD3D11DynamicRHI::RHISetShaderResourceViewParameter(FVertexShaderRHIParamRe
 
 void FD3D11DynamicRHI::RHISetShaderResourceViewParameter(FComputeShaderRHIParamRef ComputeShaderRHI,uint32 TextureIndex,FShaderResourceViewRHIParamRef SRVRHI)
 {
+	GRHICommandList.Verify();
+
 	//VALIDATE_BOUND_SHADER(ComputeShaderRHI);
 
 	DYNAMIC_CAST_D3D11RESOURCE(ShaderResourceView,SRV);
 
 	FD3D11BaseShaderResource* Resource = nullptr;
 	ID3D11ShaderResourceView* D3D11SRV = nullptr;
-
+	
 	if (SRV)
 	{
 		Resource = SRV->Resource;
@@ -383,13 +460,15 @@ void FD3D11DynamicRHI::RHISetShaderResourceViewParameter(FComputeShaderRHIParamR
 
 void FD3D11DynamicRHI::RHISetShaderResourceViewParameter(FHullShaderRHIParamRef HullShaderRHI,uint32 TextureIndex,FShaderResourceViewRHIParamRef SRVRHI)
 {
+	GRHICommandList.Verify();
+
 	VALIDATE_BOUND_SHADER(HullShaderRHI);
 
 	DYNAMIC_CAST_D3D11RESOURCE(ShaderResourceView,SRV);
 
 	FD3D11BaseShaderResource* Resource = nullptr;
 	ID3D11ShaderResourceView* D3D11SRV = nullptr;
-
+	
 	if (SRV)
 	{
 		Resource = SRV->Resource;
@@ -401,13 +480,15 @@ void FD3D11DynamicRHI::RHISetShaderResourceViewParameter(FHullShaderRHIParamRef
 
 void FD3D11DynamicRHI::RHISetShaderResourceViewParameter(FDomainShaderRHIParamRef DomainShaderRHI,uint32 TextureIndex,FShaderResourceViewRHIParamRef SRVRHI)
 {
+	GRHICommandList.Verify();
+
 	VALIDATE_BOUND_SHADER(DomainShaderRHI);
 
 	DYNAMIC_CAST_D3D11RESOURCE(ShaderResourceView,SRV);
 
 	FD3D11BaseShaderResource* Resource = nullptr;
 	ID3D11ShaderResourceView* D3D11SRV = nullptr;
-
+	
 	if (SRV)
 	{
 		Resource = SRV->Resource;
@@ -419,13 +500,15 @@ void FD3D11DynamicRHI::RHISetShaderResourceViewParameter(FDomainShaderRHIParamRe
 
 void FD3D11DynamicRHI::RHISetShaderResourceViewParameter(FGeometryShaderRHIParamRef GeometryShaderRHI,uint32 TextureIndex,FShaderResourceViewRHIParamRef SRVRHI)
 {
+	GRHICommandList.Verify();
+
 	VALIDATE_BOUND_SHADER(GeometryShaderRHI);
 
 	DYNAMIC_CAST_D3D11RESOURCE(ShaderResourceView,SRV);
 
 	FD3D11BaseShaderResource* Resource = nullptr;
 	ID3D11ShaderResourceView* D3D11SRV = nullptr;
-
+	
 	if (SRV)
 	{
 		Resource = SRV->Resource;
@@ -437,6 +520,8 @@ void FD3D11DynamicRHI::RHISetShaderResourceViewParameter(FGeometryShaderRHIParam
 
 void FD3D11DynamicRHI::RHISetShaderSampler(FVertexShaderRHIParamRef VertexShaderRHI,uint32 SamplerIndex,FSamplerStateRHIParamRef NewStateRHI)
 {
+	GRHICommandList.Verify();
+
 	VALIDATE_BOUND_SHADER(VertexShaderRHI);
 
 	DYNAMIC_CAST_D3D11RESOURCE(VertexShader,VertexShader);
@@ -448,6 +533,8 @@ void FD3D11DynamicRHI::RHISetShaderSampler(FVertexShaderRHIParamRef VertexShader
 
 void FD3D11DynamicRHI::RHISetShaderSampler(FHullShaderRHIParamRef HullShaderRHI,uint32 SamplerIndex,FSamplerStateRHIParamRef NewStateRHI)
 {
+	GRHICommandList.Verify();
+
 	VALIDATE_BOUND_SHADER(HullShaderRHI);
 
 	DYNAMIC_CAST_D3D11RESOURCE(HullShader,HullShader);
@@ -459,6 +546,8 @@ void FD3D11DynamicRHI::RHISetShaderSampler(FHullShaderRHIParamRef HullShaderRHI,
 
 void FD3D11DynamicRHI::RHISetShaderSampler(FDomainShaderRHIParamRef DomainShaderRHI,uint32 SamplerIndex,FSamplerStateRHIParamRef NewStateRHI)
 {
+	GRHICommandList.Verify();
+
 	VALIDATE_BOUND_SHADER(DomainShaderRHI);
 
 	DYNAMIC_CAST_D3D11RESOURCE(DomainShader,DomainShader);
@@ -470,6 +559,8 @@ void FD3D11DynamicRHI::RHISetShaderSampler(FDomainShaderRHIParamRef DomainShader
 
 void FD3D11DynamicRHI::RHISetShaderSampler(FGeometryShaderRHIParamRef GeometryShaderRHI,uint32 SamplerIndex,FSamplerStateRHIParamRef NewStateRHI)
 {
+	GRHICommandList.Verify();
+
 	VALIDATE_BOUND_SHADER(GeometryShaderRHI);
 
 	DYNAMIC_CAST_D3D11RESOURCE(GeometryShader,GeometryShader);
@@ -481,6 +572,8 @@ void FD3D11DynamicRHI::RHISetShaderSampler(FGeometryShaderRHIParamRef GeometrySh
 
 void FD3D11DynamicRHI::RHISetShaderSampler(FPixelShaderRHIParamRef PixelShaderRHI,uint32 SamplerIndex,FSamplerStateRHIParamRef NewStateRHI)
 {
+	GRHICommandList.Verify();
+
 	VALIDATE_BOUND_SHADER(PixelShaderRHI);
 
 	DYNAMIC_CAST_D3D11RESOURCE(PixelShader,PixelShader);
@@ -492,6 +585,8 @@ void FD3D11DynamicRHI::RHISetShaderSampler(FPixelShaderRHIParamRef PixelShaderRH
 
 void FD3D11DynamicRHI::RHISetShaderSampler(FComputeShaderRHIParamRef ComputeShaderRHI,uint32 SamplerIndex,FSamplerStateRHIParamRef NewStateRHI)
 {
+	GRHICommandList.Verify();
+
 	//VALIDATE_BOUND_SHADER(ComputeShaderRHI);
 	DYNAMIC_CAST_D3D11RESOURCE(ComputeShader,ComputeShader);
 	DYNAMIC_CAST_D3D11RESOURCE(SamplerState,NewState);
@@ -502,6 +597,8 @@ void FD3D11DynamicRHI::RHISetShaderSampler(FComputeShaderRHIParamRef ComputeShad
 
 void FD3D11DynamicRHI::RHISetShaderUniformBuffer(FVertexShaderRHIParamRef VertexShader,uint32 BufferIndex,FUniformBufferRHIParamRef BufferRHI)
 {
+	GRHICommandList.Verify();
+
 	VALIDATE_BOUND_SHADER(VertexShader);
 	DYNAMIC_CAST_D3D11RESOURCE(UniformBuffer,Buffer);
 #if PLATFORM_XBOXONE
@@ -515,10 +612,15 @@ void FD3D11DynamicRHI::RHISetShaderUniformBuffer(FVertexShaderRHIParamRef Vertex
 		ID3D11Buffer* ConstantBuffer = Buffer ? Buffer->Resource : NULL;
 		StateCache.SetConstantBuffer<SF_Vertex>(ConstantBuffer, BufferIndex);
 	}
+
+	BoundUniformBuffers[SF_Vertex][BufferIndex] = BufferRHI;
+	DirtyUniformBuffers[SF_Vertex] |= (1 << BufferIndex);
 }
 
 void FD3D11DynamicRHI::RHISetShaderUniformBuffer(FHullShaderRHIParamRef HullShader,uint32 BufferIndex,FUniformBufferRHIParamRef BufferRHI)
 {
+	GRHICommandList.Verify();
+
 	VALIDATE_BOUND_SHADER(HullShader);
 	DYNAMIC_CAST_D3D11RESOURCE(UniformBuffer,Buffer);
 #if PLATFORM_XBOXONE
@@ -532,10 +634,15 @@ void FD3D11DynamicRHI::RHISetShaderUniformBuffer(FHullShaderRHIParamRef HullShad
 		ID3D11Buffer* ConstantBuffer = Buffer ? Buffer->Resource : NULL;
 		StateCache.SetConstantBuffer<SF_Hull>(ConstantBuffer, BufferIndex);
 	}
+
+	BoundUniformBuffers[SF_Hull][BufferIndex] = BufferRHI;
+	DirtyUniformBuffers[SF_Hull] |= (1 << BufferIndex);
 }
 
 void FD3D11DynamicRHI::RHISetShaderUniformBuffer(FDomainShaderRHIParamRef DomainShader,uint32 BufferIndex,FUniformBufferRHIParamRef BufferRHI)
 {
+	GRHICommandList.Verify();
+
 	VALIDATE_BOUND_SHADER(DomainShader);
 	DYNAMIC_CAST_D3D11RESOURCE(UniformBuffer,Buffer);
 #if PLATFORM_XBOXONE
@@ -549,10 +656,15 @@ void FD3D11DynamicRHI::RHISetShaderUniformBuffer(FDomainShaderRHIParamRef Domain
 		ID3D11Buffer* ConstantBuffer = Buffer ? Buffer->Resource : NULL;
 		StateCache.SetConstantBuffer<SF_Domain>(ConstantBuffer, BufferIndex);
 	}
+
+	BoundUniformBuffers[SF_Domain][BufferIndex] = BufferRHI;
+	DirtyUniformBuffers[SF_Domain] |= (1 << BufferIndex);
 }
 
 void FD3D11DynamicRHI::RHISetShaderUniformBuffer(FGeometryShaderRHIParamRef GeometryShader,uint32 BufferIndex,FUniformBufferRHIParamRef BufferRHI)
 {
+	GRHICommandList.Verify();
+
 	VALIDATE_BOUND_SHADER(GeometryShader);
 	DYNAMIC_CAST_D3D11RESOURCE(UniformBuffer,Buffer);
 #if PLATFORM_XBOXONE
@@ -566,27 +678,28 @@ void FD3D11DynamicRHI::RHISetShaderUniformBuffer(FGeometryShaderRHIParamRef Geom
 		ID3D11Buffer* ConstantBuffer = Buffer ? Buffer->Resource : NULL;
 		StateCache.SetConstantBuffer<SF_Geometry>(ConstantBuffer, BufferIndex);
 	}
+
+	BoundUniformBuffers[SF_Geometry][BufferIndex] = BufferRHI;
+	DirtyUniformBuffers[SF_Geometry] |= (1 << BufferIndex);
 }
 
 void FD3D11DynamicRHI::RHISetShaderUniformBuffer(FPixelShaderRHIParamRef PixelShader,uint32 BufferIndex,FUniformBufferRHIParamRef BufferRHI)
 {
+	GRHICommandList.Verify();
+
 	VALIDATE_BOUND_SHADER(PixelShader);
 	DYNAMIC_CAST_D3D11RESOURCE(UniformBuffer,Buffer);
-#if PLATFORM_XBOXONE
-	if (Buffer && Buffer->RingAllocation.IsValid())
-	{
-		StateCache.SetDynamicConstantBuffer<SF_Pixel>(BufferIndex, Buffer->RingAllocation);
-	}
-	else
-#endif
-	{
-		ID3D11Buffer* ConstantBuffer = Buffer ? Buffer->Resource : NULL;
-		StateCache.SetConstantBuffer<SF_Pixel>(ConstantBuffer, BufferIndex);
-	}
+	ID3D11Buffer* ConstantBuffer = Buffer ? Buffer->Resource : NULL;
+	StateCache.SetConstantBuffer<SF_Pixel>(ConstantBuffer, BufferIndex);
+
+	BoundUniformBuffers[SF_Pixel][BufferIndex] = BufferRHI;
+	DirtyUniformBuffers[SF_Pixel] |= (1 << BufferIndex);
 }
 
 void FD3D11DynamicRHI::RHISetShaderUniformBuffer(FComputeShaderRHIParamRef ComputeShader,uint32 BufferIndex,FUniformBufferRHIParamRef BufferRHI)
 {
+	GRHICommandList.Verify();
+
 	//VALIDATE_BOUND_SHADER(ComputeShader);
 	DYNAMIC_CAST_D3D11RESOURCE(UniformBuffer,Buffer);
 #if PLATFORM_XBOXONE
@@ -600,10 +713,15 @@ void FD3D11DynamicRHI::RHISetShaderUniformBuffer(FComputeShaderRHIParamRef Compu
 		ID3D11Buffer* ConstantBuffer = Buffer ? Buffer->Resource : NULL;
 		StateCache.SetConstantBuffer<SF_Compute>(ConstantBuffer, BufferIndex);
 	}
+
+	BoundUniformBuffers[SF_Compute][BufferIndex] = BufferRHI;
+	DirtyUniformBuffers[SF_Compute] |= (1 << BufferIndex);
 }
 
 void FD3D11DynamicRHI::RHISetShaderParameter(FHullShaderRHIParamRef HullShaderRHI,uint32 BufferIndex,uint32 BaseIndex,uint32 NumBytes,const void* NewValue)
 {
+	GRHICommandList.Verify();
+
 	VALIDATE_BOUND_SHADER(HullShaderRHI);
 	checkSlow(HSConstantBuffers[BufferIndex]);
 	HSConstantBuffers[BufferIndex]->UpdateConstant((const uint8*)NewValue,BaseIndex,NumBytes);
@@ -611,6 +729,8 @@ void FD3D11DynamicRHI::RHISetShaderParameter(FHullShaderRHIParamRef HullShaderRH
 
 void FD3D11DynamicRHI::RHISetShaderParameter(FDomainShaderRHIParamRef DomainShaderRHI,uint32 BufferIndex,uint32 BaseIndex,uint32 NumBytes,const void* NewValue)
 {
+	GRHICommandList.Verify();
+
 	VALIDATE_BOUND_SHADER(DomainShaderRHI);
 	checkSlow(DSConstantBuffers[BufferIndex]);
 	DSConstantBuffers[BufferIndex]->UpdateConstant((const uint8*)NewValue,BaseIndex,NumBytes);
@@ -618,6 +738,8 @@ void FD3D11DynamicRHI::RHISetShaderParameter(FDomainShaderRHIParamRef DomainShad
 
 void FD3D11DynamicRHI::RHISetShaderParameter(FVertexShaderRHIParamRef VertexShaderRHI,uint32 BufferIndex,uint32 BaseIndex,uint32 NumBytes,const void* NewValue)
 {
+	GRHICommandList.Verify();
+
 	VALIDATE_BOUND_SHADER(VertexShaderRHI);
 	checkSlow(VSConstantBuffers[BufferIndex]);
 	VSConstantBuffers[BufferIndex]->UpdateConstant((const uint8*)NewValue,BaseIndex,NumBytes);
@@ -625,6 +747,8 @@ void FD3D11DynamicRHI::RHISetShaderParameter(FVertexShaderRHIParamRef VertexShad
 
 void FD3D11DynamicRHI::RHISetShaderParameter(FPixelShaderRHIParamRef PixelShaderRHI,uint32 BufferIndex,uint32 BaseIndex,uint32 NumBytes,const void* NewValue)
 {
+	GRHICommandList.Verify();
+
 	VALIDATE_BOUND_SHADER(PixelShaderRHI);
 	checkSlow(PSConstantBuffers[BufferIndex]);
 	PSConstantBuffers[BufferIndex]->UpdateConstant((const uint8*)NewValue,BaseIndex,NumBytes);
@@ -632,6 +756,8 @@ void FD3D11DynamicRHI::RHISetShaderParameter(FPixelShaderRHIParamRef PixelShader
 
 void FD3D11DynamicRHI::RHISetShaderParameter(FGeometryShaderRHIParamRef GeometryShaderRHI,uint32 BufferIndex,uint32 BaseIndex,uint32 NumBytes,const void* NewValue)
 {
+	GRHICommandList.Verify();
+
 	VALIDATE_BOUND_SHADER(GeometryShaderRHI);
 	checkSlow(GSConstantBuffers[BufferIndex]);
 	GSConstantBuffers[BufferIndex]->UpdateConstant((const uint8*)NewValue,BaseIndex,NumBytes);
@@ -639,6 +765,8 @@ void FD3D11DynamicRHI::RHISetShaderParameter(FGeometryShaderRHIParamRef Geometry
 
 void FD3D11DynamicRHI::RHISetShaderParameter(FComputeShaderRHIParamRef ComputeShaderRHI,uint32 BufferIndex,uint32 BaseIndex,uint32 NumBytes,const void* NewValue)
 {
+	GRHICommandList.Verify();
+
 	//VALIDATE_BOUND_SHADER(ComputeShaderRHI);
 	checkSlow(CSConstantBuffers[BufferIndex]);
 	CSConstantBuffers[BufferIndex]->UpdateConstant((const uint8*)NewValue,BaseIndex,NumBytes);
@@ -646,6 +774,8 @@ void FD3D11DynamicRHI::RHISetShaderParameter(FComputeShaderRHIParamRef ComputeSh
 
 void FD3D11DynamicRHI::RHISetDepthStencilState(FDepthStencilStateRHIParamRef NewStateRHI,uint32 StencilRef)
 {
+	GRHICommandList.Verify();
+
 	DYNAMIC_CAST_D3D11RESOURCE(DepthStencilState,NewState);
 
 	if (CurrentDepthTexture && NewState->AccessType != CurrentDSVAccessType)
@@ -664,12 +794,16 @@ void FD3D11DynamicRHI::RHISetDepthStencilState(FDepthStencilStateRHIParamRef New
 
 void FD3D11DynamicRHI::RHISetBlendState(FBlendStateRHIParamRef NewStateRHI,const FLinearColor& BlendFactor)
 {
+	GRHICommandList.Verify();
+
 	DYNAMIC_CAST_D3D11RESOURCE(BlendState,NewState);
 	StateCache.SetBlendState(NewState->Resource, (const float*)&BlendFactor, 0xffffffff);
 }
 
 void FD3D11DynamicRHI::CommitRenderTargetsAndUAVs()
 {
+	GRHICommandList.Verify();
+
 	ID3D11RenderTargetView* RTArray[D3D11_SIMULTANEOUS_RENDER_TARGET_COUNT];
 	for(uint32 RenderTargetIndex = 0;RenderTargetIndex < NumSimultaneousRenderTargets;++RenderTargetIndex)
 	{
@@ -776,6 +910,8 @@ void FD3D11DynamicRHI::RHISetRenderTargets(
 	const FUnorderedAccessViewRHIParamRef* UAVs
 	)
 {
+	GRHICommandList.Verify();
+
 	FD3D11TextureBase* NewDepthStencilTarget = GetD3D11TextureFromRHITexture(NewDepthStencilTargetRHI);
 
 	check(NewNumSimultaneousRenderTargets + NewNumUAVs <= MaxSimultaneousRenderTargets);
@@ -907,12 +1043,16 @@ void FD3D11DynamicRHI::RHISetRenderTargets(
 
 void FD3D11DynamicRHI::RHIDiscardRenderTargets(bool Depth, bool Stencil, uint32 ColorBitMask)
 {
+	GRHICommandList.Verify();
+
 	// Could support in DX11.1 via ID3D11DeviceContext1::Discard*() functions.
 }
 
 // Occlusion/Timer queries.
 void FD3D11DynamicRHI::RHIBeginRenderQuery(FRenderQueryRHIParamRef QueryRHI)
 {
+	GRHICommandList.Verify();
+
 	DYNAMIC_CAST_D3D11RESOURCE(OcclusionQuery,Query);
 
 	if(Query->QueryType == RQT_Occlusion)
@@ -927,6 +1067,8 @@ void FD3D11DynamicRHI::RHIBeginRenderQuery(FRenderQueryRHIParamRef QueryRHI)
 }
 void FD3D11DynamicRHI::RHIEndRenderQuery(FRenderQueryRHIParamRef QueryRHI)
 {
+	GRHICommandList.Verify();
+
 	DYNAMIC_CAST_D3D11RESOURCE(OcclusionQuery,Query);
 	Direct3DDeviceIMContext->End(Query->Resource);
 
@@ -1058,6 +1200,8 @@ static uint32 GetVertexCountForPrimitiveCount(uint32 NumPrimitives, uint32 Primi
 
 void FD3D11DynamicRHI::CommitNonComputeShaderConstants()
 {
+	GRHICommandList.Verify();
+
 	FD3D11BoundShaderState* CurrentBoundShaderState = (FD3D11BoundShaderState*)BoundShaderStateHistory.GetLast();
 	check(CurrentBoundShaderState);
 
@@ -1125,6 +1269,8 @@ void FD3D11DynamicRHI::CommitNonComputeShaderConstants()
 
 void FD3D11DynamicRHI::CommitComputeShaderConstants()
 {
+	GRHICommandList.Verify();
+
 	bool bLocalDiscardSharedConstants = true;
 
 	// Commit and bind compute shader constants
@@ -1135,11 +1281,124 @@ void FD3D11DynamicRHI::CommitComputeShaderConstants()
 	}
 }
 
+template <EShaderFrequency Frequency>
+FORCEINLINE void SetResource(FD3D11DynamicRHI* RESTRICT D3D11RHI, FD3D11StateCache* RESTRICT StateCache, uint32 BindIndex, FD3D11BaseShaderResource* RESTRICT ShaderResource, ID3D11ShaderResourceView* RESTRICT SRV)
+{
+	// We set the resource through the RHI to track state for the purposes of unbinding SRVs when a UAV or RTV is bound.
+	// todo: need to support SRV_Static for faster calls when possible
+	D3D11RHI->SetShaderResourceView<Frequency>(ShaderResource,SRV,BindIndex,FD3D11StateCache::SRV_Unknown);
+}
+
+template <EShaderFrequency Frequency>
+FORCEINLINE void SetResource(FD3D11DynamicRHI* RESTRICT D3D11RHI, FD3D11StateCache* RESTRICT StateCache, uint32 BindIndex, FD3D11BaseShaderResource* RESTRICT ShaderResource, ID3D11SamplerState* RESTRICT SamplerState)
+{
+	StateCache->SetSamplerState<Frequency>(SamplerState,BindIndex);
+}
+
+template <class D3DResourceType, EShaderFrequency ShaderFrequency>
+inline int32 SetShaderResourcesFromBuffer(FD3D11DynamicRHI* RESTRICT D3D11RHI, FD3D11StateCache* RESTRICT StateCache, FD3D11UniformBuffer* RESTRICT Buffer, const uint32 * RESTRICT ResourceMap, int32 BufferIndex)
+{
+	int32 NumSetCalls = 0;
+	uint32 BufferOffset = ResourceMap[BufferIndex];
+	if (BufferOffset > 0)
+	{
+		const uint32* RESTRICT ResourceInfos = &ResourceMap[BufferOffset];
+		uint32 ResourceInfo = *ResourceInfos++;
+		do 
+		{
+			checkSlow(FRHIResourceTableEntry::GetUniformBufferIndex(ResourceInfo) == BufferIndex);
+			const uint16 ResourceIndex = FRHIResourceTableEntry::GetResourceIndex(ResourceInfo);
+			const uint8 BindIndex = FRHIResourceTableEntry::GetBindIndex(ResourceInfo);
+
+			// todo: could coalesce adjacent bound resources.
+			FD3D11UniformBuffer::FResourcePair* RESTRICT ResourcePair = &Buffer->RawResourceTable[ResourceIndex];
+			FD3D11BaseShaderResource* ShaderResource = ResourcePair->ShaderResource;
+			D3DResourceType* D3D11Resource = (D3DResourceType*)ResourcePair->D3D11Resource;
+			SetResource<ShaderFrequency>(D3D11RHI, StateCache, BindIndex, ShaderResource, D3D11Resource);
+			NumSetCalls++;
+			ResourceInfo = *ResourceInfos++;
+		} while (FRHIResourceTableEntry::GetUniformBufferIndex(ResourceInfo) == BufferIndex);
+	}
+	return NumSetCalls;
+}
+
+template <class ShaderType>
+void FD3D11DynamicRHI::SetResourcesFromTables(const ShaderType* RESTRICT Shader)
+{
+	checkSlow(Shader);
+
+	// Mask the dirty bits by those buffers from which the shader has bound resources.
+	uint32 DirtyBits = Shader->ShaderResourceTable.ResourceTableBits & DirtyUniformBuffers[ShaderType::StaticFrequency];
+	uint32 NumSetCalls = 0;
+	while (DirtyBits)
+	{
+		// Scan for the lowest set bit, compute its index, clear it in the set of dirty bits.
+		const uint32 LowestBitMask = (DirtyBits) & (-(int32)DirtyBits);
+		const int32 BufferIndex = FMath::FloorLog2(LowestBitMask); // todo: This has a branch on zero, we know it could never be zero...
+		DirtyBits ^= LowestBitMask;
+		FD3D11UniformBuffer* Buffer = (FD3D11UniformBuffer*)BoundUniformBuffers[ShaderType::StaticFrequency][BufferIndex].GetReference();
+		check(Buffer);
+		check(BufferIndex < Shader->ShaderResourceTable.ResourceTableLayoutHashes.Num());
+		check(Buffer->GetLayout().GetHash() == Shader->ShaderResourceTable.ResourceTableLayoutHashes[BufferIndex]);
+		Buffer->CacheResources(ResourceTableFrameCounter);
+
+		// todo: could make this two pass: gather then set
+		NumSetCalls += SetShaderResourcesFromBuffer<ID3D11ShaderResourceView, (EShaderFrequency)ShaderType::StaticFrequency>(this, &StateCache, Buffer, Shader->ShaderResourceTable.ShaderResourceViewMap.GetData(), BufferIndex);
+		NumSetCalls += SetShaderResourcesFromBuffer<ID3D11SamplerState, (EShaderFrequency)ShaderType::StaticFrequency>(this, &StateCache, Buffer, Shader->ShaderResourceTable.SamplerMap.GetData(), BufferIndex);
+	}
+	DirtyUniformBuffers[ShaderType::StaticFrequency] = 0;
+	SetTextureInTableCalls += NumSetCalls;
+}
+
+void FD3D11DynamicRHI::CommitGraphicsResourceTables()
+{
+	uint32 Start = FPlatformTime::Cycles();
+
+	GRHICommandList.Verify();
+
+	FD3D11BoundShaderState* RESTRICT CurrentBoundShaderState = (FD3D11BoundShaderState*)BoundShaderStateHistory.GetLast();
+	check(CurrentBoundShaderState);
+
+	if (auto* Shader = CurrentBoundShaderState->GetVertexShader())
+	{
+		SetResourcesFromTables(Shader);
+	}
+	if (auto* Shader = CurrentBoundShaderState->GetPixelShader())
+	{
+		SetResourcesFromTables(Shader);
+	}
+	if (auto* Shader = CurrentBoundShaderState->GetHullShader())
+	{
+		SetResourcesFromTables(Shader);
+	}
+	if (auto* Shader = CurrentBoundShaderState->GetDomainShader())
+	{
+		SetResourcesFromTables(Shader);
+	}
+	if (auto* Shader = CurrentBoundShaderState->GetGeometryShader())
+	{
+		SetResourcesFromTables(Shader);
+	}
+
+	CommitResourceTableCycles += (FPlatformTime::Cycles() - Start);
+}
+
+void FD3D11DynamicRHI::CommitComputeResourceTables(FD3D11ComputeShader* InComputeShader)
+{
+	FD3D11ComputeShader* RESTRICT ComputeShader = InComputeShader;
+	check(ComputeShader);
+	SetResourcesFromTables(ComputeShader);
+}
+
 void FD3D11DynamicRHI::RHIDrawPrimitive(uint32 PrimitiveType,uint32 BaseVertexIndex,uint32 NumPrimitives,uint32 NumInstances)
 {
+	GRHICommandList.Verify();
+
 	RHI_DRAW_CALL_STATS(PrimitiveType,NumInstances*NumPrimitives);
 
+	CommitGraphicsResourceTables();
 	CommitNonComputeShaderConstants();
+
 	uint32 VertexCount = RHIGetVertexCountForPrimitiveCount(NumPrimitives,PrimitiveType);
 
 	GPUProfilingData.RegisterGPUWork(NumPrimitives * NumInstances, VertexCount * NumInstances);
@@ -1157,12 +1416,15 @@ void FD3D11DynamicRHI::RHIDrawPrimitive(uint32 PrimitiveType,uint32 BaseVertexIn
 
 void FD3D11DynamicRHI::RHIDrawPrimitiveIndirect(uint32 PrimitiveType,FVertexBufferRHIParamRef ArgumentBufferRHI,uint32 ArgumentOffset)
 {
+	GRHICommandList.Verify();
+
 	DYNAMIC_CAST_D3D11RESOURCE(VertexBuffer,ArgumentBuffer);
 
 	RHI_DRAW_CALL_INC();
 
 	GPUProfilingData.RegisterGPUWork(0);
 
+	CommitGraphicsResourceTables();
 	CommitNonComputeShaderConstants();
 
 	StateCache.SetPrimitiveTopology(GetD3D11PrimitiveType(PrimitiveType,bUsingTessellation));
@@ -1171,6 +1433,8 @@ void FD3D11DynamicRHI::RHIDrawPrimitiveIndirect(uint32 PrimitiveType,FVertexBuff
 
 void FD3D11DynamicRHI::RHIDrawIndexedIndirect(FIndexBufferRHIParamRef IndexBufferRHI, uint32 PrimitiveType, FStructuredBufferRHIParamRef ArgumentsBufferRHI, int32 DrawArgumentsIndex, uint32 NumInstances)
 {
+	GRHICommandList.Verify();
+
 	DYNAMIC_CAST_D3D11RESOURCE(IndexBuffer,IndexBuffer);
 	DYNAMIC_CAST_D3D11RESOURCE(StructuredBuffer,ArgumentsBuffer);
 
@@ -1178,7 +1442,9 @@ void FD3D11DynamicRHI::RHIDrawIndexedIndirect(FIndexBufferRHIParamRef IndexBuffe
 
 	GPUProfilingData.RegisterGPUWork(1);
 
+	CommitGraphicsResourceTables();
 	CommitNonComputeShaderConstants();
+
 	// determine 16bit vs 32bit indices
 	uint32 SizeFormat = sizeof(DXGI_FORMAT);
 	const DXGI_FORMAT Format = (IndexBuffer->GetStride() == sizeof(uint16) ? DXGI_FORMAT_R16_UINT : DXGI_FORMAT_R32_UINT);
@@ -1198,6 +1464,8 @@ void FD3D11DynamicRHI::RHIDrawIndexedIndirect(FIndexBufferRHIParamRef IndexBuffe
 
 void FD3D11DynamicRHI::RHIDrawIndexedPrimitive(FIndexBufferRHIParamRef IndexBufferRHI,uint32 PrimitiveType,int32 BaseVertexIndex,uint32 MinIndex,uint32 NumVertices,uint32 StartIndex,uint32 NumPrimitives,uint32 NumInstances)
 {
+	GRHICommandList.Verify();
+
 	DYNAMIC_CAST_D3D11RESOURCE(IndexBuffer,IndexBuffer);
 
 	// called should make sure the input is valid, this avoid hidden bugs
@@ -1207,7 +1475,9 @@ void FD3D11DynamicRHI::RHIDrawIndexedPrimitive(FIndexBufferRHIParamRef IndexBuff
 
 	GPUProfilingData.RegisterGPUWork(NumPrimitives * NumInstances, NumVertices * NumInstances);
 
+	CommitGraphicsResourceTables();
 	CommitNonComputeShaderConstants();
+
 	// determine 16bit vs 32bit indices
 	uint32 SizeFormat = sizeof(DXGI_FORMAT);
 	const DXGI_FORMAT Format = (IndexBuffer->GetStride() == sizeof(uint16) ? DXGI_FORMAT_R16_UINT : DXGI_FORMAT_R32_UINT);
@@ -1234,6 +1504,8 @@ void FD3D11DynamicRHI::RHIDrawIndexedPrimitive(FIndexBufferRHIParamRef IndexBuff
 
 void FD3D11DynamicRHI::RHIDrawIndexedPrimitiveIndirect(uint32 PrimitiveType,FIndexBufferRHIParamRef IndexBufferRHI,FVertexBufferRHIParamRef ArgumentBufferRHI,uint32 ArgumentOffset)
 {
+	GRHICommandList.Verify();
+
 	DYNAMIC_CAST_D3D11RESOURCE(IndexBuffer,IndexBuffer);
 	DYNAMIC_CAST_D3D11RESOURCE(VertexBuffer,ArgumentBuffer);
 
@@ -1241,6 +1513,7 @@ void FD3D11DynamicRHI::RHIDrawIndexedPrimitiveIndirect(uint32 PrimitiveType,FInd
 
 	GPUProfilingData.RegisterGPUWork(0);
 	
+	CommitGraphicsResourceTables();
 	CommitNonComputeShaderConstants();
 	
 	// Set the index buffer.
@@ -1261,6 +1534,8 @@ void FD3D11DynamicRHI::RHIDrawIndexedPrimitiveIndirect(uint32 PrimitiveType,FInd
  */
 void FD3D11DynamicRHI::RHIBeginDrawPrimitiveUP( uint32 PrimitiveType, uint32 NumPrimitives, uint32 NumVertices, uint32 VertexDataStride, void*& OutVertexData)
 {
+	GRHICommandList.Verify();
+
 	checkSlow( PendingNumVertices == 0 );
 
 	// Remember the parameters for the draw call.
@@ -1278,6 +1553,8 @@ void FD3D11DynamicRHI::RHIBeginDrawPrimitiveUP( uint32 PrimitiveType, uint32 Num
  */
 void FD3D11DynamicRHI::RHIEndDrawPrimitiveUP()
 {
+	GRHICommandList.Verify();
+
 	RHI_DRAW_CALL_STATS(PendingPrimitiveType,PendingNumPrimitives);
 
 	checkSlow(!bUsingTessellation || PendingPrimitiveType == PT_TriangleList);
@@ -1289,6 +1566,7 @@ void FD3D11DynamicRHI::RHIEndDrawPrimitiveUP()
 	uint32 VBOffset = 0;
 
 	// Issue the draw call.
+	CommitGraphicsResourceTables();
 	CommitNonComputeShaderConstants();
 	StateCache.SetStreamSource(D3DBuffer, 0, PendingVertexDataStride, VBOffset);
 	StateCache.SetPrimitiveTopology(GetD3D11PrimitiveType(PendingPrimitiveType,bUsingTessellation));
@@ -1315,6 +1593,8 @@ void FD3D11DynamicRHI::RHIEndDrawPrimitiveUP()
  */
 void FD3D11DynamicRHI::RHIBeginDrawIndexedPrimitiveUP( uint32 PrimitiveType, uint32 NumPrimitives, uint32 NumVertices, uint32 VertexDataStride, void*& OutVertexData, uint32 MinVertexIndex, uint32 NumIndices, uint32 IndexDataStride, void*& OutIndexData)
 {
+	GRHICommandList.Verify();
+
 	checkSlow((sizeof(uint16) == IndexDataStride) || (sizeof(uint32) == IndexDataStride));
 
 	// Store off information needed for the draw call.
@@ -1336,6 +1616,8 @@ void FD3D11DynamicRHI::RHIBeginDrawIndexedPrimitiveUP( uint32 PrimitiveType, uin
  */
 void FD3D11DynamicRHI::RHIEndDrawIndexedPrimitiveUP()
 {
+	GRHICommandList.Verify();
+
 	// tessellation only supports trilists
 	checkSlow(!bUsingTessellation || PendingPrimitiveType == PT_TriangleList);
 
@@ -1349,6 +1631,7 @@ void FD3D11DynamicRHI::RHIEndDrawIndexedPrimitiveUP()
 	uint32 VBOffset = 0;
 
 	// Issue the draw call.
+	CommitGraphicsResourceTables();
 	CommitNonComputeShaderConstants();
 	StateCache.SetStreamSource(VertexBuffer, 0, PendingVertexDataStride, VBOffset);
 	StateCache.SetIndexBuffer(IndexBuffer, PendingIndexDataStride == sizeof(uint16) ? DXGI_FORMAT_R16_UINT : DXGI_FORMAT_R32_UINT, 0);
@@ -1368,11 +1651,15 @@ void FD3D11DynamicRHI::RHIEndDrawIndexedPrimitiveUP()
 // Raster operations.
 void FD3D11DynamicRHI::RHIClear(bool bClearColor,const FLinearColor& Color,bool bClearDepth,float Depth,bool bClearStencil,uint32 Stencil, FIntRect ExcludeRect)
 {
+	GRHICommandList.Verify();
+
 	FD3D11DynamicRHI::RHIClearMRT(bClearColor, 1, &Color, bClearDepth, Depth, bClearStencil, Stencil, ExcludeRect);
 }
 
 void FD3D11DynamicRHI::RHIClearMRT(bool bClearColor,int32 NumClearColors,const FLinearColor* ClearColorArray,bool bClearDepth,float Depth,bool bClearStencil,uint32 Stencil, FIntRect ExcludeRect)
 {
+	GRHICommandList.Verify();
+
 	// Helper struct to record and restore device states RHIClearMRT modifies.
 	class FDeviceStateHelper
 	{
@@ -1700,7 +1987,9 @@ void FD3D11DynamicRHI::RHIClearMRT(bool bClearColor,int32 NumClearColors,const F
 			ShaderClearColors[i] = ClearColorArray[i];
 		}
 
-		SetShaderValueArray(PixelShader->GetPixelShader(),PixelShader->ColorParameter,ShaderClearColors,NumClearColors);
+		//@todo-rco: RHIPacketList
+		FRHICommandList* RHICmdList = nullptr;
+		SetShaderValueArray(RHICmdList, PixelShader->GetPixelShader(),PixelShader->ColorParameter,ShaderClearColors,NumClearColors);
 		
 		{
 			// Draw a fullscreen quad
@@ -1792,16 +2081,22 @@ void FD3D11DynamicRHI::RHIClearMRT(bool bClearColor,int32 NumClearColors,const F
 
 void FD3D11DynamicRHI::RHISuspendRendering()
 {
+	GRHICommandList.Verify();
+
 	// Not supported
 }
 
 void FD3D11DynamicRHI::RHIResumeRendering()
 {
+	GRHICommandList.Verify();
+
 	// Not supported
 }
 
 bool FD3D11DynamicRHI::RHIIsRenderingSuspended()
 {
+	GRHICommandList.Verify();
+
 	// Not supported
 	return false;
 }
@@ -1809,6 +2104,8 @@ bool FD3D11DynamicRHI::RHIIsRenderingSuspended()
 // Blocks the CPU until the GPU catches up and goes idle.
 void FD3D11DynamicRHI::RHIBlockUntilGPUIdle()
 {
+	GRHICommandList.Verify();
+
 	// Not really supported
 }
 
@@ -1820,6 +2117,349 @@ uint32 FD3D11DynamicRHI::RHIGetGPUFrameCycles()
 	return GGPUFrameTime;
 }
 
+void FD3D11DynamicRHI::RHIExecuteCommandList(FRHICommandList* CmdList)
+{
+	auto* CmdPtr = CmdList->GetHead();
+	auto* CmdTail = CmdList->GetTail();
+
+	while (CmdPtr < CmdTail)
+	{
+		auto* Cmd = (FRHICommand*)CmdPtr;
+		switch (Cmd->Type)
+		{
+		case ERCT_NopBlob:
+			{
+				// Nop
+				auto* RHICmd = (FRHICommandNopBlob*)Cmd;
+				CmdPtr += sizeof(FRHICommandNopBlob) + RHICmd->Size;
+			}
+			break;
+		case ERCT_SetRasterizerState:
+			{
+				auto* RHICmd = (FRHICommandSetRasterizerState*)Cmd;
+				{
+					auto* State = (FD3D11RasterizerState*)RHICmd->State;
+					StateCache.SetRasterizerState(State->Resource);
+				}
+				CmdPtr += sizeof(FRHICommandSetRasterizerState);
+				RHICmd->State->Release();
+			}
+			break;
+		case ERCT_SetShaderParameter:
+			{
+				auto* RHICmd = (FRHICommandSetShaderParameter*)Cmd;
+				{
+					switch (RHICmd->ShaderFrequency)
+					{
+					case SF_Vertex: 
+						VALIDATE_BOUND_SHADER((FVertexShaderRHIParamRef)RHICmd->Shader);
+						checkSlow(VSConstantBuffers[RHICmd->BufferIndex]);
+						VSConstantBuffers[RHICmd->BufferIndex]->UpdateConstant((const uint8*)RHICmd->NewValue, RHICmd->BaseIndex, RHICmd->NumBytes);
+						break;
+					case SF_Hull:
+						VALIDATE_BOUND_SHADER((FHullShaderRHIParamRef)RHICmd->Shader);
+						checkSlow(HSConstantBuffers[RHICmd->BufferIndex]);
+						HSConstantBuffers[RHICmd->BufferIndex]->UpdateConstant((const uint8*)RHICmd->NewValue, RHICmd->BaseIndex, RHICmd->NumBytes);
+						break;
+					case SF_Domain:
+						VALIDATE_BOUND_SHADER((FDomainShaderRHIParamRef)RHICmd->Shader);
+						checkSlow(DSConstantBuffers[RHICmd->BufferIndex]);
+						DSConstantBuffers[RHICmd->BufferIndex]->UpdateConstant((const uint8*)RHICmd->NewValue, RHICmd->BaseIndex, RHICmd->NumBytes);
+						break;
+					case SF_Geometry:
+						VALIDATE_BOUND_SHADER((FGeometryShaderRHIParamRef)RHICmd->Shader);
+						checkSlow(GSConstantBuffers[RHICmd->BufferIndex]);
+						GSConstantBuffers[RHICmd->BufferIndex]->UpdateConstant((const uint8*)RHICmd->NewValue, RHICmd->BaseIndex, RHICmd->NumBytes);
+						break;
+					case SF_Pixel:
+						VALIDATE_BOUND_SHADER((FPixelShaderRHIParamRef)RHICmd->Shader);
+						checkSlow(PSConstantBuffers[RHICmd->BufferIndex]);
+						PSConstantBuffers[RHICmd->BufferIndex]->UpdateConstant((const uint8*)RHICmd->NewValue, RHICmd->BaseIndex, RHICmd->NumBytes);
+						break;
+					case SF_Compute:
+						VALIDATE_BOUND_SHADER((FComputeShaderRHIParamRef)RHICmd->Shader);
+						checkSlow(CSConstantBuffers[RHICmd->BufferIndex]);
+						CSConstantBuffers[RHICmd->BufferIndex]->UpdateConstant((const uint8*)RHICmd->NewValue, RHICmd->BaseIndex, RHICmd->NumBytes);
+						break;
+					default: check(0); break;
+					}
+				}
+				CmdPtr += sizeof(FRHICommandSetShaderParameter);
+			}
+			break;
+		case ERCT_SetShaderUniformBuffer:
+			{
+				auto* RHICmd = (FRHICommandSetShaderUniformBuffer*)Cmd;
+				{
+					FD3D11UniformBuffer* Buffer = (FD3D11UniformBuffer*)RHICmd->UniformBuffer;
+					ID3D11Buffer* ConstantBuffer = Buffer ? Buffer->Resource : nullptr;
+					switch (RHICmd->ShaderFrequency)
+					{
+					case SF_Vertex:
+						VALIDATE_BOUND_SHADER((FVertexShaderRHIParamRef)RHICmd->Shader);
+						StateCache.SetConstantBuffer<SF_Vertex>(ConstantBuffer, RHICmd->BaseIndex);
+						BoundUniformBuffers[SF_Vertex][RHICmd->BaseIndex] = Buffer;
+						DirtyUniformBuffers[SF_Vertex] |= (1 << RHICmd->BaseIndex);
+						break;
+					case SF_Hull:
+						VALIDATE_BOUND_SHADER((FHullShaderRHIParamRef)RHICmd->Shader);
+						StateCache.SetConstantBuffer<SF_Hull>(ConstantBuffer, RHICmd->BaseIndex);
+						BoundUniformBuffers[SF_Hull][RHICmd->BaseIndex] = Buffer;
+						DirtyUniformBuffers[SF_Hull] |= (1 << RHICmd->BaseIndex);
+						break;
+					case SF_Domain:
+						VALIDATE_BOUND_SHADER((FDomainShaderRHIParamRef)RHICmd->Shader);
+						StateCache.SetConstantBuffer<SF_Domain>(ConstantBuffer, RHICmd->BaseIndex);
+						BoundUniformBuffers[SF_Domain][RHICmd->BaseIndex] = Buffer;
+						DirtyUniformBuffers[SF_Domain] |= (1 << RHICmd->BaseIndex);
+						break;
+					case SF_Geometry:
+						VALIDATE_BOUND_SHADER((FGeometryShaderRHIParamRef)RHICmd->Shader);
+						StateCache.SetConstantBuffer<SF_Geometry>(ConstantBuffer, RHICmd->BaseIndex);
+						BoundUniformBuffers[SF_Geometry][RHICmd->BaseIndex] = Buffer;
+						DirtyUniformBuffers[SF_Geometry] |= (1 << RHICmd->BaseIndex);
+						break;
+					case SF_Pixel:
+						VALIDATE_BOUND_SHADER((FPixelShaderRHIParamRef)RHICmd->Shader);
+						StateCache.SetConstantBuffer<SF_Pixel>(ConstantBuffer, RHICmd->BaseIndex);
+						BoundUniformBuffers[SF_Pixel][RHICmd->BaseIndex] = Buffer;
+						DirtyUniformBuffers[SF_Pixel] |= (1 << RHICmd->BaseIndex);
+						break;
+					case SF_Compute:
+						VALIDATE_BOUND_SHADER((FComputeShaderRHIParamRef)RHICmd->Shader);
+						StateCache.SetConstantBuffer<SF_Compute>(ConstantBuffer, RHICmd->BaseIndex);
+						BoundUniformBuffers[SF_Compute][RHICmd->BaseIndex] = Buffer;
+						DirtyUniformBuffers[SF_Compute] |= (1 << RHICmd->BaseIndex);
+						break;
+					default: check(0); break;
+					}
+				}
+				CmdPtr += sizeof(FRHICommandSetShaderUniformBuffer);
+				RHICmd->Shader->Release(); 
+				RHICmd->UniformBuffer->Release();
+			}
+			break;
+		case ERCT_SetShaderSampler:
+			{
+				auto* RHICmd = (FRHICommandSetShaderSampler*)Cmd;
+				{
+					FD3D11SamplerState* SamplerState = (FD3D11SamplerState*)RHICmd->Sampler;
+					ID3D11SamplerState* StateResource = SamplerState->Resource;
+					switch (RHICmd->ShaderFrequency)
+					{
+					case SF_Vertex:
+						VALIDATE_BOUND_SHADER((FVertexShaderRHIParamRef)RHICmd->Shader);
+						StateCache.SetSamplerState<SF_Vertex>(StateResource, RHICmd->SamplerIndex);
+						break;
+					case SF_Hull:
+						VALIDATE_BOUND_SHADER((FHullShaderRHIParamRef)RHICmd->Shader);
+						StateCache.SetSamplerState<SF_Hull>(StateResource, RHICmd->SamplerIndex);
+						break;
+					case SF_Domain:
+						VALIDATE_BOUND_SHADER((FDomainShaderRHIParamRef)RHICmd->Shader);
+						StateCache.SetSamplerState<SF_Domain>(StateResource, RHICmd->SamplerIndex);
+						break;
+					case SF_Geometry:
+						VALIDATE_BOUND_SHADER((FGeometryShaderRHIParamRef)RHICmd->Shader);
+						StateCache.SetSamplerState<SF_Geometry>(StateResource, RHICmd->SamplerIndex);
+						break;
+					case SF_Pixel:
+						VALIDATE_BOUND_SHADER((FPixelShaderRHIParamRef)RHICmd->Shader);
+						StateCache.SetSamplerState<SF_Pixel>(StateResource, RHICmd->SamplerIndex);
+						break;
+					case SF_Compute:
+						VALIDATE_BOUND_SHADER((FComputeShaderRHIParamRef)RHICmd->Shader);
+						StateCache.SetSamplerState<SF_Compute>(StateResource, RHICmd->SamplerIndex);
+						break;
+					default: check(0); break;
+					}
+				}
+				CmdPtr += sizeof(FRHICommandSetShaderSampler);
+				RHICmd->Shader->Release(); 
+				RHICmd->Sampler->Release();
+			}
+			break;
+		case ERCT_SetShaderTexture:
+			{
+				auto* RHICmd = (FRHICommandSetShaderTexture*)Cmd;
+				{
+					uint32 Start = FPlatformTime::Cycles();
+					FD3D11TextureBase* Texture = GetD3D11TextureFromRHITexture(RHICmd->Texture);
+					ID3D11ShaderResourceView* ShaderResourceView = Texture ? Texture->GetShaderResourceView() : nullptr;
+					auto SRVType = (Texture == nullptr || Texture->GetRenderTargetView(0, 0) != nullptr || Texture->HasDepthStencilView()) ? FD3D11StateCache::SRV_Dynamic : FD3D11StateCache::SRV_Static;
+					switch (RHICmd->ShaderFrequency)
+					{
+					case SF_Vertex:
+						VALIDATE_BOUND_SHADER((FVertexShaderRHIParamRef)RHICmd->Shader);
+						SetShaderResourceView<SF_Vertex>(Texture, ShaderResourceView, RHICmd->TextureIndex, SRVType);
+						break;
+					case SF_Hull:
+						VALIDATE_BOUND_SHADER((FHullShaderRHIParamRef)RHICmd->Shader);
+						SetShaderResourceView<SF_Hull>(Texture, ShaderResourceView, RHICmd->TextureIndex, SRVType);
+						break;
+					case SF_Domain:
+						VALIDATE_BOUND_SHADER((FDomainShaderRHIParamRef)RHICmd->Shader);
+						SetShaderResourceView<SF_Domain>(Texture, ShaderResourceView, RHICmd->TextureIndex, SRVType);
+						break;
+					case SF_Geometry:
+						VALIDATE_BOUND_SHADER((FGeometryShaderRHIParamRef)RHICmd->Shader);
+						SetShaderResourceView<SF_Geometry>(Texture, ShaderResourceView, RHICmd->TextureIndex, SRVType);
+						break;
+					case SF_Pixel:
+						VALIDATE_BOUND_SHADER((FPixelShaderRHIParamRef)RHICmd->Shader);
+						SetShaderResourceView<SF_Pixel>(Texture, ShaderResourceView, RHICmd->TextureIndex, SRVType);
+						break;
+					case SF_Compute:
+						VALIDATE_BOUND_SHADER((FComputeShaderRHIParamRef)RHICmd->Shader);
+						SetShaderResourceView<SF_Compute>(Texture, ShaderResourceView, RHICmd->TextureIndex, SRVType);
+						break;
+					default: check(0); break;
+					}
+					SetShaderTextureCycles += (FPlatformTime::Cycles() - Start);
+					SetShaderTextureCalls++;
+				}
+				CmdPtr += sizeof(FRHICommandSetShaderTexture);
+				RHICmd->Shader->Release(); 
+				RHICmd->Texture->Release();
+			}
+			break;
+		case ERCT_DrawPrimitive:
+			{
+				auto* RHICmd = (FRHICommandDrawPrimitive*)Cmd;
+				{
+					RHI_DRAW_CALL_STATS(RHICmd->PrimitiveType, RHICmd->NumInstances * RHICmd->NumPrimitives);
+
+					CommitGraphicsResourceTables();
+					CommitNonComputeShaderConstants();
+					const uint32 VertexCount = RHIGetVertexCountForPrimitiveCount(RHICmd->NumPrimitives, RHICmd->PrimitiveType);
+
+					GPUProfilingData.RegisterGPUWork(RHICmd->NumPrimitives * RHICmd->NumInstances, VertexCount * RHICmd->NumInstances);
+					StateCache.SetPrimitiveTopology(GetD3D11PrimitiveType(RHICmd->PrimitiveType, bUsingTessellation));
+					if (RHICmd->NumInstances > 1)
+					{
+						Direct3DDeviceIMContext->DrawInstanced(VertexCount, RHICmd->NumInstances, RHICmd->BaseVertexIndex,0);
+					}
+					else
+					{
+						Direct3DDeviceIMContext->Draw(VertexCount, RHICmd->BaseVertexIndex);
+					}
+				}
+				CmdPtr += sizeof(FRHICommandDrawPrimitive);
+			}
+			break;
+		case ERCT_DrawIndexedPrimitive:
+			{
+				auto* RHICmd = (FRHICommandDrawIndexedPrimitive*)Cmd;
+				{
+					auto* IndexBuffer = (FD3D11IndexBuffer*)RHICmd->IndexBuffer;
+
+					// Caller should make sure the input is valid, this avoid hidden bugs
+					ensure(RHICmd->NumPrimitives > 0);
+
+					RHI_DRAW_CALL_STATS(RHICmd->PrimitiveType, RHICmd->NumInstances* RHICmd->NumPrimitives);
+
+					GPUProfilingData.RegisterGPUWork(RHICmd->NumPrimitives * RHICmd->NumInstances, RHICmd->NumVertices * RHICmd->NumInstances);
+
+					CommitGraphicsResourceTables();
+					CommitNonComputeShaderConstants();
+
+					// determine 16bit vs 32bit indices
+					uint32 SizeFormat = sizeof(DXGI_FORMAT);
+					const DXGI_FORMAT Format = (IndexBuffer->GetStride() == sizeof(uint16) ? DXGI_FORMAT_R16_UINT : DXGI_FORMAT_R32_UINT);
+
+					uint32 IndexCount = RHIGetVertexCountForPrimitiveCount(RHICmd->NumPrimitives, RHICmd->PrimitiveType);
+
+					// Verify that we are not trying to read outside the index buffer range
+					checkf(RHICmd->StartIndex + IndexCount <= IndexBuffer->GetSize() / IndexBuffer->GetStride(), 
+						TEXT("Start %u, Count %u, Type %u, Buffer Size %u, Buffer stride %u"), RHICmd->StartIndex, IndexCount, RHICmd->PrimitiveType, IndexBuffer->GetSize(), IndexBuffer->GetStride());
+
+					StateCache.SetIndexBuffer(IndexBuffer->Resource, Format, 0);
+					StateCache.SetPrimitiveTopology(GetD3D11PrimitiveType(RHICmd->PrimitiveType, bUsingTessellation));
+
+					if (RHICmd->NumInstances > 1)
+					{
+						Direct3DDeviceIMContext->DrawIndexedInstanced(IndexCount, RHICmd->NumInstances, RHICmd->StartIndex, RHICmd->BaseVertexIndex,0);
+					}
+					else
+					{
+						Direct3DDeviceIMContext->DrawIndexed(IndexCount, RHICmd->StartIndex, RHICmd->BaseVertexIndex);
+					}
+				}
+				CmdPtr += sizeof(FRHICommandDrawIndexedPrimitive);
+				RHICmd->IndexBuffer->Release();
+			}
+			break;
+		case ERCT_SetBoundShaderState:
+			{
+				auto* RHICmd = (FRHICommandSetBoundShaderState*)Cmd;
+				{
+					auto* BoundShaderState = (FD3D11BoundShaderState*)RHICmd->BoundShaderState;
+
+					StateCache.SetInputLayout(BoundShaderState->InputLayout);
+					StateCache.SetVertexShader(BoundShaderState->VertexShader);
+					StateCache.SetPixelShader(BoundShaderState->PixelShader);
+
+					StateCache.SetHullShader(BoundShaderState->HullShader);
+					StateCache.SetDomainShader(BoundShaderState->DomainShader);
+					StateCache.SetGeometryShader(BoundShaderState->GeometryShader);
+
+					if(BoundShaderState->HullShader != NULL && BoundShaderState->DomainShader != NULL)
+					{
+						bUsingTessellation = true;
+					}
+					else
+					{
+						bUsingTessellation = false;
+					}
+
+					// @TODO : really should only discard the constants if the shader state has actually changed.
+					bDiscardSharedConstants = true;
+
+					// Prevent transient bound shader states from being recreated for each use by keeping a history of the most recently used bound shader states.
+					// The history keeps them alive, and the bound shader state cache allows them to be reused if needed.
+					BoundShaderStateHistory.Add(BoundShaderState);
+
+					// Shader changed so all resource tables are dirty
+					DirtyUniformBuffers[SF_Vertex] = 0xffff;
+					DirtyUniformBuffers[SF_Pixel] = 0xffff;
+					DirtyUniformBuffers[SF_Hull] = 0xffff;
+					DirtyUniformBuffers[SF_Domain] = 0xffff;
+					DirtyUniformBuffers[SF_Geometry] = 0xffff;
+				}
+				CmdPtr += sizeof(FRHICommandSetBoundShaderState);
+				RHICmd->BoundShaderState->Release();
+			}
+			break;
+		case ERCT_SetBlendState:
+			{
+				auto* RHICmd = (FRHICommandSetBlendState*)Cmd;
+				{
+					auto* State = (FD3D11BlendState*)RHICmd->State;
+					StateCache.SetBlendState(State->Resource, (const float*)&RHICmd->BlendFactor, 0xffffffff);
+				}
+				CmdPtr += sizeof(FRHICommandSetBlendState);
+				RHICmd->State->Release();
+			}
+			break;
+		case ERCT_SetStreamSource:
+			{
+				auto* RHICmd = (FRHICommandSetStreamSource*)Cmd;
+				{
+					auto* VertexBuffer = (FD3D11VertexBuffer*)RHICmd->VertexBuffer;
+					ID3D11Buffer* D3DBuffer = VertexBuffer ? VertexBuffer->Resource : nullptr;
+					StateCache.SetStreamSource(D3DBuffer, RHICmd->StreamIndex, RHICmd->Stride, RHICmd->Offset);
+				}
+				CmdPtr += sizeof(FRHICommandSetStreamSource);
+				RHICmd->VertexBuffer->Release();
+			}
+			break;
+		default:
+			checkf(0, TEXT("Unimplemented RHICmd %d!"), Cmd->Type);
+		}
+
+		CmdPtr = Align(CmdPtr, FRHICommandList::Alignment);
+	}
+}
 
 // NVIDIA Depth Bounds Test interface
 void FD3D11DynamicRHI::RHIEnableDepthBoundsTest(bool bEnable,float MinDepth,float MaxDepth)
diff --git a/Engine/Source/Runtime/Windows/D3D11RHI/Private/D3D11Device.cpp b/Engine/Source/Runtime/Windows/D3D11RHI/Private/D3D11Device.cpp
--- a/Engine/Source/Runtime/Windows/D3D11RHI/Private/D3D11Device.cpp
+++ b/Engine/Source/Runtime/Windows/D3D11RHI/Private/D3D11Device.cpp
@@ -42,6 +42,14 @@ FD3D11DynamicRHI::FD3D11DynamicRHI(IDXGIFactory* InDXGIFactory,D3D_FEATURE_LEVEL
 	CurrentDepthTexture(NULL),
 	NumSimultaneousRenderTargets(0),
 	NumUAVs(0),
+	CommitResourceTableCycles(0),
+	CacheResourceTableCalls(0),
+	CacheResourceTableCycles(0),
+	SetShaderTextureCycles(0),
+	SetShaderTextureCalls(0),
+	SetTextureInTableCalls(0),
+	SceneFrameCounter(0),
+	ResourceTableFrameCounter(INDEX_NONE),
 	CurrentDSVAccessType(DSAT_Writable),
 	bDiscardSharedConstants(false),
 	GPUProfilingData(this),
@@ -189,6 +197,11 @@ FD3D11DynamicRHI::FD3D11DynamicRHI(IDXGIFactory* InDXGIFactory,D3D_FEATURE_LEVEL
 	DynamicVB = new FD3D11DynamicBuffer(this,D3D11_BIND_VERTEX_BUFFER,DynamicVBSizes);
 	uint32 DynamicIBSizes[] = {128,1024,64*1024,1024*1024,0};
 	DynamicIB = new FD3D11DynamicBuffer(this,D3D11_BIND_INDEX_BUFFER,DynamicIBSizes);
+
+	for (int32 Frequency = 0; Frequency < SF_NumFrequencies; ++Frequency)
+	{
+		DirtyUniformBuffers[Frequency] = 0;
+	}
 }
 
 void FD3D11DynamicRHI::Shutdown()
@@ -410,6 +423,15 @@ void FD3D11DynamicRHI::CleanupD3DDevice()
 		ReleasePooledUniformBuffers();
 		ReleasePooledTextures();
 
+		// Release references to bound uniform buffers.
+		for (int32 Frequency = 0; Frequency < SF_NumFrequencies; ++Frequency)
+		{
+			for (int32 BindIndex = 0; BindIndex < MAX_UNIFORM_BUFFERS_PER_SHADER_STAGE; ++BindIndex)
+			{
+				BoundUniformBuffers[Frequency][BindIndex].SafeRelease();
+			}
+		}
+
 		// Release the device and its IC
 		StateCache.SetContext(nullptr);
 		Direct3DDeviceIMContext = NULL;
diff --git a/Engine/Source/Runtime/Windows/D3D11RHI/Private/D3D11RHI.cpp b/Engine/Source/Runtime/Windows/D3D11RHI/Private/D3D11RHI.cpp
--- a/Engine/Source/Runtime/Windows/D3D11RHI/Private/D3D11RHI.cpp
+++ b/Engine/Source/Runtime/Windows/D3D11RHI/Private/D3D11RHI.cpp
@@ -5,6 +5,7 @@
 =============================================================================*/
 
 #include "D3D11RHIPrivate.h"
+#include "RHIStaticStates.h"
 
 #ifndef WITH_DX_PERF
 	#define WITH_DX_PERF	1
@@ -421,10 +422,26 @@ float FD3D11EventNode::GetTiming()
 }
 
 void FD3D11DynamicRHI::RHIBeginScene()
-{}
+{
+	// Increment the frame counter. INDEX_NONE is a special value meaning "uninitialized", so if
+	// we hit it just wrap around to zero.
+	SceneFrameCounter++;
+	if (SceneFrameCounter == INDEX_NONE)
+	{
+		SceneFrameCounter++;
+	}
+
+	static auto* ResourceTableCachingCvar = IConsoleManager::Get().FindTConsoleVariableDataInt(TEXT("rhi.ResourceTableCaching"));
+	if (ResourceTableCachingCvar == NULL || ResourceTableCachingCvar->GetValueOnAnyThread() == 1)
+	{
+		ResourceTableFrameCounter = SceneFrameCounter;
+	}
+}
 
 void FD3D11DynamicRHI::RHIEndScene()
-{}
+{
+	ResourceTableFrameCounter = INDEX_NONE;
+}
 
 void FD3DGPUProfiler::PushEvent(const TCHAR* Name)
 {
diff --git a/Engine/Source/Runtime/Windows/D3D11RHI/Private/D3D11RHIPrivate.h b/Engine/Source/Runtime/Windows/D3D11RHI/Private/D3D11RHIPrivate.h
--- a/Engine/Source/Runtime/Windows/D3D11RHI/Private/D3D11RHIPrivate.h
+++ b/Engine/Source/Runtime/Windows/D3D11RHI/Private/D3D11RHIPrivate.h
@@ -58,6 +58,13 @@ DECLARE_DWORD_COUNTER_STAT_EXTERN(TEXT("Textures Allocated"),STAT_D3D11TexturesA
 DECLARE_DWORD_COUNTER_STAT_EXTERN(TEXT("Textures Released"),STAT_D3D11TexturesReleased,STATGROUP_D3D11RHI, );
 DECLARE_MEMORY_STAT_EXTERN(TEXT("Texture object pool memory"),STAT_D3D11TexturePoolMemory,STATGROUP_D3D11RHI, );
 
+DECLARE_CYCLE_STAT_EXTERN(TEXT("Commit resource tables"),STAT_D3D11CommitResourceTables,STATGROUP_D3D11RHI, );
+DECLARE_CYCLE_STAT_EXTERN(TEXT("Cache resource tables"),STAT_D3D11CacheResourceTables,STATGROUP_D3D11RHI, );
+DECLARE_DWORD_COUNTER_STAT_EXTERN(TEXT("Num cached resource tables"),STAT_D3D11CacheResourceTableCalls,STATGROUP_D3D11RHI, );
+DECLARE_DWORD_COUNTER_STAT_EXTERN(TEXT("Num textures in tables"),STAT_D3D11SetTextureInTableCalls,STATGROUP_D3D11RHI, );
+DECLARE_CYCLE_STAT_EXTERN(TEXT("SetShaderTexture time"),STAT_D3D11SetShaderTextureTime,STATGROUP_D3D11RHI, );
+DECLARE_DWORD_COUNTER_STAT_EXTERN(TEXT("SetShaderTexture calls"),STAT_D3D11SetShaderTextureCalls,STATGROUP_D3D11RHI, );
+
 // This class has multiple inheritance but really FGPUTiming is a static class
 class FD3D11BufferedGPUTiming : public FRenderResource, public FGPUTiming
 {
@@ -440,6 +447,32 @@ protected:
 	uint32 NumSimultaneousRenderTargets;
 	uint32 NumUAVs;
 
+	friend class FD3D11UniformBuffer;
+	uint32 CommitResourceTableCycles;
+	uint32 CacheResourceTableCalls;
+	uint32 CacheResourceTableCycles;
+	uint32 SetShaderTextureCycles;
+	uint32 SetShaderTextureCalls;
+	uint32 SetTextureInTableCalls;
+	
+	/** Internal frame counter, incremented on each call to RHIBeginScene. */
+	uint32 SceneFrameCounter;
+
+	/**
+	 * Internal counter used for resource table caching.
+	 * INDEX_NONE means caching is not allowed.
+	 */
+	uint32 ResourceTableFrameCounter;
+
+	/** D3D11 defines a maximum of 14 constant buffers per shader stage. */
+	enum { MAX_UNIFORM_BUFFERS_PER_SHADER_STAGE = 14 };
+
+	/** Track the currently bound uniform buffers. */
+	FUniformBufferRHIRef BoundUniformBuffers[SF_NumFrequencies][MAX_UNIFORM_BUFFERS_PER_SHADER_STAGE];
+
+	/** Bit array to track which uniform buffers have changed since the last draw call. */
+	uint16 DirtyUniformBuffers[SF_NumFrequencies];
+
 	/** Tracks the current depth stencil access type. */
 	EDepthStencilAccessType CurrentDSVAccessType;
 
@@ -496,6 +529,10 @@ protected:
 	/** needs to be called before each dispatch call */
 	virtual void CommitComputeShaderConstants();
 
+	template <class ShaderType> void SetResourcesFromTables(const ShaderType* RESTRICT);
+	void CommitGraphicsResourceTables();
+	void CommitComputeResourceTables(FD3D11ComputeShader* ComputeShader);
+
 	/** 
 	 * Gets the best supported MSAA settings from the provided MSAA count to check against. 
 	 * 
diff --git a/Engine/Source/Runtime/Windows/D3D11RHI/Private/D3D11RenderTarget.cpp b/Engine/Source/Runtime/Windows/D3D11RHI/Private/D3D11RenderTarget.cpp
--- a/Engine/Source/Runtime/Windows/D3D11RHI/Private/D3D11RenderTarget.cpp
+++ b/Engine/Source/Runtime/Windows/D3D11RHI/Private/D3D11RenderTarget.cpp
@@ -28,7 +28,7 @@ namespace
 		}
 		FResolveDepthPS() {}
 
-		void SetParameters(ID3D11DeviceContext* Direct3DDeviceContext,FParameter)
+		void SetParameters(FRHICommandList* RHICmdList, ID3D11DeviceContext* Direct3DDeviceContext,FParameter)
 		{
 		}
 
@@ -60,7 +60,7 @@ namespace
 		}
 		FResolveDepthNonMSPS() {}
 
-		void SetParameters(ID3D11DeviceContext* Direct3DDeviceContext,FParameter)
+		void SetParameters(FRHICommandList* RHICmdList, ID3D11DeviceContext* Direct3DDeviceContext,FParameter)
 		{
 		}
 
@@ -92,9 +92,9 @@ namespace
 		}
 		FResolveSingleSamplePS() {}
 
-		void SetParameters(ID3D11DeviceContext* Direct3DDeviceContext,uint32 SingleSampleIndexValue)
+		void SetParameters(FRHICommandList* RHICmdList, ID3D11DeviceContext* Direct3DDeviceContext,uint32 SingleSampleIndexValue)
 		{
-			SetShaderValue(GetPixelShader(),SingleSampleIndex,SingleSampleIndexValue);
+			SetShaderValue(RHICmdList, GetPixelShader(),SingleSampleIndex,SingleSampleIndexValue);
 		}
 
 		virtual bool Serialize(FArchive& Ar)
@@ -171,6 +171,9 @@ void FD3D11DynamicRHI::ResolveTextureUsingShader(
 	typename TPixelShader::FParameter PixelShaderParameter
 	)
 {
+	//@todo-rco: RHIPacketList
+	FRHICommandList* RHICmdList = nullptr;
+
 	// Save the current viewport so that it can be restored
 	D3D11_VIEWPORT SavedViewport;
 	uint32 NumSavedViewports = 1;
@@ -248,7 +251,7 @@ void FD3D11DynamicRHI::ResolveTextureUsingShader(
 	TShaderMapRef<TPixelShader> ResolvePixelShader(GetGlobalShaderMap());
 	SetGlobalBoundShaderState(ResolveBoundShaderState, GScreenVertexDeclaration.VertexDeclarationRHI, *ResolveVertexShader, *ResolvePixelShader);
 
-	ResolvePixelShader->SetParameters(Direct3DDeviceContext,PixelShaderParameter);
+	ResolvePixelShader->SetParameters(RHICmdList, Direct3DDeviceContext,PixelShaderParameter);
 
 	// Set the source texture.
 	const uint32 TextureIndex = ResolvePixelShader->UnresolvedSurface.GetBaseIndex();
diff --git a/Engine/Source/Runtime/Windows/D3D11RHI/Private/D3D11Shaders.cpp b/Engine/Source/Runtime/Windows/D3D11RHI/Private/D3D11Shaders.cpp
--- a/Engine/Source/Runtime/Windows/D3D11RHI/Private/D3D11Shaders.cpp
+++ b/Engine/Source/Runtime/Windows/D3D11RHI/Private/D3D11Shaders.cpp
@@ -9,51 +9,59 @@
 FVertexShaderRHIRef FD3D11DynamicRHI::RHICreateVertexShader(const TArray<uint8>& Code)
 {
 	check(Code.Num());
-	TRefCountPtr<ID3D11VertexShader> D3DShader;
-	VERIFYD3D11RESULT(Direct3DDevice->CreateVertexShader((void*)Code.GetTypedData(),Code.Num() - 1,NULL,D3DShader.GetInitReference()));
-	return new FD3D11VertexShader(D3DShader,Code);
-}
 
-FPixelShaderRHIRef FD3D11DynamicRHI::RHICreatePixelShader(const TArray<uint8>& Code)
-{
-	check(Code.Num());
-	TRefCountPtr<ID3D11PixelShader> D3DShader;
-	// bGlobalUniformBufferUsed is in the last byte, see CompileD3D11Shader
-	VERIFYD3D11RESULT(Direct3DDevice->CreatePixelShader((void*)Code.GetTypedData(),Code.Num() - 1,NULL,(ID3D11PixelShader**)D3DShader.GetInitReference()));
-	return new FD3D11PixelShader(D3DShader, Code);
-}
+	FD3D11VertexShader* Shader = new FD3D11VertexShader;
 
-FHullShaderRHIRef FD3D11DynamicRHI::RHICreateHullShader(const TArray<uint8>& Code) 
-{ 
-	check(Code.Num());
-	TRefCountPtr<ID3D11HullShader> D3DShader;
-	// bGlobalUniformBufferUsed is in the last byte, see CompileD3D11Shader
-	VERIFYD3D11RESULT(Direct3DDevice->CreateHullShader((void*)Code.GetTypedData(),Code.Num() - 1,NULL,(ID3D11HullShader**)D3DShader.GetInitReference()));
-	return new FD3D11HullShader(D3DShader, Code);
-}
+	FMemoryReader Ar( Code, true );
+	Ar << Shader->ShaderResourceTable;
+	int32 Offset = Ar.Tell();
+	const uint8* CodePtr = Code.GetTypedData() + Offset;
+	const size_t CodeSize = Code.Num() - Offset - 1;
 
-FDomainShaderRHIRef FD3D11DynamicRHI::RHICreateDomainShader(const TArray<uint8>& Code) 
-{ 
-	check(Code.Num());
-	TRefCountPtr<ID3D11DomainShader> D3DShader;
+	VERIFYD3D11RESULT( Direct3DDevice->CreateVertexShader( (void*)CodePtr, CodeSize, NULL, Shader->Resource.GetInitReference() ) );
+	
 	// bGlobalUniformBufferUsed is in the last byte, see CompileD3D11Shader
-	VERIFYD3D11RESULT(Direct3DDevice->CreateDomainShader((void*)Code.GetTypedData(),Code.Num() - 1,NULL,(ID3D11DomainShader**)D3DShader.GetInitReference()));
-	return new FD3D11DomainShader(D3DShader, Code);
+	Shader->bShaderNeedsGlobalConstantBuffer = Code[Code.Num() - 1] != 0;
+	
+	// TEMP
+	Shader->Code = Code;
+	Shader->Offset = Offset;
+
+	return Shader;
 }
 
 FGeometryShaderRHIRef FD3D11DynamicRHI::RHICreateGeometryShader(const TArray<uint8>& Code) 
 { 
 	check(Code.Num());
-	TRefCountPtr<ID3D11GeometryShader> D3DShader;
+
+	FD3D11GeometryShader* Shader = new FD3D11GeometryShader;
+
+	FMemoryReader Ar( Code, true );
+	Ar << Shader->ShaderResourceTable;
+	int32 Offset = Ar.Tell();
+	const uint8* CodePtr = Code.GetTypedData() + Offset;
+	const size_t CodeSize = Code.Num() - Offset - 1;
+
+	VERIFYD3D11RESULT( Direct3DDevice->CreateGeometryShader( (void*)CodePtr, CodeSize, NULL, Shader->Resource.GetInitReference() ) );
+	
 	// bGlobalUniformBufferUsed is in the last byte, see CompileD3D11Shader
-	VERIFYD3D11RESULT(Direct3DDevice->CreateGeometryShader((void*)Code.GetTypedData(),Code.Num() - 1,NULL,(ID3D11GeometryShader**)D3DShader.GetInitReference()));
-	return new FD3D11GeometryShader(D3DShader, Code);
+	Shader->bShaderNeedsGlobalConstantBuffer = Code[Code.Num() - 1] != 0;
+
+	return Shader;
 }
 
 FGeometryShaderRHIRef FD3D11DynamicRHI::RHICreateGeometryShaderWithStreamOutput(const TArray<uint8>& Code, const FStreamOutElementList& ElementList, uint32 NumStrides, const uint32* Strides, int32 RasterizedStream) 
 { 
 	check(Code.Num());
 
+	FD3D11GeometryShader* Shader = new FD3D11GeometryShader;
+
+	FMemoryReader Ar( Code, true );
+	Ar << Shader->ShaderResourceTable;
+	int32 Offset = Ar.Tell();
+	const uint8* CodePtr = Code.GetTypedData() + Offset;
+	const size_t CodeSize = Code.Num() - Offset - 1;
+
 	uint32 D3DRasterizedStream = RasterizedStream;
 	if (RasterizedStream == -1)
 	{
@@ -72,29 +80,102 @@ FGeometryShaderRHIRef FD3D11DynamicRHI::RHICreateGeometryShaderWithStreamOutput(
 		StreamOutEntries[EntryIndex].OutputSlot = ElementList[EntryIndex].OutputSlot;
 	}
 
-	TRefCountPtr<ID3D11GeometryShader> D3DShader;
-	VERIFYD3D11RESULT(Direct3DDevice->CreateGeometryShaderWithStreamOutput(
-		(const void*)Code.GetTypedData(),
-		// bGlobalUniformBufferUsed is in the last byte, see CompileD3D11Shader
-		Code.Num() - 1,
+	VERIFYD3D11RESULT( Direct3DDevice->CreateGeometryShaderWithStreamOutput(
+		(void*)CodePtr,
+		CodeSize,
 		StreamOutEntries,
 		ElementList.Num(),
 		Strides,
 		NumStrides,
 		D3DRasterizedStream,
 		NULL,
-		(ID3D11GeometryShader**)D3DShader.GetInitReference()));
+		Shader->Resource.GetInitReference()
+	) );
+	
+	// bGlobalUniformBufferUsed is in the last byte, see CompileD3D11Shader
+	Shader->bShaderNeedsGlobalConstantBuffer = Code[Code.Num() - 1] != 0;
+
+	return Shader;
+}
+
+FHullShaderRHIRef FD3D11DynamicRHI::RHICreateHullShader(const TArray<uint8>& Code) 
+{ 
+	check(Code.Num());
 
-	return new FD3D11GeometryShader(D3DShader, Code);
+	FD3D11HullShader* Shader = new FD3D11HullShader;
+
+	FMemoryReader Ar( Code, true );
+	Ar << Shader->ShaderResourceTable;
+	int32 Offset = Ar.Tell();
+	const uint8* CodePtr = Code.GetTypedData() + Offset;
+	const size_t CodeSize = Code.Num() - Offset - 1;
+
+	VERIFYD3D11RESULT( Direct3DDevice->CreateHullShader( (void*)CodePtr, CodeSize, NULL, Shader->Resource.GetInitReference() ) );
+	
+	// bGlobalUniformBufferUsed is in the last byte, see CompileD3D11Shader
+	Shader->bShaderNeedsGlobalConstantBuffer = Code[Code.Num() - 1] != 0;
+
+	return Shader;
+}
+
+FDomainShaderRHIRef FD3D11DynamicRHI::RHICreateDomainShader(const TArray<uint8>& Code) 
+{ 
+	check(Code.Num());
+
+	FD3D11DomainShader* Shader = new FD3D11DomainShader;
+
+	FMemoryReader Ar( Code, true );
+	Ar << Shader->ShaderResourceTable;
+	int32 Offset = Ar.Tell();
+	const uint8* CodePtr = Code.GetTypedData() + Offset;
+	const size_t CodeSize = Code.Num() - Offset - 1;
+
+	VERIFYD3D11RESULT( Direct3DDevice->CreateDomainShader( (void*)CodePtr, CodeSize, NULL, Shader->Resource.GetInitReference() ) );
+	
+	// bGlobalUniformBufferUsed is in the last byte, see CompileD3D11Shader
+	Shader->bShaderNeedsGlobalConstantBuffer = Code[Code.Num() - 1] != 0;
+
+	return Shader;
+}
+
+FPixelShaderRHIRef FD3D11DynamicRHI::RHICreatePixelShader(const TArray<uint8>& Code)
+{
+	check(Code.Num());
+
+	FD3D11PixelShader* Shader = new FD3D11PixelShader;
+
+	FMemoryReader Ar( Code, true );
+	Ar << Shader->ShaderResourceTable;
+	int32 Offset = Ar.Tell();
+	const uint8* CodePtr = Code.GetTypedData() + Offset;
+	const size_t CodeSize = Code.Num() - Offset - 1;
+
+	VERIFYD3D11RESULT( Direct3DDevice->CreatePixelShader( (void*)CodePtr, CodeSize, NULL, Shader->Resource.GetInitReference() ) );
+	
+	// bGlobalUniformBufferUsed is in the last byte, see CompileD3D11Shader
+	Shader->bShaderNeedsGlobalConstantBuffer = Code[Code.Num() - 1] != 0;
+
+	return Shader;
 }
 
 FComputeShaderRHIRef FD3D11DynamicRHI::RHICreateComputeShader(const TArray<uint8>& Code) 
 { 
 	check(Code.Num());
-	TRefCountPtr<ID3D11ComputeShader> D3DShader;
+
+	FD3D11ComputeShader* Shader = new FD3D11ComputeShader;
+
+	FMemoryReader Ar( Code, true );
+	Ar << Shader->ShaderResourceTable;
+	int32 Offset = Ar.Tell();
+	const uint8* CodePtr = Code.GetTypedData() + Offset;
+	const size_t CodeSize = Code.Num() - Offset - 1;
+
+	VERIFYD3D11RESULT( Direct3DDevice->CreateComputeShader( (void*)CodePtr, CodeSize, NULL, Shader->Resource.GetInitReference() ) );
+	
 	// bGlobalUniformBufferUsed is in the last byte, see CompileD3D11Shader
-	VERIFYD3D11RESULT(Direct3DDevice->CreateComputeShader((void*)Code.GetTypedData(),Code.Num() - 1,NULL,(ID3D11ComputeShader**)D3DShader.GetInitReference()));
-	return new FD3D11ComputeShader(D3DShader, Code);
+	Shader->bShaderNeedsGlobalConstantBuffer = Code[Code.Num() - 1] != 0;
+
+	return Shader;
 }
 
 void FD3D11DynamicRHI::RHISetMultipleViewports(uint32 Count, const FViewportBounds* Data) 
@@ -134,8 +215,8 @@ FD3D11BoundShaderState::FD3D11BoundShaderState(
 	VERIFYD3D11RESULT(Direct3DDevice->CreateInputLayout(
 		InVertexDeclaration ? InVertexDeclaration->VertexElements.GetTypedData() : &NullInputElement,
 		InVertexDeclaration ? InVertexDeclaration->VertexElements.Num() : 0,
-		InVertexShader->Code.GetTypedData(),
-		InVertexShader->Code.Num() - 1,
+		&InVertexShader->Code[ InVertexShader->Offset ],			// TEMP ugly
+		InVertexShader->Code.Num() - 1 - InVertexShader->Offset,
 		InputLayout.GetInitReference()
 		));
 
diff --git a/Engine/Source/Runtime/Windows/D3D11RHI/Private/D3D11StateCache.cpp b/Engine/Source/Runtime/Windows/D3D11RHI/Private/D3D11StateCache.cpp
--- a/Engine/Source/Runtime/Windows/D3D11RHI/Private/D3D11StateCache.cpp
+++ b/Engine/Source/Runtime/Windows/D3D11RHI/Private/D3D11StateCache.cpp
@@ -32,7 +32,7 @@ static FD3D11ToggleStateCacheExecHelper GD3D11ToggleStateCacheExecHelper;
 
 #endif	// D3D11_ALLOW_STATE_CACHE && D3D11_STATE_CACHE_RUNTIME_TOGGLE
 
-#if D3D11_STATE_CACHE_DEBUG && DO_CHECK
+#if D3D11_ALLOW_STATE_CACHE && D3D11_STATE_CACHE_DEBUG && DO_CHECK
 
 template <EShaderFrequency ShaderFrequency>
 void FD3D11StateCacheBase::VerifySamplerStates()
@@ -284,7 +284,7 @@ void FD3D11StateCacheBase::VerifyCacheState()
 		VerifyShaderResourceViews<SF_Compute>();
 	}
 }
-#endif	// D3D11_STATE_CACHE_DEBUG && DO_CHECK
+#endif	// D3D11_ALLOW_STATE_CACHE && D3D11_STATE_CACHE_DEBUG && DO_CHECK
 
 void FD3D11StateCacheBase::ClearState()
 {
diff --git a/Engine/Source/Runtime/Windows/D3D11RHI/Private/D3D11StateCachePrivate.h b/Engine/Source/Runtime/Windows/D3D11RHI/Private/D3D11StateCachePrivate.h
--- a/Engine/Source/Runtime/Windows/D3D11RHI/Private/D3D11StateCachePrivate.h
+++ b/Engine/Source/Runtime/Windows/D3D11RHI/Private/D3D11StateCachePrivate.h
@@ -33,7 +33,7 @@
 //#define D3D11_STATE_CACHE_DEBUG 1
 
 
-#if D3D11_STATE_CACHE_DEBUG && DO_CHECK
+#if D3D11_ALLOW_STATE_CACHE && D3D11_STATE_CACHE_DEBUG && DO_CHECK
 	#define D3D11_STATE_CACHE_VERIFY(...) VerifyCacheState()
 	#define D3D11_STATE_CACHE_VERIFY_PRE(...) VerifyCacheStatePre()
 	#define D3D11_STATE_CACHE_VERIFY_POST(...) VerifyCacheStatePost()
@@ -879,7 +879,7 @@ public:
 	 */
 	virtual void ClearState();
 
-#if D3D11_STATE_CACHE_DEBUG
+#if D3D11_ALLOW_STATE_CACHE && D3D11_STATE_CACHE_DEBUG
 protected:
 	// Debug helper methods to verify cached state integrity.
 	template <EShaderFrequency ShaderFrequency>
diff --git a/Engine/Source/Runtime/Windows/D3D11RHI/Private/D3D11Texture.cpp b/Engine/Source/Runtime/Windows/D3D11RHI/Private/D3D11Texture.cpp
--- a/Engine/Source/Runtime/Windows/D3D11RHI/Private/D3D11Texture.cpp
+++ b/Engine/Source/Runtime/Windows/D3D11RHI/Private/D3D11Texture.cpp
@@ -1498,3 +1498,30 @@ void FD3D11DynamicRHI::RHIVirtualTextureSetFirstMipInMemory(FTexture2DRHIParamRe
 void FD3D11DynamicRHI::RHIVirtualTextureSetFirstMipVisible(FTexture2DRHIParamRef TextureRHI, uint32 FirstMip)
 {
 }
+
+FTextureReferenceRHIRef FD3D11DynamicRHI::RHICreateTextureReference()
+{
+	return new FD3D11TextureReference(this);
+}
+
+void FD3D11DynamicRHI::RHIUpdateTextureReference(FTextureReferenceRHIParamRef TextureRefRHI, FTextureRHIParamRef NewTextureRHI)
+{
+	// Updating texture references is disallowed while the RHI could be caching them in referenced resource tables.
+	check(ResourceTableFrameCounter == INDEX_NONE);
+
+	FD3D11TextureReference* TextureRef = (FD3D11TextureReference*)TextureRefRHI;
+	if (TextureRef)
+	{
+		FD3D11TextureBase* NewTexture = NULL;
+		ID3D11ShaderResourceView* NewSRV = NULL;
+		if (NewTextureRHI)
+		{
+			NewTexture = GetD3D11TextureFromRHITexture(NewTextureRHI);
+			if (NewTexture)
+			{
+				NewSRV = NewTexture->GetShaderResourceView();
+			}
+		}
+		TextureRef->SetReferencedTexture(NewTextureRHI,NewTexture,NewSRV);
+	}
+}
