UE4 - merge first pass at parallel rendering to main

diff --git a/Engine/Source/Runtime/Windows/D3D11RHI/Private/D3D11UniformBuffer.cpp b/Engine/Source/Runtime/Windows/D3D11RHI/Private/D3D11UniformBuffer.cpp
--- a/Engine/Source/Runtime/Windows/D3D11RHI/Private/D3D11UniformBuffer.cpp
+++ b/Engine/Source/Runtime/Windows/D3D11RHI/Private/D3D11UniformBuffer.cpp
@@ -92,112 +92,134 @@ static bool IsPoolingEnabled()
 	return CVarValue != 0;
 };
 
-FUniformBufferRHIRef FD3D11DynamicRHI::RHICreateUniformBuffer(const void* Contents,uint32 NumBytes,EUniformBufferUsage Usage)
+FUniformBufferRHIRef FD3D11DynamicRHI::RHICreateUniformBuffer(const void* Contents,const FRHIUniformBufferLayout& Layout,EUniformBufferUsage Usage)
 {
 	check(IsInRenderingThread());
 
-	// Explicitly check that the size is nonzero before allowing CreateBuffer to opaquely fail.
-	check(NumBytes > 0);
-
-	// Constant buffers must also be 16-byte aligned.
-	check(Align(NumBytes,16) == NumBytes);
-	check(Align(Contents,16) == Contents);
-	check(NumBytes <= D3D11_REQ_CONSTANT_BUFFER_ELEMENT_COUNT * 16);
-	check(NumBytes < (1 << NumPoolBuckets));
-
-	SCOPE_CYCLE_COUNTER(STAT_D3D11UpdateUniformBufferTime);
-
-	// @todo: Only pool single-use uniform buffers?
-	if(IsPoolingEnabled())
+	FD3D11UniformBuffer* NewUniformBuffer = nullptr;
+	const uint32 NumBytes = Layout.ConstantBufferSize;
+	if (NumBytes > 0)
 	{
-		TRefCountPtr<ID3D11Buffer> UniformBufferResource;
-		FRingAllocation RingAllocation;
-
-#if PLATFORM_XBOXONE
-		if (Usage == UniformBuffer_SingleUse)
+		// Constant buffers must also be 16-byte aligned.
+		check(Align(NumBytes,16) == NumBytes);
+		check(Align(Contents,16) == Contents);
+		check(NumBytes <= D3D11_REQ_CONSTANT_BUFFER_ELEMENT_COUNT * 16);
+		check(NumBytes < (1 << NumPoolBuckets));
+
+		SCOPE_CYCLE_COUNTER(STAT_D3D11UpdateUniformBufferTime);
+		if (IsPoolingEnabled())
 		{
-			RingAllocation = DynamicCB.MapAndFill(GetDeviceContext(),NumBytes,Contents);
-			check(RingAllocation.IsValid());
-		}
-#endif
+			TRefCountPtr<ID3D11Buffer> UniformBufferResource;
+			FRingAllocation RingAllocation;
 
-		if (!RingAllocation.IsValid())
-		{
-			// Find the appropriate bucket based on size
-			const uint32 BucketIndex = GetPoolBucketIndex(NumBytes);
-			TArray<FPooledUniformBuffer>& PoolBucket = UniformBufferPool[BucketIndex];
-
-			if (PoolBucket.Num() > 0)
+#if PLATFORM_XBOXONE
+			if (Usage == UniformBuffer_SingleFrame)
 			{
-				// Reuse the last entry in this size bucket
-				FPooledUniformBuffer FreeBufferEntry = PoolBucket.Pop();
-				check(IsValidRef(FreeBufferEntry.Buffer));
-				UniformBufferResource = FreeBufferEntry.Buffer;
-				checkf(FreeBufferEntry.CreatedSize >= NumBytes, TEXT("%u %u %u %u"), NumBytes, BucketIndex, FreeBufferEntry.CreatedSize, GetPoolBucketSize(NumBytes));
-				DEC_DWORD_STAT(STAT_D3D11NumFreeUniformBuffers);
-				DEC_MEMORY_STAT_BY(STAT_D3D11FreeUniformBufferMemory, FreeBufferEntry.CreatedSize);
+				RingAllocation = DynamicCB.MapAndFill(GetDeviceContext(),NumBytes,Contents);
+				check(RingAllocation.IsValid());
 			}
+#endif
 
-			// Nothing usable was found in the free pool, create a new uniform buffer
-			if (!IsValidRef(UniformBufferResource))
+			if (!RingAllocation.IsValid())
 			{
-				D3D11_BUFFER_DESC Desc;
-				// Allocate based on the bucket size, since this uniform buffer will be reused later
-				Desc.ByteWidth = GetPoolBucketSize(NumBytes);
-				// Use D3D11_USAGE_DYNAMIC, which allows multiple CPU writes for pool reuses
-				// This is method of updating is vastly superior to creating a new constant buffer each time with D3D11_USAGE_IMMUTABLE, 
-				// Since that inserts the data into the command buffer which causes GPU flushes
-				Desc.Usage = D3D11_USAGE_DYNAMIC;
-				Desc.BindFlags = D3D11_BIND_CONSTANT_BUFFER;
-				Desc.CPUAccessFlags = D3D11_CPU_ACCESS_WRITE;
-				Desc.MiscFlags = 0;
-				Desc.StructureByteStride = 0;
-
-				VERIFYD3D11RESULT(Direct3DDevice->CreateBuffer(&Desc, NULL, UniformBufferResource.GetInitReference()));
-
-				UpdateBufferStats(UniformBufferResource, true);
+				// Find the appropriate bucket based on size
+				const uint32 BucketIndex = GetPoolBucketIndex(NumBytes);
+				TArray<FPooledUniformBuffer>& PoolBucket = UniformBufferPool[BucketIndex];
+
+				if (PoolBucket.Num() > 0)
+				{
+					// Reuse the last entry in this size bucket
+					FPooledUniformBuffer FreeBufferEntry = PoolBucket.Pop();
+					check(IsValidRef(FreeBufferEntry.Buffer));
+					UniformBufferResource = FreeBufferEntry.Buffer;
+					checkf(FreeBufferEntry.CreatedSize >= NumBytes, TEXT("%u %u %u %u"), NumBytes, BucketIndex, FreeBufferEntry.CreatedSize, GetPoolBucketSize(NumBytes));
+					DEC_DWORD_STAT(STAT_D3D11NumFreeUniformBuffers);
+					DEC_MEMORY_STAT_BY(STAT_D3D11FreeUniformBufferMemory, FreeBufferEntry.CreatedSize);
+				}
+
+				// Nothing usable was found in the free pool, create a new uniform buffer
+				if (!IsValidRef(UniformBufferResource))
+				{
+					D3D11_BUFFER_DESC Desc;
+					// Allocate based on the bucket size, since this uniform buffer will be reused later
+					Desc.ByteWidth = GetPoolBucketSize(NumBytes);
+					// Use D3D11_USAGE_DYNAMIC, which allows multiple CPU writes for pool reuses
+					// This is method of updating is vastly superior to creating a new constant buffer each time with D3D11_USAGE_IMMUTABLE, 
+					// Since that inserts the data into the command buffer which causes GPU flushes
+					Desc.Usage = D3D11_USAGE_DYNAMIC;
+					Desc.BindFlags = D3D11_BIND_CONSTANT_BUFFER;
+					Desc.CPUAccessFlags = D3D11_CPU_ACCESS_WRITE;
+					Desc.MiscFlags = 0;
+					Desc.StructureByteStride = 0;
+
+					VERIFYD3D11RESULT(Direct3DDevice->CreateBuffer(&Desc, NULL, UniformBufferResource.GetInitReference()));
+
+					UpdateBufferStats(UniformBufferResource, true);
+				}
+
+				check(IsValidRef(UniformBufferResource));
+
+				D3D11_MAPPED_SUBRESOURCE MappedSubresource;
+				// Discard previous results since we always do a full update
+				VERIFYD3D11RESULT(Direct3DDeviceIMContext->Map(UniformBufferResource, 0, D3D11_MAP_WRITE_DISCARD, 0, &MappedSubresource));
+				check(MappedSubresource.RowPitch >= NumBytes);
+				FMemory::Memcpy(MappedSubresource.pData, Contents, NumBytes);
+				Direct3DDeviceIMContext->Unmap(UniformBufferResource, 0);
 			}
 
-			check(IsValidRef(UniformBufferResource));
-
-			D3D11_MAPPED_SUBRESOURCE MappedSubresource;
-			// Discard previous results since we always do a full update
-			VERIFYD3D11RESULT(Direct3DDeviceIMContext->Map(UniformBufferResource, 0, D3D11_MAP_WRITE_DISCARD, 0, &MappedSubresource));
-			check(MappedSubresource.RowPitch >= NumBytes);
-			FMemory::Memcpy(MappedSubresource.pData, Contents, NumBytes);
-			Direct3DDeviceIMContext->Unmap(UniformBufferResource, 0);
+			NewUniformBuffer = new FD3D11UniformBuffer(this, Layout, UniformBufferResource, RingAllocation);
+		}
+		else
+		{
+			// No pooling
+			D3D11_BUFFER_DESC Desc;
+			Desc.ByteWidth = NumBytes;
+			Desc.Usage = D3D11_USAGE_IMMUTABLE;
+			Desc.BindFlags = D3D11_BIND_CONSTANT_BUFFER;
+			Desc.CPUAccessFlags = 0;
+			Desc.MiscFlags = 0;
+			Desc.StructureByteStride = 0;
+
+			D3D11_SUBRESOURCE_DATA ImmutableData;
+			ImmutableData.pSysMem = Contents;
+			ImmutableData.SysMemPitch = ImmutableData.SysMemSlicePitch = 0;
+
+			TRefCountPtr<ID3D11Buffer> UniformBufferResource;
+			VERIFYD3D11RESULT(Direct3DDevice->CreateBuffer(&Desc,&ImmutableData,UniformBufferResource.GetInitReference()));
+
+			FRingAllocation RingAllocation;
+			NewUniformBuffer = new FD3D11UniformBuffer(this, Layout, UniformBufferResource, RingAllocation);
 		}
-
-		return new FD3D11UniformBuffer(NumBytes,UniformBufferResource,RingAllocation);
 	}
 	else
 	{
-		D3D11_BUFFER_DESC Desc;
-		Desc.ByteWidth = NumBytes;
-		Desc.Usage = D3D11_USAGE_IMMUTABLE;
-		Desc.BindFlags = D3D11_BIND_CONSTANT_BUFFER;
-		Desc.CPUAccessFlags = 0;
-		Desc.MiscFlags = 0;
-		Desc.StructureByteStride = 0;
-
-		D3D11_SUBRESOURCE_DATA ImmutableData;
-		ImmutableData.pSysMem = Contents;
-		ImmutableData.SysMemPitch = ImmutableData.SysMemSlicePitch = 0;
-
-		TRefCountPtr<ID3D11Buffer> UniformBufferResource;
-		VERIFYD3D11RESULT(Direct3DDevice->CreateBuffer(&Desc,&ImmutableData,UniformBufferResource.GetInitReference()));
-
+		// This uniform buffer contains no constants, only a resource table.
 		FRingAllocation RingAllocation;
-		return new FD3D11UniformBuffer(NumBytes,UniformBufferResource,RingAllocation);
+		NewUniformBuffer = new FD3D11UniformBuffer(this, Layout, nullptr, RingAllocation);
 	}
+
+	if (Layout.Resources.Num())
+	{
+		int32 NumResources = Layout.Resources.Num();
+		FRHIResource** InResources = (FRHIResource**)((uint8*)Contents + Layout.ResourceOffset);
+		NewUniformBuffer->ResourceTable.Empty(NumResources);
+		NewUniformBuffer->ResourceTable.AddZeroed(NumResources);
+		for (int32 i = 0; i < NumResources; ++i)
+		{
+			check(InResources[i]);
+			NewUniformBuffer->ResourceTable[i] = InResources[i];
+		}
+		NewUniformBuffer->RawResourceTable.Empty(NumResources);
+		NewUniformBuffer->RawResourceTable.AddZeroed(NumResources);
+	}
+
+	return NewUniformBuffer;
 }
 
 FD3D11UniformBuffer::~FD3D11UniformBuffer()
 {
-	check(IsInRenderingThread());
-
 	// Do not return the allocation to the pool if it is in the dynamic constant buffer!
-	if (!RingAllocation.IsValid())
+	if (!RingAllocation.IsValid() && Resource != nullptr)
 	{
 		D3D11_BUFFER_DESC Desc;
 		Resource->GetDesc(&Desc);
@@ -222,6 +244,64 @@ FD3D11UniformBuffer::~FD3D11UniformBuffer()
 	}
 }
 
+void FD3D11UniformBuffer::CacheResourcesInternal()
+{
+	uint32 Start = FPlatformTime::Cycles();
+
+	const FRHIUniformBufferLayout& Layout = GetLayout();
+	int32 NumResources = Layout.Resources.Num();
+	const uint8* RESTRICT ResourceTypes = Layout.Resources.GetData();
+	const TRefCountPtr<FRHIResource>* RESTRICT Resources = ResourceTable.GetData();
+	FResourcePair* RESTRICT RawResources = RawResourceTable.GetData();
+	float CurrentTime = FApp::GetCurrentTime();
+
+	// todo: Immutable resources, i.e. not textures, can be safely cached across frames.
+	// Texture streaming makes textures complicated :)
+	for (int32 i = 0; i < NumResources; ++i)
+	{
+		switch (ResourceTypes[i])
+		{
+		case UBMT_SRV:
+			{
+				FD3D11ShaderResourceView* ShaderResourceViewRHI = (FD3D11ShaderResourceView*)Resources[i].GetReference();
+				RawResources[i].ShaderResource = ShaderResourceViewRHI->Resource.GetReference();
+				RawResources[i].D3D11Resource = ShaderResourceViewRHI->View.GetReference();
+			}
+			break;
+
+		case UBMT_TEXTURE:
+			{
+				// todo: this does multiple virtual function calls to find the right type to cast to
+				// this is due to multiple inheritance nastiness, NEEDS CLEANUP
+				FRHITexture* TextureRHI = (FRHITexture*)Resources[i].GetReference();
+				TextureRHI->SetLastCachedTime(CurrentTime);
+				FD3D11TextureBase* TextureD3D11 = GetD3D11TextureFromRHITexture(TextureRHI);
+				RawResources[i].ShaderResource = TextureD3D11->GetBaseShaderResource();
+				RawResources[i].D3D11Resource = TextureD3D11->GetShaderResourceView();
+			}
+			break;
+
+		case UBMT_UAV:
+			RawResources[i].ShaderResource = NULL;
+			RawResources[i].D3D11Resource = NULL;
+			check(0);
+			break;
+
+		case UBMT_SAMPLER:
+			RawResources[i].ShaderResource = NULL;
+			RawResources[i].D3D11Resource = ((FD3D11SamplerState*)Resources[i].GetReference())->Resource.GetReference();
+			break;
+
+		default:
+			check(0);
+			break;
+		}
+	}
+
+	D3D11RHI->CacheResourceTableCycles += (FPlatformTime::Cycles() - Start);
+	D3D11RHI->CacheResourceTableCalls++;
+}
+
 void FD3D11DynamicRHI::ReleasePooledUniformBuffers()
 {
 	// Free D3D resources in the pool
diff --git a/Engine/Source/Runtime/Windows/D3D11RHI/Private/D3D11Util.cpp b/Engine/Source/Runtime/Windows/D3D11RHI/Private/D3D11Util.cpp
--- a/Engine/Source/Runtime/Windows/D3D11RHI/Private/D3D11Util.cpp
+++ b/Engine/Source/Runtime/Windows/D3D11RHI/Private/D3D11Util.cpp
@@ -382,4 +382,11 @@ DEFINE_STAT(STAT_D3D11TexturePoolMemory);
 DEFINE_STAT(STAT_D3D11FreeUniformBufferMemory);
 DEFINE_STAT(STAT_D3D11NumFreeUniformBuffers);
 
+DEFINE_STAT(STAT_D3D11CommitResourceTables);
+DEFINE_STAT(STAT_D3D11CacheResourceTables);
+DEFINE_STAT(STAT_D3D11CacheResourceTableCalls);
+DEFINE_STAT(STAT_D3D11SetShaderTextureTime);
+DEFINE_STAT(STAT_D3D11SetShaderTextureCalls);
+DEFINE_STAT(STAT_D3D11SetTextureInTableCalls);
+
 #undef LOCTEXT_NAMESPACE
diff --git a/Engine/Source/Runtime/Windows/D3D11RHI/Private/D3D11Viewport.cpp b/Engine/Source/Runtime/Windows/D3D11RHI/Private/D3D11Viewport.cpp
--- a/Engine/Source/Runtime/Windows/D3D11RHI/Private/D3D11Viewport.cpp
+++ b/Engine/Source/Runtime/Windows/D3D11RHI/Private/D3D11Viewport.cpp
@@ -468,7 +468,7 @@ void FD3D11DynamicRHI::RHIBeginDrawingViewport(FViewportRHIParamRef ViewportRHI,
 	DrawingViewport = Viewport;
 
 	// Set the render target and viewport.
-	if( !IsValidRef( RenderTarget ) )
+	if( RenderTarget == NULL )
 	{
 		RenderTarget = Viewport->GetBackBuffer();
 	}
@@ -516,6 +516,20 @@ void FD3D11DynamicRHI::RHIEndDrawingViewport(FViewportRHIParamRef ViewportRHI,bo
 	StateCache.SetGeometryShader(nullptr);
 	// Compute Shader is set to NULL after each Dispatch call, so no need to clear it here
 
+	FThreadStats::AddMessage(GET_STATFNAME(STAT_D3D11CommitResourceTables), EStatOperation::Set, int64(CommitResourceTableCycles), true);
+	FThreadStats::AddMessage(GET_STATFNAME(STAT_D3D11CacheResourceTables), EStatOperation::Set, int64(CacheResourceTableCycles), true);
+	FThreadStats::AddMessage(GET_STATFNAME(STAT_D3D11SetShaderTextureTime), EStatOperation::Set, int64(SetShaderTextureCycles), true);
+	INC_DWORD_STAT_BY(STAT_D3D11SetShaderTextureCalls, SetShaderTextureCalls);
+	INC_DWORD_STAT_BY(STAT_D3D11CacheResourceTableCalls, CacheResourceTableCalls);
+	INC_DWORD_STAT_BY(STAT_D3D11SetTextureInTableCalls, SetTextureInTableCalls);
+
+	CommitResourceTableCycles = 0;
+	CacheResourceTableCalls = 0;
+	CacheResourceTableCycles = 0;
+	SetShaderTextureCycles = 0;
+	SetShaderTextureCalls = 0;
+	SetTextureInTableCalls = 0;
+
 	if(bPresent)
 	{
 		Viewport->Present(bLockToVsync);
diff --git a/Engine/Source/Runtime/Windows/D3D11RHI/Public/D3D11Resources.h b/Engine/Source/Runtime/Windows/D3D11RHI/Public/D3D11Resources.h
--- a/Engine/Source/Runtime/Windows/D3D11RHI/Public/D3D11Resources.h
+++ b/Engine/Source/Runtime/Windows/D3D11RHI/Public/D3D11Resources.h
@@ -7,37 +7,7 @@
 #pragma once
 
 #include "BoundShaderStateCache.h"
-
-/**
- * Combined shader state and vertex definition for rendering geometry. 
- * Each unique instance consists of a vertex decl, vertex shader, and pixel shader.
- */
-class FD3D11BoundShaderState : public FRHIBoundShaderState
-{
-public:
-
-	FCachedBoundShaderStateLink CacheLink;
-
-	TRefCountPtr<ID3D11InputLayout> InputLayout;
-	TRefCountPtr<ID3D11VertexShader> VertexShader;
-	TRefCountPtr<ID3D11PixelShader> PixelShader;
-	TRefCountPtr<ID3D11HullShader> HullShader;
-	TRefCountPtr<ID3D11DomainShader> DomainShader;
-	TRefCountPtr<ID3D11GeometryShader> GeometryShader;
-
-	bool bShaderNeedsGlobalConstantBuffer[SF_NumFrequencies];
-
-	/** Initialization constructor. */
-	FD3D11BoundShaderState(
-		FVertexDeclarationRHIParamRef InVertexDeclarationRHI,
-		FVertexShaderRHIParamRef InVertexShaderRHI,
-		FPixelShaderRHIParamRef InPixelShaderRHI,
-		FHullShaderRHIParamRef InHullShaderRHI,
-		FDomainShaderRHIParamRef InDomainShaderRHI,
-		FGeometryShaderRHIParamRef InGeometryShaderRHI,
-		ID3D11Device* Direct3DDevice
-		);
-};
+#include "D3D11ShaderResources.h"
 
 template <>
 struct TTypeTraits<D3D11_INPUT_ELEMENT_DESC> : public TTypeTraitsBase<D3D11_INPUT_ELEMENT_DESC>
@@ -66,113 +36,126 @@ public:
 class FD3D11VertexShader : public FRHIVertexShader
 {
 public:
+	enum { StaticFrequency = SF_Vertex };
 
 	/** The vertex shader resource. */
 	TRefCountPtr<ID3D11VertexShader> Resource;
 
+	FD3D11ShaderResourceTable ShaderResourceTable;
+
 	/** The vertex shader's bytecode, with custom data in the last byte. */
 	TArray<uint8> Code;
 
-	bool bShaderNeedsGlobalConstantBuffer;
+	// TEMP remove with removal of bound shader state
+	int32 Offset;
 
-	/** Initialization constructor. */
-	FD3D11VertexShader(ID3D11VertexShader* InResource,const TArray<uint8>& InCode):
-		Resource(InResource),
-		Code(InCode)
-	{
-		// bGlobalUniformBufferUsed is in the last byte, see CompileD3D11Shader
-		bShaderNeedsGlobalConstantBuffer = InCode[InCode.Num() - 1] != 0;
-	}
+	bool bShaderNeedsGlobalConstantBuffer;
 };
 
 class FD3D11GeometryShader : public FRHIGeometryShader
 {
 public:
+	enum { StaticFrequency = SF_Geometry };
 
 	/** The shader resource. */
 	TRefCountPtr<ID3D11GeometryShader> Resource;
 
-	bool bShaderNeedsGlobalConstantBuffer;
+	FD3D11ShaderResourceTable ShaderResourceTable;
 
-	/** Initialization constructor. */
-	FD3D11GeometryShader(ID3D11GeometryShader* InResource,const TArray<uint8>& InCode):
-		Resource(InResource)
-	{
-		// bGlobalUniformBufferUsed is in the last byte, see CompileD3D11Shader
-		bShaderNeedsGlobalConstantBuffer = InCode[InCode.Num() - 1] != 0;
-	}
+	bool bShaderNeedsGlobalConstantBuffer;
 };
 
 class FD3D11HullShader : public FRHIHullShader
 {
 public:
+	enum { StaticFrequency = SF_Hull };
 
 	/** The shader resource. */
 	TRefCountPtr<ID3D11HullShader> Resource;
 
-	bool bShaderNeedsGlobalConstantBuffer;
+	FD3D11ShaderResourceTable ShaderResourceTable;
 
-	/** Initialization constructor. */
-	FD3D11HullShader(ID3D11HullShader* InResource,const TArray<uint8>& InCode):
-		Resource(InResource)
-	{
-		// bGlobalUniformBufferUsed is in the last byte, see CompileD3D11Shader
-		bShaderNeedsGlobalConstantBuffer = InCode[InCode.Num() - 1] != 0;
-	}
+	bool bShaderNeedsGlobalConstantBuffer;
 };
 
 class FD3D11DomainShader : public FRHIDomainShader
 {
 public:
+	enum { StaticFrequency = SF_Domain };
 
 	/** The shader resource. */
 	TRefCountPtr<ID3D11DomainShader> Resource;
 
-	bool bShaderNeedsGlobalConstantBuffer;
+	FD3D11ShaderResourceTable ShaderResourceTable;
 
-	/** Initialization constructor. */
-	FD3D11DomainShader(ID3D11DomainShader* InResource,const TArray<uint8>& InCode):
-		Resource(InResource)
-	{
-		// bGlobalUniformBufferUsed is in the last byte, see CompileD3D11Shader
-		bShaderNeedsGlobalConstantBuffer = InCode[InCode.Num() - 1] != 0;
-	}
+	bool bShaderNeedsGlobalConstantBuffer;
 };
 
 class FD3D11PixelShader : public FRHIPixelShader
 {
 public:
+	enum { StaticFrequency = SF_Pixel };
 
 	/** The shader resource. */
 	TRefCountPtr<ID3D11PixelShader> Resource;
 
-	bool bShaderNeedsGlobalConstantBuffer;
+	FD3D11ShaderResourceTable ShaderResourceTable;
 
-	/** Initialization constructor. */
-	FD3D11PixelShader(ID3D11PixelShader* InResource,const TArray<uint8>& InCode):
-		Resource(InResource)
-	{
-		// bGlobalUniformBufferUsed is in the last byte, see CompileD3D11Shader
-		bShaderNeedsGlobalConstantBuffer = InCode[InCode.Num() - 1] != 0;
-	}
+	bool bShaderNeedsGlobalConstantBuffer;
 };
 
 class FD3D11ComputeShader : public FRHIComputeShader
 {
 public:
+	enum { StaticFrequency = SF_Compute };
 
 	/** The shader resource. */
 	TRefCountPtr<ID3D11ComputeShader> Resource;
 
+	FD3D11ShaderResourceTable ShaderResourceTable;
+
 	bool bShaderNeedsGlobalConstantBuffer;
+};
+
+/**
+ * Combined shader state and vertex definition for rendering geometry. 
+ * Each unique instance consists of a vertex decl, vertex shader, and pixel shader.
+ */
+class FD3D11BoundShaderState : public FRHIBoundShaderState
+{
+public:
+
+	FCachedBoundShaderStateLink CacheLink;
+
+	TRefCountPtr<ID3D11InputLayout> InputLayout;
+	TRefCountPtr<ID3D11VertexShader> VertexShader;
+	TRefCountPtr<ID3D11PixelShader> PixelShader;
+	TRefCountPtr<ID3D11HullShader> HullShader;
+	TRefCountPtr<ID3D11DomainShader> DomainShader;
+	TRefCountPtr<ID3D11GeometryShader> GeometryShader;
+
+	bool bShaderNeedsGlobalConstantBuffer[SF_NumFrequencies];
+
 
 	/** Initialization constructor. */
-	FD3D11ComputeShader(ID3D11ComputeShader* InResource,const TArray<uint8>& InCode):
-		Resource(InResource)
-	{
-		// bGlobalUniformBufferUsed is in the last byte, see CompileD3D11Shader
-		bShaderNeedsGlobalConstantBuffer = InCode[InCode.Num() - 1] != 0;
-	}
+	FD3D11BoundShaderState(
+		FVertexDeclarationRHIParamRef InVertexDeclarationRHI,
+		FVertexShaderRHIParamRef InVertexShaderRHI,
+		FPixelShaderRHIParamRef InPixelShaderRHI,
+		FHullShaderRHIParamRef InHullShaderRHI,
+		FDomainShaderRHIParamRef InDomainShaderRHI,
+		FGeometryShaderRHIParamRef InGeometryShaderRHI,
+		ID3D11Device* Direct3DDevice
+		);
+
+	/**
+	 * Get the shader for the given frequency.
+	 */
+	FORCEINLINE FD3D11VertexShader*   GetVertexShader() const   { return (FD3D11VertexShader*)CacheLink.GetVertexShader(); }
+	FORCEINLINE FD3D11PixelShader*    GetPixelShader() const    { return (FD3D11PixelShader*)CacheLink.GetPixelShader(); }
+	FORCEINLINE FD3D11HullShader*     GetHullShader() const     { return (FD3D11HullShader*)CacheLink.GetHullShader(); }
+	FORCEINLINE FD3D11DomainShader*   GetDomainShader() const   { return (FD3D11DomainShader*)CacheLink.GetDomainShader(); }
+	FORCEINLINE FD3D11GeometryShader* GetGeometryShader() const { return (FD3D11GeometryShader*)CacheLink.GetGeometryShader(); }
 };
 
 /** The base class of resources that may be bound as shader resources. */
@@ -196,6 +179,7 @@ public:
 		) 
 	: D3DRHI(InD3DRHI)
 	, MemorySize(0)
+	, BaseShaderResource(this)
 	, Resource(InResource)
 	, ShaderResourceView(InShaderResourceView)
 	, RTVArraySize(InRTVArraySize)
@@ -232,6 +216,7 @@ public:
 	// Accessors.
 	ID3D11Resource* GetResource() const { return Resource; }
 	ID3D11ShaderResourceView* GetShaderResourceView() const { return ShaderResourceView; }
+	FD3D11BaseShaderResource* GetBaseShaderResource() const { return BaseShaderResource; }
 
 	/** 
 	 * Get the render target view for the specified mip and array slice.
@@ -278,8 +263,9 @@ protected:
 	/** Amount of memory allocated by this texture, in bytes. */
 	int32 MemorySize;
 
-private:
-	
+	/** Pointer to the base shader resource. Usually the object itself, but not for texture references. */
+	FD3D11BaseShaderResource* BaseShaderResource;
+
 	/** The texture resource. */
 	TRefCountPtr<ID3D11Resource> Resource;
 
@@ -495,6 +481,39 @@ typedef TD3D11Texture2D<FD3D11BaseTexture2D>      FD3D11Texture2D;
 typedef TD3D11Texture2D<FD3D11BaseTexture2DArray> FD3D11Texture2DArray;
 typedef TD3D11Texture2D<FD3D11BaseTextureCube>    FD3D11TextureCube;
 
+/** Texture reference class. */
+class FD3D11TextureReference : public FRHITextureReference, public FD3D11TextureBase
+{
+public:
+	explicit FD3D11TextureReference(class FD3D11DynamicRHI* InD3DRHI)
+		: FRHITextureReference()
+		, FD3D11TextureBase(InD3DRHI,NULL,NULL, 0, false,TArray<TRefCountPtr<ID3D11RenderTargetView> >(),NULL)
+	{
+		BaseShaderResource = NULL;
+	}
+
+	void SetReferencedTexture(FRHITexture* InTexture, FD3D11BaseShaderResource* InBaseShaderResource, ID3D11ShaderResourceView* InSRV)
+	{
+		ShaderResourceView = InSRV;
+		BaseShaderResource = InBaseShaderResource;
+		FRHITextureReference::SetReferencedTexture(InTexture);
+	}
+
+	// IRefCountedObject interface.
+	virtual uint32 AddRef() const
+	{
+		return FRefCountedObject::AddRef();
+	}
+	virtual uint32 Release() const
+	{
+		return FRefCountedObject::Release();
+	}
+	virtual uint32 GetRefCount() const
+	{
+		return FRefCountedObject::GetRefCount();
+	}
+};
+
 /** Given a pointer to a RHI texture that was created by the D3D11 RHI, returns a pointer to the FD3D11TextureBase it encapsulates. */
 inline FD3D11TextureBase* GetD3D11TextureFromRHITexture(FRHITexture* Texture)
 {
@@ -506,6 +525,10 @@ inline FD3D11TextureBase* GetD3D11TextureFromRHITexture(FRHITexture* Texture)
 	{
 		return static_cast<FD3D11Texture2D*>(Texture);
 	}
+	else if(Texture->GetTextureReference())
+	{
+		return static_cast<FD3D11TextureReference*>(Texture);
+	}
 	else if(Texture->GetTexture2DArray())
 	{
 		return static_cast<FD3D11Texture2DArray*>(Texture);
@@ -580,14 +603,47 @@ public:
 	/** Allocation in the constants ring buffer if applicable. */
 	FRingAllocation RingAllocation;
 
+	/** Resource table containing RHI references. */
+	TArray<TRefCountPtr<FRHIResource> > ResourceTable;
+
+	/** Cached resources need to retain the associated shader resource for bookkeeping purposes. */
+	struct FResourcePair
+	{
+		FD3D11BaseShaderResource* ShaderResource;
+		IUnknown* D3D11Resource;
+	};
+
+	/** Raw resource table, cached once per frame. */
+	TArray<FResourcePair> RawResourceTable;
+
+	/** The frame in which RawResourceTable was last cached. */
+	uint32 LastCachedFrame;
+
 	/** Initialization constructor. */
-	FD3D11UniformBuffer(uint32 Size,ID3D11Buffer* InResource,const FRingAllocation& InRingAllocation)
-	: FRHIUniformBuffer(Size)
+	FD3D11UniformBuffer(class FD3D11DynamicRHI* InD3D11RHI, const FRHIUniformBufferLayout& InLayout, ID3D11Buffer* InResource,const FRingAllocation& InRingAllocation)
+	: FRHIUniformBuffer(InLayout)
 	, Resource(InResource)
 	, RingAllocation(InRingAllocation)
+	, D3D11RHI(InD3D11RHI)
+	, LastCachedFrame((uint32)-1)
 	{}
 
 	virtual ~FD3D11UniformBuffer();
+
+	/** Cache resources if needed. */
+	inline void CacheResources(uint32 InFrameCounter)
+	{
+		if (InFrameCounter == INDEX_NONE || LastCachedFrame != InFrameCounter)
+		{
+			CacheResourcesInternal();
+			LastCachedFrame = InFrameCounter;
+		}
+	}
+
+private:
+	class FD3D11DynamicRHI* D3D11RHI;
+	/** Actually cache resources. */
+	void CacheResourcesInternal();
 };
 
 /** Index buffer resource class that stores stride information. */
diff --git a/Engine/Source/Runtime/Windows/D3D11RHI/Public/D3D11ShaderResources.h b/Engine/Source/Runtime/Windows/D3D11RHI/Public/D3D11ShaderResources.h
new file mode 100644
--- /dev/null
+++ b/Engine/Source/Runtime/Windows/D3D11RHI/Public/D3D11ShaderResources.h
@@ -0,0 +1,11 @@
+// Copyright 1998-2014 Epic Games, Inc. All Rights Reserved.
+
+/*=============================================================================
+	D3D11Resources.h: D3D resource RHI definitions.
+=============================================================================*/
+
+#pragma once
+
+struct FD3D11ShaderResourceTable : public FBaseShaderResourceTable
+{
+};
