UE4 - fix shipping builds and stat use in the d3d11 viewport code

diff --git a/Engine/Source/Runtime/Windows/D3D11RHI/Private/D3D11Viewport.cpp b/Engine/Source/Runtime/Windows/D3D11RHI/Private/D3D11Viewport.cpp
--- a/Engine/Source/Runtime/Windows/D3D11RHI/Private/D3D11Viewport.cpp
+++ b/Engine/Source/Runtime/Windows/D3D11RHI/Private/D3D11Viewport.cpp
@@ -515,13 +515,14 @@ void FD3D11DynamicRHI::RHIEndDrawingViewport(FViewportRHIParamRef ViewportRHI,bo
 	StateCache.SetDomainShader(nullptr);
 	StateCache.SetGeometryShader(nullptr);
 	// Compute Shader is set to NULL after each Dispatch call, so no need to clear it here
-
+#if STATS
 	FThreadStats::AddMessage(GET_STATFNAME(STAT_D3D11CommitResourceTables), EStatOperation::Set, int64(CommitResourceTableCycles), true);
 	FThreadStats::AddMessage(GET_STATFNAME(STAT_D3D11CacheResourceTables), EStatOperation::Set, int64(CacheResourceTableCycles), true);
 	FThreadStats::AddMessage(GET_STATFNAME(STAT_D3D11SetShaderTextureTime), EStatOperation::Set, int64(SetShaderTextureCycles), true);
 	INC_DWORD_STAT_BY(STAT_D3D11SetShaderTextureCalls, SetShaderTextureCalls);
 	INC_DWORD_STAT_BY(STAT_D3D11CacheResourceTableCalls, CacheResourceTableCalls);
 	INC_DWORD_STAT_BY(STAT_D3D11SetTextureInTableCalls, SetTextureInTableCalls);
+#endif
 
 	CommitResourceTableCycles = 0;
 	CacheResourceTableCalls = 0;
