UE4 - xb1 compile fix

diff --git a/Engine/Source/Runtime/Windows/D3D11RHI/Private/D3D11UniformBuffer.cpp b/Engine/Source/Runtime/Windows/D3D11RHI/Private/D3D11UniformBuffer.cpp
--- a/Engine/Source/Runtime/Windows/D3D11RHI/Private/D3D11UniformBuffer.cpp
+++ b/Engine/Source/Runtime/Windows/D3D11RHI/Private/D3D11UniformBuffer.cpp
@@ -265,7 +265,7 @@ void FD3D11UniformBuffer::CacheResourcesInternal()
 			{
 				FD3D11ShaderResourceView* ShaderResourceViewRHI = (FD3D11ShaderResourceView*)Resources[i].GetReference();
 				RawResources[i].ShaderResource = ShaderResourceViewRHI->Resource.GetReference();
-				RawResources[i].D3D11Resource = ShaderResourceViewRHI->View.GetReference();
+				RawResources[i].D3D11Resource = (IUnknown*)ShaderResourceViewRHI->View.GetReference();
 			}
 			break;
 
@@ -277,7 +277,7 @@ void FD3D11UniformBuffer::CacheResourcesInternal()
 				TextureRHI->SetLastCachedTime(CurrentTime);
 				FD3D11TextureBase* TextureD3D11 = GetD3D11TextureFromRHITexture(TextureRHI);
 				RawResources[i].ShaderResource = TextureD3D11->GetBaseShaderResource();
-				RawResources[i].D3D11Resource = TextureD3D11->GetShaderResourceView();
+				RawResources[i].D3D11Resource = (IUnknown*)TextureD3D11->GetShaderResourceView();
 			}
 			break;
 
@@ -289,7 +289,7 @@ void FD3D11UniformBuffer::CacheResourcesInternal()
 
 		case UBMT_SAMPLER:
 			RawResources[i].ShaderResource = NULL;
-			RawResources[i].D3D11Resource = ((FD3D11SamplerState*)Resources[i].GetReference())->Resource.GetReference();
+			RawResources[i].D3D11Resource = (IUnknown*)((FD3D11SamplerState*)Resources[i].GetReference())->Resource.GetReference();
 			break;
 
 		default:
