Fixed xbox rendering

diff --git a/Engine/Source/Runtime/Windows/D3D11RHI/Private/D3D11UniformBuffer.cpp b/Engine/Source/Runtime/Windows/D3D11RHI/Private/D3D11UniformBuffer.cpp
--- a/Engine/Source/Runtime/Windows/D3D11RHI/Private/D3D11UniformBuffer.cpp
+++ b/Engine/Source/Runtime/Windows/D3D11RHI/Private/D3D11UniformBuffer.cpp
@@ -112,8 +112,9 @@ FUniformBufferRHIRef FD3D11DynamicRHI::RHICreateUniformBuffer(const void* Conten
 			TRefCountPtr<ID3D11Buffer> UniformBufferResource;
 			FRingAllocation RingAllocation;
 
-#if PLATFORM_XBOXONE
-			if (Usage == UniformBuffer_SingleFrame)
+#if 0//PLATFORM_XBOXONE
+			// FIXME: this does not work
+			if (Usage == UniformBuffer_SingleDraw)
 			{
 				RingAllocation = DynamicCB.MapAndFill(GetDeviceContext(),NumBytes,Contents);
 				check(RingAllocation.IsValid());
