Remove ini override for texture pool size

diff --git a/Engine/Source/Runtime/EmptyRHI/Private/EmptyRHI.cpp b/Engine/Source/Runtime/EmptyRHI/Private/EmptyRHI.cpp
--- a/Engine/Source/Runtime/EmptyRHI/Private/EmptyRHI.cpp
+++ b/Engine/Source/Runtime/EmptyRHI/Private/EmptyRHI.cpp
@@ -41,7 +41,6 @@ FEmptyDynamicRHI::FEmptyDynamicRHI()
 // 	GRHISupportsTextureStreaming = true;
 // 	GMaxShadowDepthBufferSizeX = 4096;
 // 	GMaxShadowDepthBufferSizeY = 4096;
-// 	GReadTexturePoolSizeFromIni = true;
 // 
 // 	GMaxTextureDimensions = 16384;
 // 	GMaxTextureMipCount = FPlatformMath::CeilLogTwo( GMaxTextureDimensions ) + 1;
diff --git a/Engine/Source/Runtime/Engine/Private/ContentStreaming.cpp b/Engine/Source/Runtime/Engine/Private/ContentStreaming.cpp
--- a/Engine/Source/Runtime/Engine/Private/ContentStreaming.cpp
+++ b/Engine/Source/Runtime/Engine/Private/ContentStreaming.cpp
@@ -2089,13 +2089,7 @@ FStreamingManagerTexture::FStreamingManagerTexture()
 	{
 		OriginalTexturePoolSize =  int64(PoolSizeCVar) * 1024ll * 1024ll;
 		GTexturePoolSize = OriginalTexturePoolSize;
-	}
-	else if ( GReadTexturePoolSizeFromIni )
-	{
-		// Read the texture pool size straight from the .ini
-		OriginalTexturePoolSize = int64(PoolSizeIniSetting) * 1024ll * 1024ll;
-		GTexturePoolSize = OriginalTexturePoolSize;
-	}
+	}	
 	else if ( GPoolSizeVRAMPercentage )
 	{
 		// If GPoolSizeVRAMPercentage is set, the pool size has already been calculated and we're not reading it from the .ini
diff --git a/Engine/Source/Runtime/OpenGLDrv/Private/OpenGLDevice.cpp b/Engine/Source/Runtime/OpenGLDrv/Private/OpenGLDevice.cpp
--- a/Engine/Source/Runtime/OpenGLDrv/Private/OpenGLDevice.cpp
+++ b/Engine/Source/Runtime/OpenGLDrv/Private/OpenGLDevice.cpp
@@ -438,17 +438,7 @@ static void InitRHICapabilitiesForGL()
 	GTexturePoolSize = 0;
 	GPoolSizeVRAMPercentage = 0;
 #if PLATFORM_WINDOWS
-	if ( GReadTexturePoolSizeFromIni )
-	{
-		int32 PoolSize;
-		GConfig->GetInt(TEXT("TextureStreaming"), TEXT("PoolSize"), PoolSize, GEngineIni);
-
-		GTexturePoolSize = int64(PoolSize) * 1024 * 1024;
-	}
-	else
-	{
-		GConfig->GetInt( TEXT( "TextureStreaming" ), TEXT( "PoolSizeVRAMPercentage" ), GPoolSizeVRAMPercentage, GEngineIni );
-	}
+	GConfig->GetInt( TEXT( "TextureStreaming" ), TEXT( "PoolSizeVRAMPercentage" ), GPoolSizeVRAMPercentage, GEngineIni );	
 #endif
 
 	// GL vendor and version information.
diff --git a/Engine/Source/Runtime/RHI/Private/DynamicRHI.cpp b/Engine/Source/Runtime/RHI/Private/DynamicRHI.cpp
--- a/Engine/Source/Runtime/RHI/Private/DynamicRHI.cpp
+++ b/Engine/Source/Runtime/RHI/Private/DynamicRHI.cpp
@@ -38,14 +38,7 @@ void RHIInit(bool bHasEditorToken)
 	if(!GDynamicRHI)
 	{
 		const TCHAR* CmdLine = FCommandLine::Get();
-		FString Token = FParse::Token(CmdLine, false);
-
-		if ( FParse::Param(FCommandLine::Get(),TEXT("UseTexturePool")) )
-		{
-			/** Whether to read the texture pool size from engine.ini on PC. Can be turned on with -UseTexturePool on the command line. */
-			extern RHI_API bool GReadTexturePoolSizeFromIni;
-			GReadTexturePoolSizeFromIni = true;
-		}
+		FString Token = FParse::Token(CmdLine, false);		
 
 		if(USE_NULL_RHI || FParse::Param(FCommandLine::Get(),TEXT("nullrhi")) || IsRunningCommandlet() || IsRunningDedicatedServer())
 		{
diff --git a/Engine/Source/Runtime/RHI/Private/RHI.cpp b/Engine/Source/Runtime/RHI/Private/RHI.cpp
--- a/Engine/Source/Runtime/RHI/Private/RHI.cpp
+++ b/Engine/Source/Runtime/RHI/Private/RHI.cpp
@@ -153,7 +153,6 @@ FVertexElementTypeSupportInfo GVertexElementTypeSupport;
 RHI_API int32 volatile GCurrentTextureMemorySize = 0;
 RHI_API int32 volatile GCurrentRendertargetMemorySize = 0;
 RHI_API int64 GTexturePoolSize = 0 * 1024 * 1024;
-RHI_API bool GReadTexturePoolSizeFromIni = false;
 RHI_API int32 GPoolSizeVRAMPercentage = 0;
 
 RHI_API EShaderPlatform GShaderPlatformForFeatureLevel[ERHIFeatureLevel::Num] = {SP_NumPlatforms,SP_NumPlatforms,SP_NumPlatforms,SP_NumPlatforms};
diff --git a/Engine/Source/Runtime/RHI/Public/RHI.h b/Engine/Source/Runtime/RHI/Public/RHI.h
--- a/Engine/Source/Runtime/RHI/Public/RHI.h
+++ b/Engine/Source/Runtime/RHI/Public/RHI.h
@@ -189,8 +189,6 @@ extern RHI_API volatile int32 GCurrentTextureMemorySize;
 extern RHI_API volatile int32 GCurrentRendertargetMemorySize;
 /** Current texture streaming pool size, in bytes. 0 means unlimited. */
 extern RHI_API int64 GTexturePoolSize;
-/** Whether to read the texture pool size from engine.ini. Can be turned on with -UseTexturePool on the command line. */
-extern RHI_API bool GReadTexturePoolSizeFromIni;
 /** In percent. If non-zero, the texture pool size is a percentage of GTotalGraphicsMemory. */
 extern RHI_API int32 GPoolSizeVRAMPercentage;
 
diff --git a/Engine/Source/Runtime/Windows/D3D11RHI/Private/D3D11Device.cpp b/Engine/Source/Runtime/Windows/D3D11RHI/Private/D3D11Device.cpp
--- a/Engine/Source/Runtime/Windows/D3D11RHI/Private/D3D11Device.cpp
+++ b/Engine/Source/Runtime/Windows/D3D11RHI/Private/D3D11Device.cpp
@@ -68,17 +68,7 @@ FD3D11DynamicRHI::FD3D11DynamicRHI(IDXGIFactory* InDXGIFactory,D3D_FEATURE_LEVEL
 
 	GPoolSizeVRAMPercentage = 0;
 	GTexturePoolSize = 0;
-	if ( GReadTexturePoolSizeFromIni )
-	{
-		int32 PoolSize;
-		GConfig->GetInt(TEXT("TextureStreaming"), TEXT("PoolSize"), PoolSize, GEngineIni);
-
-		GTexturePoolSize = int64(PoolSize) * 1024 * 1024;
-	}
-	else
-	{
-		GConfig->GetInt( TEXT( "TextureStreaming" ), TEXT( "PoolSizeVRAMPercentage" ), GPoolSizeVRAMPercentage, GEngineIni );
-	}
+	GConfig->GetInt( TEXT( "TextureStreaming" ), TEXT( "PoolSizeVRAMPercentage" ), GPoolSizeVRAMPercentage, GEngineIni );	
 
 	// Initialize the RHI capabilities.
 	check(FeatureLevel == D3D_FEATURE_LEVEL_11_0 || FeatureLevel == D3D_FEATURE_LEVEL_10_0 );
