Use ValidateAPI when creating the mastering voice

diff --git a/Engine/Source/Runtime/Windows/XAudio2/Private/XAudio2Device.cpp b/Engine/Source/Runtime/Windows/XAudio2/Private/XAudio2Device.cpp
--- a/Engine/Source/Runtime/Windows/XAudio2/Private/XAudio2Device.cpp
+++ b/Engine/Source/Runtime/Windows/XAudio2/Private/XAudio2Device.cpp
@@ -144,17 +144,19 @@ bool FXAudio2Device::InitializeHardware()
 	}
 
 	// Create the final output voice with either 2 or 6 channels
-	if( FXAudioDeviceProperties::XAudio2->CreateMasteringVoice( &FXAudioDeviceProperties::MasteringVoice, FXAudioDeviceProperties::NumSpeakers, SampleRate, 0, 0, NULL ) != S_OK )
+	if (!ValidateAPICall(TEXT("CreateMasteringVoice"), 
+			FXAudioDeviceProperties::XAudio2->CreateMasteringVoice(&FXAudioDeviceProperties::MasteringVoice, FXAudioDeviceProperties::NumSpeakers, SampleRate, 0, 0, NULL)))
 	{
-		UE_LOG(LogInit, Log, TEXT( "Failed to create the mastering voice for XAudio2" ) );
+		UE_LOG(LogInit, Warning, TEXT( "Failed to create the mastering voice for XAudio2" ) );
 		FXAudioDeviceProperties::XAudio2 = NULL;
 		return( false );
 	}
 #else	//XAUDIO_SUPPORTS_DEVICE_DETAILS
 	// Create the final output voice
-	if (FXAudioDeviceProperties::XAudio2->CreateMasteringVoice(&FXAudioDeviceProperties::MasteringVoice, UE4_XAUDIO2_NUMCHANNELS, UE4_XAUDIO2_SAMPLERATE, 0, 0, NULL ) != S_OK)
+	if (!ValidateAPICall(TEXT("CreateMasteringVoice"),
+			FXAudioDeviceProperties::XAudio2->CreateMasteringVoice(&FXAudioDeviceProperties::MasteringVoice, UE4_XAUDIO2_NUMCHANNELS, UE4_XAUDIO2_SAMPLERATE, 0, 0, NULL )))
 	{
-		UE_LOG(LogInit, Log, TEXT( "Failed to create the mastering voice for XAudio2" ) );
+		UE_LOG(LogInit, Warning, TEXT( "Failed to create the mastering voice for XAudio2" ) );
 		FXAudioDeviceProperties::XAudio2 = NULL;
 		return false;
 	}
