Wrap some more XAudio API calls in ValidateAPI

diff --git a/Engine/Source/Runtime/Windows/XAudio2/Private/XAudio2Device.cpp b/Engine/Source/Runtime/Windows/XAudio2/Private/XAudio2Device.cpp
--- a/Engine/Source/Runtime/Windows/XAudio2/Private/XAudio2Device.cpp
+++ b/Engine/Source/Runtime/Windows/XAudio2/Private/XAudio2Device.cpp
@@ -96,7 +96,8 @@ bool FXAudio2Device::InitializeHardware()
 
 #if XAUDIO_SUPPORTS_DEVICE_DETAILS
 	UINT32 DeviceCount = 0;
-	FXAudioDeviceProperties::XAudio2->GetDeviceCount( &DeviceCount );
+	ValidateAPICall(TEXT("GetDeviceCount"),
+		FXAudioDeviceProperties::XAudio2->GetDeviceCount( &DeviceCount ));
 	if( DeviceCount < 1 )
 	{
 		UE_LOG(LogInit, Log, TEXT( "No audio devices found!" ) );
@@ -105,7 +106,8 @@ bool FXAudio2Device::InitializeHardware()
 	}
 
 	// Get the details of the default device 0
-	if( FXAudioDeviceProperties::XAudio2->GetDeviceDetails( 0, &FXAudioDeviceProperties::DeviceDetails ) != S_OK )
+	if( !ValidateAPICall(TEXT("GetDeviceDetails"),
+			FXAudioDeviceProperties::XAudio2->GetDeviceDetails( 0, &FXAudioDeviceProperties::DeviceDetails )))
 	{
 		UE_LOG(LogInit, Log, TEXT( "Failed to get DeviceDetails for XAudio2" ) );
 		FXAudioDeviceProperties::XAudio2 = NULL;
