Byte address buffer support for vertex buffers.

Oops. Really checked in this time.

diff --git a/Engine/Source/Runtime/Windows/D3D11RHI/Private/D3D11UAV.cpp b/Engine/Source/Runtime/Windows/D3D11RHI/Private/D3D11UAV.cpp
--- a/Engine/Source/Runtime/Windows/D3D11RHI/Private/D3D11UAV.cpp
+++ b/Engine/Source/Runtime/Windows/D3D11RHI/Private/D3D11UAV.cpp
@@ -101,6 +101,7 @@ FUnorderedAccessViewRHIRef FD3D11DynamicRHI::RHICreateUnorderedAccessView(FVerte
 	if (bByteAccessBuffer)
 	{
 		UAVDesc.Buffer.Flags |= D3D11_BUFFER_UAV_FLAG_RAW;
+		UAVDesc.Format = DXGI_FORMAT_R32_TYPELESS;
 	}
 
 	TRefCountPtr<ID3D11UnorderedAccessView> UnorderedAccessView;
diff --git a/Engine/Source/Runtime/Windows/D3D11RHI/Private/D3D11VertexBuffer.cpp b/Engine/Source/Runtime/Windows/D3D11RHI/Private/D3D11VertexBuffer.cpp
--- a/Engine/Source/Runtime/Windows/D3D11RHI/Private/D3D11VertexBuffer.cpp
+++ b/Engine/Source/Runtime/Windows/D3D11RHI/Private/D3D11VertexBuffer.cpp
@@ -31,6 +31,11 @@ FVertexBufferRHIRef FD3D11DynamicRHI::RHICreateVertexBuffer(uint32 Size,uint32 I
 		}
 	}
 
+	if (InUsage & BUF_ByteAddressBuffer)
+	{
+		Desc.MiscFlags |= D3D11_RESOURCE_MISC_BUFFER_ALLOW_RAW_VIEWS;
+	}
+
 	if (InUsage & BUF_StreamOutput)
 	{
 		Desc.BindFlags |= D3D11_BIND_STREAM_OUTPUT;
