UE4: Fixed up some usage of the OVERRIDE macro

diff --git a/Engine/Source/Runtime/Windows/XAudio2/Private/XAudio2Support.h b/Engine/Source/Runtime/Windows/XAudio2/Private/XAudio2Support.h
--- a/Engine/Source/Runtime/Windows/XAudio2/Private/XAudio2Support.h
+++ b/Engine/Source/Runtime/Windows/XAudio2/Private/XAudio2Support.h
@@ -347,7 +347,7 @@ public:
 	/**
 	 * Returns a string describing the source
 	 */
-	virtual FString Describe(bool bUseLongName) OVERRIDE;
+	virtual FString Describe(bool bUseLongName) override;
 
 	/**
 	 * Returns a string describing the source. For internal use to avoid recursively calling GetChannelVolumes if invoked from GetChannelVolumes.
