Fix fast constant ring buffer for TTP 338147

Also removed all instances of D3D11XboxPerformanceDevice as this is no longer needed.

diff --git a/Engine/Source/Runtime/Windows/D3D11RHI/Private/D3D11Commands.cpp b/Engine/Source/Runtime/Windows/D3D11RHI/Private/D3D11Commands.cpp
--- a/Engine/Source/Runtime/Windows/D3D11RHI/Private/D3D11Commands.cpp
+++ b/Engine/Source/Runtime/Windows/D3D11RHI/Private/D3D11Commands.cpp
@@ -689,8 +689,18 @@ void FD3D11DynamicRHI::RHISetShaderUniformBuffer(FPixelShaderRHIParamRef PixelSh
 
 	VALIDATE_BOUND_SHADER(PixelShader);
 	DYNAMIC_CAST_D3D11RESOURCE(UniformBuffer,Buffer);
-	ID3D11Buffer* ConstantBuffer = Buffer ? Buffer->Resource : NULL;
-	StateCache.SetConstantBuffer<SF_Pixel>(ConstantBuffer, BufferIndex);
+
+#if PLATFORM_XBOXONE
+	if (Buffer && Buffer->RingAllocation.IsValid())
+	{
+		StateCache.SetDynamicConstantBuffer<SF_Pixel>(BufferIndex, Buffer->RingAllocation);
+	}
+	else
+#endif
+	{
+		ID3D11Buffer* ConstantBuffer = Buffer ? Buffer->Resource : NULL;
+		StateCache.SetConstantBuffer<SF_Pixel>(ConstantBuffer, BufferIndex);
+	}
 
 	BoundUniformBuffers[SF_Pixel][BufferIndex] = BufferRHI;
 	DirtyUniformBuffers[SF_Pixel] |= (1 << BufferIndex);
diff --git a/Engine/Source/Runtime/Windows/D3D11RHI/Private/D3D11UniformBuffer.cpp b/Engine/Source/Runtime/Windows/D3D11RHI/Private/D3D11UniformBuffer.cpp
--- a/Engine/Source/Runtime/Windows/D3D11RHI/Private/D3D11UniformBuffer.cpp
+++ b/Engine/Source/Runtime/Windows/D3D11RHI/Private/D3D11UniformBuffer.cpp
@@ -112,8 +112,7 @@ FUniformBufferRHIRef FD3D11DynamicRHI::RHICreateUniformBuffer(const void* Conten
 			TRefCountPtr<ID3D11Buffer> UniformBufferResource;
 			FRingAllocation RingAllocation;
 
-#if 0//PLATFORM_XBOXONE
-			// FIXME: this does not work
+#if PLATFORM_XBOXONE
 			if (Usage == UniformBuffer_SingleDraw)
 			{
 				RingAllocation = DynamicCB.MapAndFill(GetDeviceContext(),NumBytes,Contents);
@@ -188,15 +187,13 @@ FUniformBufferRHIRef FD3D11DynamicRHI::RHICreateUniformBuffer(const void* Conten
 			TRefCountPtr<ID3D11Buffer> UniformBufferResource;
 			VERIFYD3D11RESULT(Direct3DDevice->CreateBuffer(&Desc,&ImmutableData,UniformBufferResource.GetInitReference()));
 
-			FRingAllocation RingAllocation;
-			NewUniformBuffer = new FD3D11UniformBuffer(this, Layout, UniformBufferResource, RingAllocation);
+			NewUniformBuffer = new FD3D11UniformBuffer(this, Layout, UniformBufferResource, FRingAllocation());
 		}
 	}
 	else
 	{
 		// This uniform buffer contains no constants, only a resource table.
-		FRingAllocation RingAllocation;
-		NewUniformBuffer = new FD3D11UniformBuffer(this, Layout, nullptr, RingAllocation);
+		NewUniformBuffer = new FD3D11UniformBuffer(this, Layout, nullptr, FRingAllocation());
 	}
 
 	if (Layout.Resources.Num())
