Unbind D3D11 shader resources from the device when they're being mapped

Fix for TTP 336060 - UE4: RENDERING: D3DDebug error in RenderTestMap

diff --git a/Engine/Source/Runtime/Windows/D3D11RHI/Private/D3D11IndexBuffer.cpp b/Engine/Source/Runtime/Windows/D3D11RHI/Private/D3D11IndexBuffer.cpp
--- a/Engine/Source/Runtime/Windows/D3D11RHI/Private/D3D11IndexBuffer.cpp
+++ b/Engine/Source/Runtime/Windows/D3D11RHI/Private/D3D11IndexBuffer.cpp
@@ -65,6 +65,9 @@ void* FD3D11DynamicRHI::RHILockIndexBuffer(FIndexBufferRHIParamRef IndexBufferRH
 {
 	DYNAMIC_CAST_D3D11RESOURCE(IndexBuffer,IndexBuffer);
 	
+	// If this resource is bound to the device, unbind it
+	ConditionalClearShaderResource(IndexBuffer);
+
 	// Determine whether the index buffer is dynamic or not.
 	D3D11_BUFFER_DESC Desc;
 	IndexBuffer->Resource->GetDesc(&Desc);
diff --git a/Engine/Source/Runtime/Windows/D3D11RHI/Private/D3D11StructuredBuffer.cpp b/Engine/Source/Runtime/Windows/D3D11RHI/Private/D3D11StructuredBuffer.cpp
--- a/Engine/Source/Runtime/Windows/D3D11RHI/Private/D3D11StructuredBuffer.cpp
+++ b/Engine/Source/Runtime/Windows/D3D11RHI/Private/D3D11StructuredBuffer.cpp
@@ -91,6 +91,9 @@ FStructuredBufferRHIRef FD3D11DynamicRHI::RHICreateStructuredBuffer(uint32 Strid
 void* FD3D11DynamicRHI::RHILockStructuredBuffer(FStructuredBufferRHIParamRef StructuredBufferRHI,uint32 Offset,uint32 Size,EResourceLockMode LockMode)
 {
 	DYNAMIC_CAST_D3D11RESOURCE(StructuredBuffer,StructuredBuffer);
+	
+	// If this resource is bound to the device, unbind it
+	ConditionalClearShaderResource(StructuredBuffer);
 
 	// Determine whether the Structured buffer is dynamic or not.
 	D3D11_BUFFER_DESC Desc;
diff --git a/Engine/Source/Runtime/Windows/D3D11RHI/Private/D3D11Texture.cpp b/Engine/Source/Runtime/Windows/D3D11RHI/Private/D3D11Texture.cpp
--- a/Engine/Source/Runtime/Windows/D3D11RHI/Private/D3D11Texture.cpp
+++ b/Engine/Source/Runtime/Windows/D3D11RHI/Private/D3D11Texture.cpp
@@ -1407,6 +1407,7 @@ void* FD3D11DynamicRHI::RHILockTexture2D(FTexture2DRHIParamRef TextureRHI,uint32
 {
 	check(TextureRHI);
 	DYNAMIC_CAST_D3D11RESOURCE(Texture2D,Texture);
+	ConditionalClearShaderResource(Texture);
 	return Texture->Lock(MipIndex,0,LockMode,DestStride);
 }
 
@@ -1420,6 +1421,7 @@ void FD3D11DynamicRHI::RHIUnlockTexture2D(FTexture2DRHIParamRef TextureRHI,uint3
 void* FD3D11DynamicRHI::RHILockTexture2DArray(FTexture2DArrayRHIParamRef TextureRHI,uint32 TextureIndex,uint32 MipIndex,EResourceLockMode LockMode,uint32& DestStride,bool bLockWithinMiptail)
 {
 	DYNAMIC_CAST_D3D11RESOURCE(Texture2DArray,Texture);
+	ConditionalClearShaderResource(Texture);
 	return Texture->Lock(MipIndex,TextureIndex,LockMode,DestStride);
 }
 
@@ -1477,6 +1479,7 @@ FTextureCubeRHIRef FD3D11DynamicRHI::RHICreateTextureCubeArray(uint32 Size, uint
 void* FD3D11DynamicRHI::RHILockTextureCubeFace(FTextureCubeRHIParamRef TextureCubeRHI,uint32 FaceIndex,uint32 ArrayIndex,uint32 MipIndex,EResourceLockMode LockMode,uint32& DestStride,bool bLockWithinMiptail)
 {
 	DYNAMIC_CAST_D3D11RESOURCE(TextureCube,TextureCube);
+	ConditionalClearShaderResource(TextureCube);
 	uint32 D3DFace = GetD3D11CubeFace((ECubeFace)FaceIndex);
 	return TextureCube->Lock(MipIndex,D3DFace + ArrayIndex * 6,LockMode,DestStride);
 }
diff --git a/Engine/Source/Runtime/Windows/D3D11RHI/Private/D3D11VertexBuffer.cpp b/Engine/Source/Runtime/Windows/D3D11RHI/Private/D3D11VertexBuffer.cpp
--- a/Engine/Source/Runtime/Windows/D3D11RHI/Private/D3D11VertexBuffer.cpp
+++ b/Engine/Source/Runtime/Windows/D3D11RHI/Private/D3D11VertexBuffer.cpp
@@ -90,6 +90,9 @@ void* FD3D11DynamicRHI::RHILockVertexBuffer(FVertexBufferRHIParamRef VertexBuffe
 	check(Size > 0);
 
 	DYNAMIC_CAST_D3D11RESOURCE(VertexBuffer,VertexBuffer);
+	
+	// If this resource is bound to the device, unbind it
+	ConditionalClearShaderResource(VertexBuffer);
 
 	// Determine whether the vertex buffer is dynamic or not.
 	D3D11_BUFFER_DESC Desc;
diff --git a/Engine/Source/Runtime/Windows/D3D11RHI/Public/D3D11Resources.h b/Engine/Source/Runtime/Windows/D3D11RHI/Public/D3D11Resources.h
--- a/Engine/Source/Runtime/Windows/D3D11RHI/Public/D3D11Resources.h
+++ b/Engine/Source/Runtime/Windows/D3D11RHI/Public/D3D11Resources.h
@@ -647,7 +647,7 @@ private:
 };
 
 /** Index buffer resource class that stores stride information. */
-class FD3D11IndexBuffer : public FRHIIndexBuffer
+class FD3D11IndexBuffer : public FRHIIndexBuffer, public FD3D11BaseShaderResource
 {
 public:
 
@@ -663,6 +663,20 @@ public:
 	{
 		UpdateBufferStats(Resource, false);
 	}
+
+	// IRefCountedObject interface.
+	virtual uint32 AddRef() const
+	{
+		return FRefCountedObject::AddRef();
+	}
+	virtual uint32 Release() const
+	{
+		return FRefCountedObject::Release();
+	}
+	virtual uint32 GetRefCount() const
+	{
+		return FRefCountedObject::GetRefCount();
+	}
 };
 
 /** Structured buffer resource class. */
