Stats - Misc code cleanup

diff --git a/Engine/Source/Runtime/Windows/D3D11RHI/Private/D3D11Viewport.cpp b/Engine/Source/Runtime/Windows/D3D11RHI/Private/D3D11Viewport.cpp
--- a/Engine/Source/Runtime/Windows/D3D11RHI/Private/D3D11Viewport.cpp
+++ b/Engine/Source/Runtime/Windows/D3D11RHI/Private/D3D11Viewport.cpp
@@ -516,9 +516,9 @@ void FD3D11DynamicRHI::RHIEndDrawingViewport(FViewportRHIParamRef ViewportRHI,bo
 	StateCache.SetGeometryShader(nullptr);
 	// Compute Shader is set to NULL after each Dispatch call, so no need to clear it here
 #if STATS
-	FThreadStats::AddMessage(GET_STATFNAME(STAT_D3D11CommitResourceTables), EStatOperation::Set, int64(CommitResourceTableCycles), true);
-	FThreadStats::AddMessage(GET_STATFNAME(STAT_D3D11CacheResourceTables), EStatOperation::Set, int64(CacheResourceTableCycles), true);
-	FThreadStats::AddMessage(GET_STATFNAME(STAT_D3D11SetShaderTextureTime), EStatOperation::Set, int64(SetShaderTextureCycles), true);
+	SET_CYCLE_COUNTER(STAT_D3D11CommitResourceTables, CommitResourceTableCycles);
+	SET_CYCLE_COUNTER(STAT_D3D11CacheResourceTables, CacheResourceTableCycles);
+	SET_CYCLE_COUNTER(STAT_D3D11SetShaderTextureTime, SetShaderTextureCycles);
 	INC_DWORD_STAT_BY(STAT_D3D11SetShaderTextureCalls, SetShaderTextureCalls);
 	INC_DWORD_STAT_BY(STAT_D3D11CacheResourceTableCalls, CacheResourceTableCalls);
 	INC_DWORD_STAT_BY(STAT_D3D11SetTextureInTableCalls, SetTextureInTableCalls);
