Fix crash when loading between maps
Resources being referenced by the state cache were being released before finished being used.

diff --git a/Engine/Source/Runtime/Windows/D3D11RHI/Private/D3D11StateCachePrivate.h b/Engine/Source/Runtime/Windows/D3D11RHI/Private/D3D11StateCachePrivate.h
--- a/Engine/Source/Runtime/Windows/D3D11RHI/Private/D3D11StateCachePrivate.h
+++ b/Engine/Source/Runtime/Windows/D3D11RHI/Private/D3D11StateCachePrivate.h
@@ -205,6 +205,14 @@ protected:
 		check(ResourceIndex < ARRAYSIZE(CurrentShaderResourceViews[ShaderFrequency]));
 		if ((CurrentShaderResourceViews[ShaderFrequency][ResourceIndex] != SRV) || GD3D11SkipStateCaching)
 		{
+			if(SRV)
+			{
+				SRV->AddRef();
+			}
+			if(CurrentShaderResourceViews[ShaderFrequency][ResourceIndex])
+			{
+				CurrentShaderResourceViews[ShaderFrequency][ResourceIndex]->Release();
+			}
 			CurrentShaderResourceViews[ShaderFrequency][ResourceIndex] = SRV;
 			if (AlternatePathFunction != nullptr)
 			{
