Create and use auxilliary depth texture in the forward rendering path when simulationous depth test/read isn't supported on the target hardware.

Fixes TTP 335869 - UE4: RENDERING: CRITICAL: EDITOR: CRASH: Occurs when launching mobile preview on the StarterMap using the Asus UX31 Sandy Bridge Ultrabook

diff --git a/Engine/Shaders/ResolvePixelShader.usf b/Engine/Shaders/ResolvePixelShader.usf
--- a/Engine/Shaders/ResolvePixelShader.usf
+++ b/Engine/Shaders/ResolvePixelShader.usf
@@ -50,7 +50,7 @@ void MainSingleSample(
 	OutColor = UnresolvedSurface.Load(IntUV,SingleSampleIndex);
 }
 
-#elif PLATFORM_LEVEL == PLATFORM_LEVEL_SM4
+#else
 
 Texture2D UnresolvedSurfaceNonMS;
 void MainDepthNonMS(
diff --git a/Engine/Source/Runtime/Renderer/Private/ForwardShadingRenderer.cpp b/Engine/Source/Runtime/Renderer/Private/ForwardShadingRenderer.cpp
--- a/Engine/Source/Runtime/Renderer/Private/ForwardShadingRenderer.cpp
+++ b/Engine/Source/Runtime/Renderer/Private/ForwardShadingRenderer.cpp
@@ -88,6 +88,9 @@ void FForwardShadingSceneRenderer::Render()
 
 	RenderForwardShadingBasePass(RHICmdList);
 
+	// Make a copy of the scene depth if the current hardware doesn't support reading and writing to the same depth buffer
+	GSceneRenderTargets.ResolveSceneDepthToAuxiliaryTexture();
+
 	// Notify the FX system that opaque primitives have been rendered.
 	if (Scene->FXSystem)
 	{
diff --git a/Engine/Source/Runtime/Renderer/Private/PostProcess/SceneRenderTargets.cpp b/Engine/Source/Runtime/Renderer/Private/PostProcess/SceneRenderTargets.cpp
--- a/Engine/Source/Runtime/Renderer/Private/PostProcess/SceneRenderTargets.cpp
+++ b/Engine/Source/Runtime/Renderer/Private/PostProcess/SceneRenderTargets.cpp
@@ -1028,6 +1028,13 @@ void FSceneRenderTargets::AllocateForwardShadingPathRenderTargets()
 	{
 		SceneAlphaCopy = GSystemTextures.MaxFP16Depth;
 	}
+
+	// When targeting DX Feature Level 10, create an auxiliary texture to store the resolved scene depth, and a render-targetable surface to hold the unresolved scene depth.
+	if (!GSupportsDepthFetchDuringDepthTest)
+	{
+		FPooledRenderTargetDesc Desc(FPooledRenderTargetDesc::Create2DDesc(BufferSize, PF_DepthStencil, TexCreate_None, TexCreate_DepthStencilTargetable, false));
+		GRenderTargetPool.FindFreeElement(Desc, AuxiliarySceneDepthZ, TEXT("AuxiliarySceneDepthZ"));
+	}
 }
 
 // for easier use of "VisualizeTexture"
diff --git a/Engine/Source/Runtime/Renderer/Private/PostProcess/VisualizeTexture.cpp b/Engine/Source/Runtime/Renderer/Private/PostProcess/VisualizeTexture.cpp
--- a/Engine/Source/Runtime/Renderer/Private/PostProcess/VisualizeTexture.cpp
+++ b/Engine/Source/Runtime/Renderer/Private/PostProcess/VisualizeTexture.cpp
@@ -883,7 +883,7 @@ IPooledRenderTarget* FVisualizeTexture::GetObservedElement() const
 
 void FVisualizeTexture::OnStartFrame(const FSceneView& View)
 {
-	bEnabled = View.GetFeatureLevel() >= ERHIFeatureLevel::SM3;
+	bEnabled = GMaxRHIFeatureLevel >= ERHIFeatureLevel::SM3;
 	ViewRect = View.UnscaledViewRect;
 	AspectRatioConstrainedViewRect = View.Family->EngineShowFlags.CameraAspectRatioBars ? View.CameraConstrainedViewRect : ViewRect;
 
diff --git a/Engine/Source/Runtime/Windows/D3D11RHI/Private/D3D11RenderTarget.cpp b/Engine/Source/Runtime/Windows/D3D11RHI/Private/D3D11RenderTarget.cpp
--- a/Engine/Source/Runtime/Windows/D3D11RHI/Private/D3D11RenderTarget.cpp
+++ b/Engine/Source/Runtime/Windows/D3D11RHI/Private/D3D11RenderTarget.cpp
@@ -52,7 +52,7 @@ namespace
 
 		typedef FDummyResolveParameter FParameter;
 
-		static bool ShouldCache(EShaderPlatform Platform) { return Platform == SP_PCD3D_SM4; }
+		static bool ShouldCache(EShaderPlatform Platform) { return GetMaxSupportedFeatureLevel(Platform) <= ERHIFeatureLevel::SM4; }
 
 		FResolveDepthNonMSPS(const ShaderMetaType::CompiledShaderInitializerType& Initializer):
 		FGlobalShader(Initializer)
@@ -119,7 +119,7 @@ namespace
 		DECLARE_SHADER_TYPE(FResolveVS,Global);
 	public:
 
-		static bool ShouldCache(EShaderPlatform Platform) { return IsFeatureLevelSupported(Platform, ERHIFeatureLevel::SM4); }
+		static bool ShouldCache(EShaderPlatform Platform) { return true; }
 
 		FResolveVS(const ShaderMetaType::CompiledShaderInitializerType& Initializer):
 			FGlobalShader(Initializer)
