When running with D3D debug, clear state and flush the device to get rid of spurious live objects in D3D11's report.

diff --git a/Engine/Source/Runtime/Windows/D3D11RHI/Private/D3D11Device.cpp b/Engine/Source/Runtime/Windows/D3D11RHI/Private/D3D11Device.cpp
--- a/Engine/Source/Runtime/Windows/D3D11RHI/Private/D3D11Device.cpp
+++ b/Engine/Source/Runtime/Windows/D3D11RHI/Private/D3D11Device.cpp
@@ -424,6 +424,14 @@ void FD3D11DynamicRHI::CleanupD3DDevice()
 
 		// Release the device and its IC
 		StateCache.SetContext(nullptr);
+
+		// When running with D3D debug, clear state and flush the device to get rid of spurious live objects in D3D11's report.
+		if (D3D11RHI_ShouldCreateWithD3DDebug())
+		{
+			Direct3DDeviceIMContext->ClearState();
+			Direct3DDeviceIMContext->Flush();
+		}
+
 		Direct3DDeviceIMContext = NULL;
 
 		Direct3DDevice = NULL;
