Disable CMASK on XB1
Fixes 'checkerboarding' artifact.  Needs to be re-enabled when we determine the root cause of the CMASK corruption.

diff --git a/Engine/Source/Runtime/Windows/D3D11RHI/Private/D3D11Texture.cpp b/Engine/Source/Runtime/Windows/D3D11RHI/Private/D3D11Texture.cpp
--- a/Engine/Source/Runtime/Windows/D3D11RHI/Private/D3D11Texture.cpp
+++ b/Engine/Source/Runtime/Windows/D3D11RHI/Private/D3D11Texture.cpp
@@ -576,6 +576,10 @@ TD3D11Texture2D<BaseResourceType>* FD3D11DynamicRHI::CreateD3D11Texture2D(uint32
 		check(!(Flags & TexCreate_DepthStencilTargetable));
 		check(!(Flags & TexCreate_ResolveTargetable));
 		TextureDesc.BindFlags |= D3D11_BIND_RENDER_TARGET;
+#if PLATFORM_XBOXONE
+		// Turn off color compression until we get to the bottom of the random CMASK corruption.
+		TextureDesc.MiscFlags |= D3D11X_RESOURCE_MISC_NO_COLOR_COMPRESSION;
+#endif
 		bCreateRTV = true;
 	}
 	else if(Flags & TexCreate_DepthStencilTargetable)
