Remove unused duplicate define

diff --git a/Engine/Source/Runtime/Windows/D3D11RHI/Private/D3D11RHIPrivate.h b/Engine/Source/Runtime/Windows/D3D11RHI/Private/D3D11RHIPrivate.h
--- a/Engine/Source/Runtime/Windows/D3D11RHI/Private/D3D11RHIPrivate.h
+++ b/Engine/Source/Runtime/Windows/D3D11RHI/Private/D3D11RHIPrivate.h
@@ -36,7 +36,6 @@ DECLARE_LOG_CATEGORY_EXTERN(LogD3D11RHI, Log, All);
 #define DX_MAX_MSAA_COUNT 8
 
 // Definitions.
-#define DEBUG_SHADERS 0
 #define D3D11 1
 
 /**
