Do not disable audio when an audio devices that support > 6 channels is attached

diff --git a/Engine/Source/Runtime/Windows/XAudio2/Private/XAudio2Device.cpp b/Engine/Source/Runtime/Windows/XAudio2/Private/XAudio2Device.cpp
--- a/Engine/Source/Runtime/Windows/XAudio2/Private/XAudio2Device.cpp
+++ b/Engine/Source/Runtime/Windows/XAudio2/Private/XAudio2Device.cpp
@@ -123,14 +123,6 @@ bool FXAudio2Device::InitializeHardware()
 	FXAudioDeviceProperties::NumSpeakers = UE4_XAUDIO2_NUMCHANNELS;
 	SampleRate = FXAudioDeviceProperties::DeviceDetails.OutputFormat.Format.nSamplesPerSec;
 
-	// Fail on devices using channels greater than the max
-	if (FXAudioDeviceProperties::NumSpeakers > SPEAKER_COUNT)
-	{
-		UE_LOG(LogInit, Log, TEXT("XAudio2 device used more channels than the max supported. Used: %d, Max: %d"), FXAudioDeviceProperties::NumSpeakers, SPEAKER_COUNT);
-		FXAudioDeviceProperties::XAudio2 = NULL;
-		return( false );
-	}
-
 	// Clamp the output frequency to the limits of the reverb XAPO
 	if( SampleRate > XAUDIO2FX_REVERB_MAX_FRAMERATE )
 	{
