Fix crash using 7.1 speaker set up due to X3DAudioCalculate overrunning the matrix coefficient buffer we set out

diff --git a/Engine/Source/Runtime/Windows/XAudio2/Private/XAudio2Support.h b/Engine/Source/Runtime/Windows/XAudio2/Private/XAudio2Support.h
--- a/Engine/Source/Runtime/Windows/XAudio2/Private/XAudio2Support.h
+++ b/Engine/Source/Runtime/Windows/XAudio2/Private/XAudio2Support.h
@@ -452,7 +452,12 @@ class FSpatializationHelper
 	X3DAUDIO_DISTANCE_CURVE       ReverbVolumeCurve;
 
 	float                         EmitterAzimuths[UE4_XAUDIO3D_INPUTCHANNELS];
-	float					      MatrixCoefficients[UE4_XAUDIO3D_INPUTCHANNELS * SPEAKER_COUNT];
+
+	// TODO: Hardcoding this to 8 because X3DAudioCalculate is ignoring the destination speaker count we put in and
+	//       using the number of speakers on the output device.  For 7.1 this means that it writes to 8 speakers,
+	//       overrunning the buffer and trashing other static variables
+	// 	float					      MatrixCoefficients[UE4_XAUDIO3D_INPUTCHANNELS * SPEAKER_COUNT];
+	float					      MatrixCoefficients[8]; 
 	
 public:
 	/**
