Remove stats that are no longer useful and were causing many calls to Cycles per frame.

diff --git a/Engine/Source/Runtime/Windows/D3D11RHI/Private/D3D11Commands.cpp b/Engine/Source/Runtime/Windows/D3D11RHI/Private/D3D11Commands.cpp
--- a/Engine/Source/Runtime/Windows/D3D11RHI/Private/D3D11Commands.cpp
+++ b/Engine/Source/Runtime/Windows/D3D11RHI/Private/D3D11Commands.cpp
@@ -249,8 +249,6 @@ void FD3D11DynamicRHI::RHISetBoundShaderState( FBoundShaderStateRHIParamRef Boun
 
 void FD3D11DynamicRHI::RHISetShaderTexture(FVertexShaderRHIParamRef VertexShaderRHI,uint32 TextureIndex,FTextureRHIParamRef NewTextureRHI)
 {
-	uint32 Start = FPlatformTime::Cycles();
-
 	GRHICommandList.Verify();
 
 	VALIDATE_BOUND_SHADER(VertexShaderRHI);
@@ -262,15 +260,10 @@ void FD3D11DynamicRHI::RHISetShaderTexture(FVertexShaderRHIParamRef VertexShader
 		SetShaderResourceView<SF_Vertex>(NewTexture, ShaderResourceView, TextureIndex, FD3D11StateCache::SRV_Dynamic);
 	else
 		SetShaderResourceView<SF_Vertex>(NewTexture, ShaderResourceView, TextureIndex, FD3D11StateCache::SRV_Static);
-
-	SetShaderTextureCycles += (FPlatformTime::Cycles() - Start);
-	SetShaderTextureCalls++;
 }
 
 void FD3D11DynamicRHI::RHISetShaderTexture(FHullShaderRHIParamRef HullShaderRHI,uint32 TextureIndex,FTextureRHIParamRef NewTextureRHI)
 {
-	uint32 Start = FPlatformTime::Cycles();
-
 	GRHICommandList.Verify();
 
 	VALIDATE_BOUND_SHADER(HullShaderRHI);
@@ -282,15 +275,10 @@ void FD3D11DynamicRHI::RHISetShaderTexture(FHullShaderRHIParamRef HullShaderRHI,
 		SetShaderResourceView<SF_Hull>(NewTexture, ShaderResourceView, TextureIndex, FD3D11StateCache::SRV_Dynamic);
 	else
 		SetShaderResourceView<SF_Hull>(NewTexture, ShaderResourceView, TextureIndex, FD3D11StateCache::SRV_Static);
-
-	SetShaderTextureCycles += (FPlatformTime::Cycles() - Start);
-	SetShaderTextureCalls++;
 }
 
 void FD3D11DynamicRHI::RHISetShaderTexture(FDomainShaderRHIParamRef DomainShaderRHI,uint32 TextureIndex,FTextureRHIParamRef NewTextureRHI)
 {
-	uint32 Start = FPlatformTime::Cycles();
-
 	GRHICommandList.Verify();
 
 	VALIDATE_BOUND_SHADER(DomainShaderRHI);
@@ -302,15 +290,10 @@ void FD3D11DynamicRHI::RHISetShaderTexture(FDomainShaderRHIParamRef DomainShader
 		SetShaderResourceView<SF_Domain>(NewTexture, ShaderResourceView, TextureIndex, FD3D11StateCache::SRV_Dynamic);
 	else
 		SetShaderResourceView<SF_Domain>(NewTexture, ShaderResourceView, TextureIndex, FD3D11StateCache::SRV_Static);
-
-	SetShaderTextureCycles += (FPlatformTime::Cycles() - Start);
-	SetShaderTextureCalls++;
 }
 
 void FD3D11DynamicRHI::RHISetShaderTexture(FGeometryShaderRHIParamRef GeometryShaderRHI,uint32 TextureIndex,FTextureRHIParamRef NewTextureRHI)
 {
-	uint32 Start = FPlatformTime::Cycles();
-
 	GRHICommandList.Verify();
 
 	VALIDATE_BOUND_SHADER(GeometryShaderRHI);
@@ -322,14 +305,10 @@ void FD3D11DynamicRHI::RHISetShaderTexture(FGeometryShaderRHIParamRef GeometrySh
 		SetShaderResourceView<SF_Geometry>(NewTexture, ShaderResourceView, TextureIndex, FD3D11StateCache::SRV_Dynamic);
 	else
 		SetShaderResourceView<SF_Geometry>(NewTexture, ShaderResourceView, TextureIndex, FD3D11StateCache::SRV_Static);
-
-	SetShaderTextureCycles += (FPlatformTime::Cycles() - Start);
-	SetShaderTextureCalls++;
 }
 
 void FD3D11DynamicRHI::RHISetShaderTexture(FPixelShaderRHIParamRef PixelShaderRHI,uint32 TextureIndex,FTextureRHIParamRef NewTextureRHI)
 {
-	uint32 Start = FPlatformTime::Cycles();
 
 	GRHICommandList.Verify();
 
@@ -341,15 +320,10 @@ void FD3D11DynamicRHI::RHISetShaderTexture(FPixelShaderRHIParamRef PixelShaderRH
 		SetShaderResourceView<SF_Pixel>(NewTexture, ShaderResourceView, TextureIndex, FD3D11StateCache::SRV_Dynamic);
 	else
 		SetShaderResourceView<SF_Pixel>(NewTexture, ShaderResourceView, TextureIndex, FD3D11StateCache::SRV_Static);
-
-	SetShaderTextureCycles += (FPlatformTime::Cycles() - Start);
-	SetShaderTextureCalls++;
 }
 
 void FD3D11DynamicRHI::RHISetShaderTexture(FComputeShaderRHIParamRef ComputeShaderRHI,uint32 TextureIndex,FTextureRHIParamRef NewTextureRHI)
 {
-	uint32 Start = FPlatformTime::Cycles();
-
 	GRHICommandList.Verify();
 	//VALIDATE_BOUND_SHADER(ComputeShaderRHI);
 
@@ -360,9 +334,6 @@ void FD3D11DynamicRHI::RHISetShaderTexture(FComputeShaderRHIParamRef ComputeShad
 		SetShaderResourceView<SF_Compute>(NewTexture, ShaderResourceView, TextureIndex, FD3D11StateCache::SRV_Dynamic);
 	else
 		SetShaderResourceView<SF_Compute>(NewTexture, ShaderResourceView, TextureIndex, FD3D11StateCache::SRV_Static);
-
-	SetShaderTextureCycles += (FPlatformTime::Cycles() - Start);
-	SetShaderTextureCalls++;
 }
 
 void FD3D11DynamicRHI::RHISetUAVParameter(FComputeShaderRHIParamRef ComputeShaderRHI,uint32 UAVIndex,FUnorderedAccessViewRHIParamRef UAVRHI)
@@ -1292,7 +1263,6 @@ void FD3D11DynamicRHI::SetResourcesFromTables(const ShaderType* RESTRICT Shader)
 
 	// Mask the dirty bits by those buffers from which the shader has bound resources.
 	uint32 DirtyBits = Shader->ShaderResourceTable.ResourceTableBits & DirtyUniformBuffers[ShaderType::StaticFrequency];
-	uint32 NumSetCalls = 0;
 	while (DirtyBits)
 	{
 		// Scan for the lowest set bit, compute its index, clear it in the set of dirty bits.
@@ -1306,17 +1276,14 @@ void FD3D11DynamicRHI::SetResourcesFromTables(const ShaderType* RESTRICT Shader)
 		Buffer->CacheResources(ResourceTableFrameCounter);
 
 		// todo: could make this two pass: gather then set
-		NumSetCalls += SetShaderResourcesFromBuffer<ID3D11ShaderResourceView, (EShaderFrequency)ShaderType::StaticFrequency>(this, &StateCache, Buffer, Shader->ShaderResourceTable.ShaderResourceViewMap.GetData(), BufferIndex);
-		NumSetCalls += SetShaderResourcesFromBuffer<ID3D11SamplerState, (EShaderFrequency)ShaderType::StaticFrequency>(this, &StateCache, Buffer, Shader->ShaderResourceTable.SamplerMap.GetData(), BufferIndex);
+		SetShaderResourcesFromBuffer<ID3D11ShaderResourceView, (EShaderFrequency)ShaderType::StaticFrequency>(this, &StateCache, Buffer, Shader->ShaderResourceTable.ShaderResourceViewMap.GetData(), BufferIndex);
+		SetShaderResourcesFromBuffer<ID3D11SamplerState, (EShaderFrequency)ShaderType::StaticFrequency>(this, &StateCache, Buffer, Shader->ShaderResourceTable.SamplerMap.GetData(), BufferIndex);
 	}
 	DirtyUniformBuffers[ShaderType::StaticFrequency] = 0;
-	SetTextureInTableCalls += NumSetCalls;
 }
 
 void FD3D11DynamicRHI::CommitGraphicsResourceTables()
 {
-	uint32 Start = FPlatformTime::Cycles();
-
 	GRHICommandList.Verify();
 
 	FD3D11BoundShaderState* RESTRICT CurrentBoundShaderState = (FD3D11BoundShaderState*)BoundShaderStateHistory.GetLast();
@@ -1342,8 +1309,6 @@ void FD3D11DynamicRHI::CommitGraphicsResourceTables()
 	{
 		SetResourcesFromTables(Shader);
 	}
-
-	CommitResourceTableCycles += (FPlatformTime::Cycles() - Start);
 }
 
 void FD3D11DynamicRHI::CommitComputeResourceTables(FD3D11ComputeShader* InComputeShader)
diff --git a/Engine/Source/Runtime/Windows/D3D11RHI/Private/D3D11Device.cpp b/Engine/Source/Runtime/Windows/D3D11RHI/Private/D3D11Device.cpp
--- a/Engine/Source/Runtime/Windows/D3D11RHI/Private/D3D11Device.cpp
+++ b/Engine/Source/Runtime/Windows/D3D11RHI/Private/D3D11Device.cpp
@@ -42,12 +42,6 @@ FD3D11DynamicRHI::FD3D11DynamicRHI(IDXGIFactory* InDXGIFactory,D3D_FEATURE_LEVEL
 	CurrentDepthTexture(NULL),
 	NumSimultaneousRenderTargets(0),
 	NumUAVs(0),
-	CommitResourceTableCycles(0),
-	CacheResourceTableCalls(0),
-	CacheResourceTableCycles(0),
-	SetShaderTextureCycles(0),
-	SetShaderTextureCalls(0),
-	SetTextureInTableCalls(0),
 	SceneFrameCounter(0),
 	ResourceTableFrameCounter(INDEX_NONE),
 	CurrentDSVAccessType(DSAT_Writable),
diff --git a/Engine/Source/Runtime/Windows/D3D11RHI/Private/D3D11RHIPrivate.h b/Engine/Source/Runtime/Windows/D3D11RHI/Private/D3D11RHIPrivate.h
--- a/Engine/Source/Runtime/Windows/D3D11RHI/Private/D3D11RHIPrivate.h
+++ b/Engine/Source/Runtime/Windows/D3D11RHI/Private/D3D11RHIPrivate.h
@@ -58,13 +58,6 @@ DECLARE_DWORD_COUNTER_STAT_EXTERN(TEXT("Textures Allocated"),STAT_D3D11TexturesA
 DECLARE_DWORD_COUNTER_STAT_EXTERN(TEXT("Textures Released"),STAT_D3D11TexturesReleased,STATGROUP_D3D11RHI, );
 DECLARE_MEMORY_STAT_EXTERN(TEXT("Texture object pool memory"),STAT_D3D11TexturePoolMemory,STATGROUP_D3D11RHI, );
 
-DECLARE_CYCLE_STAT_EXTERN(TEXT("Commit resource tables"),STAT_D3D11CommitResourceTables,STATGROUP_D3D11RHI, );
-DECLARE_CYCLE_STAT_EXTERN(TEXT("Cache resource tables"),STAT_D3D11CacheResourceTables,STATGROUP_D3D11RHI, );
-DECLARE_DWORD_COUNTER_STAT_EXTERN(TEXT("Num cached resource tables"),STAT_D3D11CacheResourceTableCalls,STATGROUP_D3D11RHI, );
-DECLARE_DWORD_COUNTER_STAT_EXTERN(TEXT("Num textures in tables"),STAT_D3D11SetTextureInTableCalls,STATGROUP_D3D11RHI, );
-DECLARE_CYCLE_STAT_EXTERN(TEXT("SetShaderTexture time"),STAT_D3D11SetShaderTextureTime,STATGROUP_D3D11RHI, );
-DECLARE_DWORD_COUNTER_STAT_EXTERN(TEXT("SetShaderTexture calls"),STAT_D3D11SetShaderTextureCalls,STATGROUP_D3D11RHI, );
-
 // This class has multiple inheritance but really FGPUTiming is a static class
 class FD3D11BufferedGPUTiming : public FRenderResource, public FGPUTiming
 {
@@ -446,14 +439,6 @@ protected:
 	uint32 NumSimultaneousRenderTargets;
 	uint32 NumUAVs;
 
-	friend class FD3D11UniformBuffer;
-	uint32 CommitResourceTableCycles;
-	uint32 CacheResourceTableCalls;
-	uint32 CacheResourceTableCycles;
-	uint32 SetShaderTextureCycles;
-	uint32 SetShaderTextureCalls;
-	uint32 SetTextureInTableCalls;
-	
 	/** Internal frame counter, incremented on each call to RHIBeginScene. */
 	uint32 SceneFrameCounter;
 
diff --git a/Engine/Source/Runtime/Windows/D3D11RHI/Private/D3D11UniformBuffer.cpp b/Engine/Source/Runtime/Windows/D3D11RHI/Private/D3D11UniformBuffer.cpp
--- a/Engine/Source/Runtime/Windows/D3D11RHI/Private/D3D11UniformBuffer.cpp
+++ b/Engine/Source/Runtime/Windows/D3D11RHI/Private/D3D11UniformBuffer.cpp
@@ -244,8 +244,6 @@ FD3D11UniformBuffer::~FD3D11UniformBuffer()
 
 void FD3D11UniformBuffer::CacheResourcesInternal()
 {
-	uint32 Start = FPlatformTime::Cycles();
-
 	const FRHIUniformBufferLayout& Layout = GetLayout();
 	int32 NumResources = Layout.Resources.Num();
 	const uint8* RESTRICT ResourceTypes = Layout.Resources.GetData();
@@ -295,9 +293,6 @@ void FD3D11UniformBuffer::CacheResourcesInternal()
 			break;
 		}
 	}
-
-	D3D11RHI->CacheResourceTableCycles += (FPlatformTime::Cycles() - Start);
-	D3D11RHI->CacheResourceTableCalls++;
 }
 
 void FD3D11DynamicRHI::ReleasePooledUniformBuffers()
diff --git a/Engine/Source/Runtime/Windows/D3D11RHI/Private/D3D11Util.cpp b/Engine/Source/Runtime/Windows/D3D11RHI/Private/D3D11Util.cpp
--- a/Engine/Source/Runtime/Windows/D3D11RHI/Private/D3D11Util.cpp
+++ b/Engine/Source/Runtime/Windows/D3D11RHI/Private/D3D11Util.cpp
@@ -385,11 +385,4 @@ DEFINE_STAT(STAT_D3D11FreeUniformBufferMemory);
 DEFINE_STAT(STAT_D3D11NumFreeUniformBuffers);
 DEFINE_STAT(STAT_D3D11NumBoundShaderState);
 
-DEFINE_STAT(STAT_D3D11CommitResourceTables);
-DEFINE_STAT(STAT_D3D11CacheResourceTables);
-DEFINE_STAT(STAT_D3D11CacheResourceTableCalls);
-DEFINE_STAT(STAT_D3D11SetShaderTextureTime);
-DEFINE_STAT(STAT_D3D11SetShaderTextureCalls);
-DEFINE_STAT(STAT_D3D11SetTextureInTableCalls);
-
 #undef LOCTEXT_NAMESPACE
diff --git a/Engine/Source/Runtime/Windows/D3D11RHI/Private/D3D11Viewport.cpp b/Engine/Source/Runtime/Windows/D3D11RHI/Private/D3D11Viewport.cpp
--- a/Engine/Source/Runtime/Windows/D3D11RHI/Private/D3D11Viewport.cpp
+++ b/Engine/Source/Runtime/Windows/D3D11RHI/Private/D3D11Viewport.cpp
@@ -534,21 +534,6 @@ void FD3D11DynamicRHI::RHIEndDrawingViewport(FViewportRHIParamRef ViewportRHI,bo
 	StateCache.SetDomainShader(nullptr);
 	StateCache.SetGeometryShader(nullptr);
 	// Compute Shader is set to NULL after each Dispatch call, so no need to clear it here
-#if STATS
-	SET_CYCLE_COUNTER(STAT_D3D11CommitResourceTables, CommitResourceTableCycles);
-	SET_CYCLE_COUNTER(STAT_D3D11CacheResourceTables, CacheResourceTableCycles);
-	SET_CYCLE_COUNTER(STAT_D3D11SetShaderTextureTime, SetShaderTextureCycles);
-	INC_DWORD_STAT_BY(STAT_D3D11SetShaderTextureCalls, SetShaderTextureCalls);
-	INC_DWORD_STAT_BY(STAT_D3D11CacheResourceTableCalls, CacheResourceTableCalls);
-	INC_DWORD_STAT_BY(STAT_D3D11SetTextureInTableCalls, SetTextureInTableCalls);
-#endif
-
-	CommitResourceTableCycles = 0;
-	CacheResourceTableCalls = 0;
-	CacheResourceTableCycles = 0;
-	SetShaderTextureCycles = 0;
-	SetShaderTextureCalls = 0;
-	SetTextureInTableCalls = 0;
 
 	bool bNativelyPresented = Viewport->Present(bLockToVsync);
 
