UE4 - RHI Thread implementation for XB1

diff --git a/Engine/Source/Runtime/Core/Public/HAL/Platform.h b/Engine/Source/Runtime/Core/Public/HAL/Platform.h
--- a/Engine/Source/Runtime/Core/Public/HAL/Platform.h
+++ b/Engine/Source/Runtime/Core/Public/HAL/Platform.h
@@ -242,6 +242,10 @@
 	#define PLATFORM_SUPPORTS_RHI_THREAD		0
 #endif
 
+#ifndef PLATFORM_HAS_THREADSAFE_RHIGetRenderQueryResult
+	#define PLATFORM_HAS_THREADSAFE_RHIGetRenderQueryResult		0
+#endif
+
 #ifndef PLATFORM_SUPPORTS_MULTIPLE_NATIVE_WINDOWS
 	#define PLATFORM_SUPPORTS_MULTIPLE_NATIVE_WINDOWS	1
 #endif
diff --git a/Engine/Source/Runtime/Launch/Private/LaunchEngineLoop.cpp b/Engine/Source/Runtime/Launch/Private/LaunchEngineLoop.cpp
--- a/Engine/Source/Runtime/Launch/Private/LaunchEngineLoop.cpp
+++ b/Engine/Source/Runtime/Launch/Private/LaunchEngineLoop.cpp
@@ -1308,6 +1308,18 @@ int32 FEngineLoop::PreInit( const TCHAR* CmdLine )
 
 	if (GUseThreadedRendering)
 	{
+#if PLATFORM_SUPPORTS_RHI_THREAD
+		const bool DefaultUseRHIThread = false;
+		GUseRHIThread = DefaultUseRHIThread;
+		if (FParse::Param(FCommandLine::Get(),TEXT("rhithread")))
+		{
+			GUseRHIThread = true;
+		}
+		else if (FParse::Param(FCommandLine::Get(),TEXT("norhithread")))
+		{
+			GUseRHIThread = false;
+		}
+#endif
 		StartRenderingThread();
 	}
 	
diff --git a/Engine/Source/Runtime/RHI/Private/RHICommandList.cpp b/Engine/Source/Runtime/RHI/Private/RHICommandList.cpp
--- a/Engine/Source/Runtime/RHI/Private/RHICommandList.cpp
+++ b/Engine/Source/Runtime/RHI/Private/RHICommandList.cpp
@@ -17,8 +17,9 @@ DECLARE_DWORD_COUNTER_STAT(TEXT("Immed. Command count"), STAT_ImmedCmdListCount,
 Requirements for RHI thread
 * Microresources (those in RHIStaticStates.h) need to be able to be created by any thread at any time and be able to work with a radically simplified rhi resource lifecycle. CreateSamplerState, CreateRasterizerState, CreateDepthStencilState, CreateBlendState
 * CreateUniformBuffer needs to be threadsafe
-* GetRenderQueryResult needs to be threadsafe, but perhaps this can be supported if it can't be
+* GetRenderQueryResult should be threadsafe, if not, then PLATFORM_HAS_THREADSAFE_RHIGetRenderQueryResult should be 1
 * AdvanceFrameForGetViewportBackBuffer needs be added as an RHI method and this needs to work with GetViewportBackBuffer to give the render thread the right back buffer even though many commands relating to the beginning and end of the frame are queued.
+* ResetRenderQuery does not exist; this stuff should be done in BeginQuery
 
 ***/
 #endif
@@ -348,6 +349,16 @@ void FRHICommandListExecutor::CheckNoOutstandingCmdLists()
 	check(GRHICommandList.OutstandingCmdListCount.GetValue() == 1); // else we are attempting to delete resources while there is still a live cmdlist (other than the immediate cmd list) somewhere.
 }
 
+bool FRHICommandListExecutor::IsRHIThreadActive()
+{
+	checkSlow(IsInRenderingThread());
+	if (RHIThreadTask.GetReference() && RHIThreadTask->IsComplete())
+	{
+		RHIThreadTask = nullptr;
+	}
+	return !!RHIThreadTask.GetReference();
+}
+
 FGraphEventRef FRHICommandListExecutor::RHIThreadFence()
 {
 	check(IsInRenderingThread() && GRHIThread);
@@ -355,6 +366,7 @@ FGraphEventRef FRHICommandListExecutor::RHIThreadFence()
 		QUICK_SCOPE_CYCLE_COUNTER(STAT_RHIThreadFence_Dispatch);
 		FRHICommandListExecutor::GetImmediateCommandList().ImmediateFlush(EImmediateFlushType::DispatchToRHIThread);
 	}
+	// this isn't quite right, but is probably fine for current uses. If there is a RenderThreadSublistDispatchTask, then we need that to be the fence, however, then that RT thread task needs to "do not complete until" or something so we can track the actual completion.
 	if (RHIThreadTask.GetReference() && RHIThreadTask->IsComplete())
 	{
 		RHIThreadTask = nullptr;
diff --git a/Engine/Source/Runtime/RHI/Public/RHICommandList.h b/Engine/Source/Runtime/RHI/Public/RHICommandList.h
--- a/Engine/Source/Runtime/RHI/Public/RHICommandList.h
+++ b/Engine/Source/Runtime/RHI/Public/RHICommandList.h
@@ -1094,6 +1094,7 @@ struct FRHICommandEndRenderQuery : public FRHICommand<FRHICommandEndRenderQuery>
 	}
 };
 
+#if !PLATFORM_SUPPORTS_RHI_THREAD
 struct FRHICommandResetRenderQuery : public FRHICommand<FRHICommandResetRenderQuery>
 {
 	FRenderQueryRHIParamRef RenderQuery;
@@ -1107,6 +1108,7 @@ struct FRHICommandResetRenderQuery : public FRHICommand<FRHICommandResetRenderQu
 		ResetRenderQuery_Internal(RenderQuery);
 	}
 };
+#endif
 
 struct FRHICommandBeginScene : public FRHICommand<FRHICommandBeginScene>
 {
@@ -1695,6 +1697,7 @@ public:
 		}
 		new (AllocCommand<FRHICommandEndRenderQuery>()) FRHICommandEndRenderQuery(RenderQuery);
 	}
+#if !PLATFORM_SUPPORTS_RHI_THREAD
 	FORCEINLINE_DEBUGGABLE void ResetRenderQuery(FRenderQueryRHIParamRef RenderQuery)
 	{
 		if (Bypass())
@@ -1704,6 +1707,7 @@ public:
 		}
 		new (AllocCommand<FRHICommandResetRenderQuery>()) FRHICommandResetRenderQuery(RenderQuery);
 	}
+#endif
 
 #if PLATFORM_SUPPORTS_RHI_THREAD
 	FORCEINLINE_DEBUGGABLE void BeginScene()
@@ -1841,6 +1845,7 @@ public:
 #endif
 	}
 	static void CheckNoOutstandingCmdLists();
+	static bool IsRHIThreadActive();
 
 private:
 
diff --git a/Engine/Source/Runtime/RHI/Public/RHIMethods.h b/Engine/Source/Runtime/RHI/Public/RHIMethods.h
--- a/Engine/Source/Runtime/RHI/Public/RHIMethods.h
+++ b/Engine/Source/Runtime/RHI/Public/RHIMethods.h
@@ -1270,11 +1270,13 @@ DEFINE_RHIMETHOD_GLOBAL_1(
 	return,return new FRHIRenderQuery();
 	);
 
+#if !PLATFORM_SUPPORTS_RHI_THREAD
 DEFINE_RHIMETHOD_CMDLIST_1(
 	void,ResetRenderQuery,
 	FRenderQueryRHIParamRef,RenderQuery,
 	,
 	);
+#endif
 
 DEFINE_RHIMETHOD_CMDLIST_1(
 	void,BeginRenderQuery,
@@ -1286,8 +1288,7 @@ DEFINE_RHIMETHOD_CMDLIST_1(
 	FRenderQueryRHIParamRef,RenderQuery,
 	,
 	);
-
-#if PLATFORM_SUPPORTS_RHI_THREAD
+#if PLATFORM_HAS_THREADSAFE_RHIGetRenderQueryResult
 	DEFINE_RHIMETHOD_GLOBALTHREADSAFE_3(
 		bool,GetRenderQueryResult,
 		FRenderQueryRHIParamRef,RenderQuery,
@@ -1295,6 +1296,17 @@ DEFINE_RHIMETHOD_CMDLIST_1(
 		bool,bWait,
 		return,return true;
 	);
+#else
+	DEFINE_RHIMETHOD_3(
+		bool,GetRenderQueryResult,
+		FRenderQueryRHIParamRef,RenderQuery,
+		uint64&,OutResult,
+		bool,bWait,
+		return,return true;
+	);
+#endif
+
+#if PLATFORM_SUPPORTS_RHI_THREAD
 	DEFINE_RHIMETHOD_CMDLIST_2(
 		void,BeginDrawingViewport,
 		FViewportRHIParamRef,Viewport,
@@ -1346,13 +1358,6 @@ DEFINE_RHIMETHOD_CMDLIST_1(
 		);
 
 #else
-	DEFINE_RHIMETHOD_3(
-		bool,GetRenderQueryResult,
-		FRenderQueryRHIParamRef,RenderQuery,
-		uint64&,OutResult,
-		bool,bWait,
-		return,return true;
-	);
 	DEFINE_RHIMETHOD_2(
 		void,BeginDrawingViewport,
 		FViewportRHIParamRef,Viewport,
diff --git a/Engine/Source/Runtime/Renderer/Private/SceneOcclusion.cpp b/Engine/Source/Runtime/Renderer/Private/SceneOcclusion.cpp
--- a/Engine/Source/Runtime/Renderer/Private/SceneOcclusion.cpp
+++ b/Engine/Source/Runtime/Renderer/Private/SceneOcclusion.cpp
@@ -80,8 +80,10 @@ void FRenderQueryPool::ReleaseQuery(FRHICommandListImmediate& RHICmdList, FRende
 			// Return it to the pool.
 			Queries.Add( Query );
 
+#if !PLATFORM_SUPPORTS_RHI_THREAD
 			// Tell RHI we don't need the result anymore.
 			RHICmdList.ResetRenderQuery(Query);
+#endif
 		}
 
 		// De-ref without deleting.
diff --git a/Engine/Source/Runtime/Windows/D3D11RHI/Private/D3D11Commands.cpp b/Engine/Source/Runtime/Windows/D3D11RHI/Private/D3D11Commands.cpp
--- a/Engine/Source/Runtime/Windows/D3D11RHI/Private/D3D11Commands.cpp
+++ b/Engine/Source/Runtime/Windows/D3D11RHI/Private/D3D11Commands.cpp
@@ -54,6 +54,68 @@ DECLARE_ISBOUNDSHADER(ComputeShader)
 	#include <GPUPerfAPI/Gpa.h>
 #endif
 
+#if PLATFORM_SUPPORTS_RHI_THREAD
+void FD3D11DynamicRHI::SetupRecursiveResources()
+{
+	FRHICommandList_RecursiveHazardous RHICmdList;
+	extern int32 GCreateShadersOnLoad;
+	TGuardValue<int32> Guard(GCreateShadersOnLoad, 1);
+	auto ShaderMap = GetGlobalShaderMap(GMaxRHIFeatureLevel);
+	TShaderMapRef<TOneColorVS<true> > VertexShader(ShaderMap);
+	GD3D11Vector4VertexDeclaration.InitRHI();
+
+	for (int32 NumBuffers = 1; NumBuffers <= MaxSimultaneousRenderTargets; NumBuffers++)
+	{
+		FOneColorPS* PixelShader = NULL;
+
+		// Set the shader to write to the appropriate number of render targets
+		// On AMD PC hardware, outputting to a color index in the shader without a matching render target set has a significant performance hit
+		if (NumBuffers <= 1)
+		{
+			TShaderMapRef<TOneColorPixelShaderMRT<1> > MRTPixelShader(ShaderMap);
+			PixelShader = *MRTPixelShader;
+		}
+		else if (NumBuffers == 2)
+		{
+			TShaderMapRef<TOneColorPixelShaderMRT<2> > MRTPixelShader(ShaderMap);
+			PixelShader = *MRTPixelShader;
+		}
+		else if (NumBuffers== 3)
+		{
+			TShaderMapRef<TOneColorPixelShaderMRT<3> > MRTPixelShader(ShaderMap);
+			PixelShader = *MRTPixelShader;
+		}
+		else if (NumBuffers == 4)
+		{
+			TShaderMapRef<TOneColorPixelShaderMRT<4> > MRTPixelShader(ShaderMap);
+			PixelShader = *MRTPixelShader;
+		}
+		else if (NumBuffers == 5)
+		{
+			TShaderMapRef<TOneColorPixelShaderMRT<5> > MRTPixelShader(ShaderMap);
+			PixelShader = *MRTPixelShader;
+		}
+		else if (NumBuffers == 6)
+		{
+			TShaderMapRef<TOneColorPixelShaderMRT<6> > MRTPixelShader(ShaderMap);
+			PixelShader = *MRTPixelShader;
+		}
+		else if (NumBuffers == 7)
+		{
+			TShaderMapRef<TOneColorPixelShaderMRT<7> > MRTPixelShader(ShaderMap);
+			PixelShader = *MRTPixelShader;
+		}
+		else if (NumBuffers == 8)
+		{
+			TShaderMapRef<TOneColorPixelShaderMRT<8> > MRTPixelShader(ShaderMap);
+			PixelShader = *MRTPixelShader;
+		}
+
+		SetGlobalBoundShaderState(RHICmdList, GMaxRHIFeatureLevel, GD3D11ClearMRTBoundShaderState[NumBuffers - 1], GD3D11Vector4VertexDeclaration.VertexDeclarationRHI, *VertexShader, PixelShader);
+	}
+}
+#endif
+
 void FD3D11DynamicRHI::RHIGpuTimeBegin(uint32 Hash, bool bCompute)
 {
 	#if WITH_GPA
@@ -1039,6 +1101,9 @@ void FD3D11DynamicRHI::RHIBeginRenderQuery(FRenderQueryRHIParamRef QueryRHI)
 
 	if(Query->QueryType == RQT_Occlusion)
 	{
+#if PLATFORM_SUPPORTS_RHI_THREAD
+		Query->bResultIsCached = false;
+#endif
 		Direct3DDeviceIMContext->Begin(Query->Resource);
 	}
 	else
@@ -1052,6 +1117,9 @@ void FD3D11DynamicRHI::RHIEndRenderQuery(FRenderQueryRHIParamRef QueryRHI)
 	GRHICommandList.Verify();
 
 	DYNAMIC_CAST_D3D11RESOURCE(OcclusionQuery,Query);
+#if PLATFORM_SUPPORTS_RHI_THREAD
+	Query->bResultIsCached = false; // for occlusion queries, this is redundant with the one in begin
+#endif
 	Direct3DDeviceIMContext->End(Query->Resource);
 
 	//@todo - d3d debug spews warnings about OQ's that are being issued but not polled, need to investigate
diff --git a/Engine/Source/Runtime/Windows/D3D11RHI/Private/D3D11Query.cpp b/Engine/Source/Runtime/Windows/D3D11RHI/Private/D3D11Query.cpp
--- a/Engine/Source/Runtime/Windows/D3D11RHI/Private/D3D11Query.cpp
+++ b/Engine/Source/Runtime/Windows/D3D11RHI/Private/D3D11Query.cpp
@@ -29,12 +29,14 @@ FRenderQueryRHIRef FD3D11DynamicRHI::RHICreateRenderQuery(ERenderQueryType Query
 	return new FD3D11OcclusionQuery(Query, QueryType);
 }
 
+#if !PLATFORM_SUPPORTS_RHI_THREAD
 void FD3D11DynamicRHI::RHIResetRenderQuery(FRenderQueryRHIParamRef QueryRHI)
 {
 	DYNAMIC_CAST_D3D11RESOURCE(OcclusionQuery,Query);
 
 	Query->bResultIsCached = false;
 }
+#endif
 
 bool FD3D11DynamicRHI::RHIGetRenderQueryResult(FRenderQueryRHIParamRef QueryRHI,uint64& OutResult,bool bWait)
 {
diff --git a/Engine/Source/Runtime/Windows/D3D11RHI/Private/D3D11RHIPrivate.h b/Engine/Source/Runtime/Windows/D3D11RHI/Private/D3D11RHIPrivate.h
--- a/Engine/Source/Runtime/Windows/D3D11RHI/Private/D3D11RHIPrivate.h
+++ b/Engine/Source/Runtime/Windows/D3D11RHI/Private/D3D11RHIPrivate.h
@@ -317,6 +317,10 @@ public:
 	/** If it hasn't been initialized yet, initializes the D3D device. */
 	virtual void InitD3DDevice();
 
+#if PLATFORM_SUPPORTS_RHI_THREAD
+	void SetupRecursiveResources();
+#endif
+
 	// FDynamicRHI interface.
 	virtual void Init() override;
 	virtual void Shutdown() override;
diff --git a/Engine/Source/Runtime/Windows/D3D11RHI/Private/D3D11State.cpp b/Engine/Source/Runtime/Windows/D3D11RHI/Private/D3D11State.cpp
--- a/Engine/Source/Runtime/Windows/D3D11RHI/Private/D3D11State.cpp
+++ b/Engine/Source/Runtime/Windows/D3D11RHI/Private/D3D11State.cpp
@@ -8,10 +8,20 @@
 
 // Cache of Sampler States; we store pointers to both as we don't want the TMap to be artificially
 // modifying ref counts if not needed; so we manage that ourselves
+#if PLATFORM_SUPPORTS_RHI_THREAD
+FCriticalSection GSamplerStateCacheLock;
+#define LOCK_GSamplerStateCache (1)
+#else
+#define LOCK_GSamplerStateCache (0)
+#endif
 static TMap<ID3D11SamplerState*, FD3D11SamplerState*> GSamplerStateCache;
 
+
 void EmptyD3DSamplerStateCache()
 {
+#if LOCK_GSamplerStateCache
+	FScopeLock Lock(&GSamplerStateCacheLock);
+#endif
 	for (auto Iter = GSamplerStateCache.CreateIterator(); Iter; ++Iter )
 	{
 		auto* State = Iter.Value();
@@ -172,6 +182,11 @@ FSamplerStateRHIRef FD3D11DynamicRHI::RHICreateSamplerState(const FSamplerStateI
 	SamplerDesc.BorderColor[3] = LinearBorderColor.A;
 	SamplerDesc.ComparisonFunc = TranslateSamplerCompareFunction(Initializer.SamplerComparisonFunction);
 
+#if LOCK_GSamplerStateCache
+	QUICK_SCOPE_CYCLE_COUNTER(FD3D11DynamicRHI_RHICreateSamplerState_LockAndCreate);
+	FScopeLock Lock(&GSamplerStateCacheLock);
+#endif
+
 	// D3D11 will return the same pointer if the particular state description was already created
 	TRefCountPtr<ID3D11SamplerState> SamplerStateHandle;
 	VERIFYD3D11RESULT(Direct3DDevice->CreateSamplerState(&SamplerDesc, SamplerStateHandle.GetInitReference()));
diff --git a/Engine/Source/Runtime/Windows/D3D11RHI/Private/D3D11UniformBuffer.cpp b/Engine/Source/Runtime/Windows/D3D11RHI/Private/D3D11UniformBuffer.cpp
--- a/Engine/Source/Runtime/Windows/D3D11RHI/Private/D3D11UniformBuffer.cpp
+++ b/Engine/Source/Runtime/Windows/D3D11RHI/Private/D3D11UniformBuffer.cpp
@@ -87,6 +87,10 @@ void UniformBufferBeginFrame()
 
 static bool IsPoolingEnabled()
 {
+	if (GRHIThread && IsInRenderingThread() && GRHICommandList.IsRHIThreadActive())
+	{
+		return false; // we can't currently use pooling if the RHI thread is active. 
+	}
 	static const auto CVar = IConsoleManager::Get().FindTConsoleVariableDataInt(TEXT("r.UniformBufferPooling"));
 	int32 CVarValue = CVar->GetValueOnRenderThread();
 	return CVarValue != 0;
@@ -94,7 +98,7 @@ static bool IsPoolingEnabled()
 
 FUniformBufferRHIRef FD3D11DynamicRHI::RHICreateUniformBuffer(const void* Contents,const FRHIUniformBufferLayout& Layout,EUniformBufferUsage Usage)
 {
-	check(IsInRenderingThread());
+	check(IsInRenderingThread() || IsInRHIThread());
 
 	FD3D11UniformBuffer* NewUniformBuffer = nullptr;
 	const uint32 NumBytes = Layout.ConstantBufferSize;
diff --git a/Engine/Source/Runtime/Windows/D3D11RHI/Private/D3D11Viewport.cpp b/Engine/Source/Runtime/Windows/D3D11RHI/Private/D3D11Viewport.cpp
--- a/Engine/Source/Runtime/Windows/D3D11RHI/Private/D3D11Viewport.cpp
+++ b/Engine/Source/Runtime/Windows/D3D11RHI/Private/D3D11Viewport.cpp
@@ -579,6 +579,12 @@ bool FD3D11DynamicRHI::RHIIsDrawingViewport()
 	return DrawingViewport != NULL;
 }
 
+#if PLATFORM_SUPPORTS_RHI_THREAD
+void FD3D11DynamicRHI::RHIAdvanceFrameForGetViewportBackBuffer()
+{
+}
+#endif
+
 FTexture2DRHIRef FD3D11DynamicRHI::RHIGetViewportBackBuffer(FViewportRHIParamRef ViewportRHI)
 {
 	DYNAMIC_CAST_D3D11RESOURCE(Viewport,Viewport);
