Fix D3D11 leaking resources when deferred deletion is enabled

diff --git a/Engine/Source/Runtime/RHI/Private/RHI.cpp b/Engine/Source/Runtime/RHI/Private/RHI.cpp
--- a/Engine/Source/Runtime/RHI/Private/RHI.cpp
+++ b/Engine/Source/Runtime/RHI/Private/RHI.cpp
@@ -74,11 +74,6 @@ DECLARE_CYCLE_STAT(TEXT("Delete Resources"), STAT_DeleteResources, STATGROUP_RHI
 
 void FRHIResource::FlushPendingDeletes()
 {
-	if (GIsRequestingExit) 
-	{ 
-		// Leak these because we're terrible people. 
-		return; 
-	}
 	SCOPE_CYCLE_COUNTER(STAT_DeleteResources);
 
 	check(IsInRenderingThread());
diff --git a/Engine/Source/Runtime/Windows/D3D11RHI/Private/D3D11Device.cpp b/Engine/Source/Runtime/Windows/D3D11RHI/Private/D3D11Device.cpp
--- a/Engine/Source/Runtime/Windows/D3D11RHI/Private/D3D11Device.cpp
+++ b/Engine/Source/Runtime/Windows/D3D11RHI/Private/D3D11Device.cpp
@@ -414,9 +414,6 @@ void FD3D11DynamicRHI::CleanupD3DDevice()
 		DynamicVB = NULL;
 		DynamicIB = NULL;
 
-		ReleasePooledUniformBuffers();
-		ReleasePooledTextures();
-
 		// Release references to bound uniform buffers.
 		for (int32 Frequency = 0; Frequency < SF_NumFrequencies; ++Frequency)
 		{
@@ -429,6 +426,12 @@ void FD3D11DynamicRHI::CleanupD3DDevice()
 		// Release the device and its IC
 		StateCache.SetContext(nullptr);
 
+		// Flush all pending deletes before destroying the device.
+		FRHIResource::FlushPendingDeletes();
+
+		ReleasePooledUniformBuffers();
+		ReleasePooledTextures();
+
 		// When running with D3D debug, clear state and flush the device to get rid of spurious live objects in D3D11's report.
 		if (D3D11RHI_ShouldCreateWithD3DDebug())
 		{
