UE4 - Fix ambiguity

diff --git a/Engine/Source/Runtime/IOS/MetalRHI/Private/MetalCommands.cpp b/Engine/Source/Runtime/IOS/MetalRHI/Private/MetalCommands.cpp
--- a/Engine/Source/Runtime/IOS/MetalRHI/Private/MetalCommands.cpp
+++ b/Engine/Source/Runtime/IOS/MetalRHI/Private/MetalCommands.cpp
@@ -367,14 +367,14 @@ void FMetalDynamicRHI::RHIDiscardRenderTargets(bool Depth, bool Stencil, uint32
 void FMetalDynamicRHI::RHISetRenderTargetsAndClear(const FRHISetRenderTargetsInfo& RenderTargetsInfo)
 {
 //@todo-rco: TEMP
-	RHISetRenderTargets(RenderTargetsInfo.NumColorRenderTargets,
+	this->RHISetRenderTargets(RenderTargetsInfo.NumColorRenderTargets,
 		RenderTargetsInfo.ColorRenderTarget,
 		RenderTargetsInfo.DepthStencilTarget,
 		0,
 		nullptr);
 	if (RenderTargetsInfo.bClearColor || RenderTargetsInfo.bClearStencil || RenderTargetsInfo.bClearDepth)
 	{
-		RHIClearMRT(RenderTargetsInfo.bClearColor, RenderTargetsInfo.NumColorRenderTargets, RenderTargetsInfo.ClearColors, RenderTargetsInfo.bClearDepth, RenderTargetsInfo.DepthClearValue, RenderTargetsInfo.bClearStencil, RenderTargetsInfo.StencilClearValue, FIntRect());
+		this->RHIClearMRT(RenderTargetsInfo.bClearColor, RenderTargetsInfo.NumColorRenderTargets, RenderTargetsInfo.ClearColors, RenderTargetsInfo.bClearDepth, RenderTargetsInfo.DepthClearValue, RenderTargetsInfo.bClearStencil, RenderTargetsInfo.StencilClearValue, FIntRect());
 	}
 }
 
diff --git a/Engine/Source/Runtime/OpenGLDrv/Private/OpenGLCommands.cpp b/Engine/Source/Runtime/OpenGLDrv/Private/OpenGLCommands.cpp
--- a/Engine/Source/Runtime/OpenGLDrv/Private/OpenGLCommands.cpp
+++ b/Engine/Source/Runtime/OpenGLDrv/Private/OpenGLCommands.cpp
@@ -1635,14 +1635,14 @@ void FOpenGLDynamicRHI::RHIDiscardRenderTargets(bool Depth, bool Stencil, uint32
 
 void FOpenGLDynamicRHI::RHISetRenderTargetsAndClear(const FRHISetRenderTargetsInfo& RenderTargetsInfo)
 {
-	RHISetRenderTargets(RenderTargetsInfo.NumColorRenderTargets,
+	this->RHISetRenderTargets(RenderTargetsInfo.NumColorRenderTargets,
 		RenderTargetsInfo.ColorRenderTarget,
 		RenderTargetsInfo.DepthStencilTarget,
 		0,
 		nullptr);
 	if (RenderTargetsInfo.bClearColor || RenderTargetsInfo.bClearStencil || RenderTargetsInfo.bClearDepth)
 	{
-		RHIClearMRT(RenderTargetsInfo.bClearColor, RenderTargetsInfo.NumColorRenderTargets, RenderTargetsInfo.ClearColors, RenderTargetsInfo.bClearDepth, RenderTargetsInfo.DepthClearValue, RenderTargetsInfo.bClearStencil, RenderTargetsInfo.StencilClearValue, FIntRect());
+		this->RHIClearMRT(RenderTargetsInfo.bClearColor, RenderTargetsInfo.NumColorRenderTargets, RenderTargetsInfo.ClearColors, RenderTargetsInfo.bClearDepth, RenderTargetsInfo.DepthClearValue, RenderTargetsInfo.bClearStencil, RenderTargetsInfo.StencilClearValue, FIntRect());
 	}
 }
 
diff --git a/Engine/Source/Runtime/Windows/D3D11RHI/Private/D3D11Commands.cpp b/Engine/Source/Runtime/Windows/D3D11RHI/Private/D3D11Commands.cpp
--- a/Engine/Source/Runtime/Windows/D3D11RHI/Private/D3D11Commands.cpp
+++ b/Engine/Source/Runtime/Windows/D3D11RHI/Private/D3D11Commands.cpp
@@ -1005,14 +1005,14 @@ void FD3D11DynamicRHI::RHIDiscardRenderTargets(bool Depth, bool Stencil, uint32
 
 void FD3D11DynamicRHI::RHISetRenderTargetsAndClear(const FRHISetRenderTargetsInfo& RenderTargetsInfo)
 {
-	RHISetRenderTargets(RenderTargetsInfo.NumColorRenderTargets,
+	this->RHISetRenderTargets(RenderTargetsInfo.NumColorRenderTargets,
 		RenderTargetsInfo.ColorRenderTarget,
 		RenderTargetsInfo.DepthStencilTarget,
 		0,
 		nullptr);
 	if (RenderTargetsInfo.bClearColor || RenderTargetsInfo.bClearStencil || RenderTargetsInfo.bClearDepth)
 	{
-		RHIClearMRT(RenderTargetsInfo.bClearColor, RenderTargetsInfo.NumColorRenderTargets, RenderTargetsInfo.ClearColors, RenderTargetsInfo.bClearDepth, RenderTargetsInfo.DepthClearValue, RenderTargetsInfo.bClearStencil, RenderTargetsInfo.StencilClearValue, FIntRect());
+		this->RHIClearMRT(RenderTargetsInfo.bClearColor, RenderTargetsInfo.NumColorRenderTargets, RenderTargetsInfo.ClearColors, RenderTargetsInfo.bClearDepth, RenderTargetsInfo.DepthClearValue, RenderTargetsInfo.bClearStencil, RenderTargetsInfo.StencilClearValue, FIntRect());
 	}
 }
 
