UE4 - use a local bound shader state for long GPU task for RHI thread compatibility

diff --git a/Engine/Source/Runtime/Windows/D3D11RHI/Private/D3D11RHI.cpp b/Engine/Source/Runtime/Windows/D3D11RHI/Private/D3D11RHI.cpp
--- a/Engine/Source/Runtime/Windows/D3D11RHI/Private/D3D11RHI.cpp
+++ b/Engine/Source/Runtime/Windows/D3D11RHI/Private/D3D11RHI.cpp
@@ -213,7 +213,7 @@ void FD3D11DynamicRHI::IssueLongGPUTask()
 			TShaderMapRef<TOneColorVS<true> > VertexShader(ShaderMap);
 			TShaderMapRef<FLongGPUTaskPS> PixelShader(ShaderMap);
 
-			SetGlobalBoundShaderState(RHICmdList, GMaxRHIFeatureLevel, LongGPUTaskBoundShaderState, GD3D11Vector4VertexDeclaration.VertexDeclarationRHI, *VertexShader, *PixelShader);
+			RHICmdList.SetLocalBoundShaderState(RHICmdList.BuildLocalBoundShaderState(GD3D11Vector4VertexDeclaration.VertexDeclarationRHI, VertexShader->GetVertexShader(), FHullShaderRHIRef(), FDomainShaderRHIRef(), PixelShader->GetPixelShader(), FGeometryShaderRHIRef()));
 
 			// Draw a fullscreen quad
 			FVector4 Vertices[4];
