Memory leak fix for state cache.

diff --git a/Engine/Source/Runtime/Windows/D3D11RHI/Private/D3D11StateCache.cpp b/Engine/Source/Runtime/Windows/D3D11RHI/Private/D3D11StateCache.cpp
--- a/Engine/Source/Runtime/Windows/D3D11RHI/Private/D3D11StateCache.cpp
+++ b/Engine/Source/Runtime/Windows/D3D11RHI/Private/D3D11StateCache.cpp
@@ -295,7 +295,17 @@ void FD3D11StateCacheBase::ClearState()
 
 #if D3D11_ALLOW_STATE_CACHE
 	// Shader Resource View State Cache
-	FMemory::Memzero(CurrentShaderResourceViews, sizeof(CurrentShaderResourceViews));
+	for (uint32 ShaderFrequency = 0; ShaderFrequency < SF_MAX; ShaderFrequency++ )
+	{
+		for (uint32 Index = 0; Index < D3D11_COMMONSHADER_INPUT_RESOURCE_SLOT_COUNT; Index++)
+		{
+			if(CurrentShaderResourceViews[ShaderFrequency][Index])
+			{
+				CurrentShaderResourceViews[ShaderFrequency][Index]->Release();
+				CurrentShaderResourceViews[ShaderFrequency][Index] = NULL;
+			}
+		}
+	}
 
 	// Rasterizer State Cache
 	CurrentRasterizerState = nullptr;
diff --git a/Engine/Source/Runtime/Windows/D3D11RHI/Private/D3D11StateCachePrivate.h b/Engine/Source/Runtime/Windows/D3D11RHI/Private/D3D11StateCachePrivate.h
--- a/Engine/Source/Runtime/Windows/D3D11RHI/Private/D3D11StateCachePrivate.h
+++ b/Engine/Source/Runtime/Windows/D3D11RHI/Private/D3D11StateCachePrivate.h
@@ -858,6 +858,7 @@ public:
 	FD3D11StateCacheBase()
 		: Direct3DDeviceIMContext(nullptr)
 	{
+		FMemory::Memzero(CurrentShaderResourceViews, sizeof(CurrentShaderResourceViews));
 	}
 
 	void Init(ID3D11DeviceContext* InDeviceContext, bool bInAlwaysSetIndexBuffers = false )
