UE4 - Remove xbox from PC version of split d3d11rhi

diff --git a/Engine/Source/Runtime/Windows/D3D11RHI/Private/D3D11Commands.cpp b/Engine/Source/Runtime/Windows/D3D11RHI/Private/D3D11Commands.cpp
--- a/Engine/Source/Runtime/Windows/D3D11RHI/Private/D3D11Commands.cpp
+++ b/Engine/Source/Runtime/Windows/D3D11RHI/Private/D3D11Commands.cpp
@@ -48,7 +48,7 @@ DECLARE_ISBOUNDSHADER(ComputeShader)
 #define VALIDATE_BOUND_SHADER(s)
 #endif
 
-#define WITH_GPA (!PLATFORM_XBOXONE)
+#define WITH_GPA (1)
 #if WITH_GPA
 	#define GPA_WINDOWS 1
 	#include <GPUPerfAPI/Gpa.h>
@@ -578,13 +578,6 @@ void FD3D11DynamicRHI::RHISetShaderUniformBuffer(FVertexShaderRHIParamRef Vertex
 {
 	VALIDATE_BOUND_SHADER(VertexShader);
 	DYNAMIC_CAST_D3D11RESOURCE(UniformBuffer,Buffer);
-#if PLATFORM_XBOXONE
-	if (Buffer && Buffer->RingAllocation.IsValid())
-	{
-		StateCache.SetDynamicConstantBuffer<SF_Vertex>(BufferIndex, Buffer->RingAllocation);
-	}
-	else
-#endif
 	{
 		ID3D11Buffer* ConstantBuffer = Buffer ? Buffer->Resource : NULL;
 		StateCache.SetConstantBuffer<SF_Vertex>(ConstantBuffer, BufferIndex);
@@ -598,13 +591,6 @@ void FD3D11DynamicRHI::RHISetShaderUniformBuffer(FHullShaderRHIParamRef HullShad
 {
 	VALIDATE_BOUND_SHADER(HullShader);
 	DYNAMIC_CAST_D3D11RESOURCE(UniformBuffer,Buffer);
-#if PLATFORM_XBOXONE
-	if (Buffer && Buffer->RingAllocation.IsValid())
-	{
-		StateCache.SetDynamicConstantBuffer<SF_Hull>(BufferIndex, Buffer->RingAllocation);
-	}
-	else
-#endif
 	{
 		ID3D11Buffer* ConstantBuffer = Buffer ? Buffer->Resource : NULL;
 		StateCache.SetConstantBuffer<SF_Hull>(ConstantBuffer, BufferIndex);
@@ -618,13 +604,6 @@ void FD3D11DynamicRHI::RHISetShaderUniformBuffer(FDomainShaderRHIParamRef Domain
 {
 	VALIDATE_BOUND_SHADER(DomainShader);
 	DYNAMIC_CAST_D3D11RESOURCE(UniformBuffer,Buffer);
-#if PLATFORM_XBOXONE
-	if (Buffer && Buffer->RingAllocation.IsValid())
-	{
-		StateCache.SetDynamicConstantBuffer<SF_Domain>(BufferIndex, Buffer->RingAllocation);
-	}
-	else
-#endif
 	{
 		ID3D11Buffer* ConstantBuffer = Buffer ? Buffer->Resource : NULL;
 		StateCache.SetConstantBuffer<SF_Domain>(ConstantBuffer, BufferIndex);
@@ -638,13 +617,6 @@ void FD3D11DynamicRHI::RHISetShaderUniformBuffer(FGeometryShaderRHIParamRef Geom
 {
 	VALIDATE_BOUND_SHADER(GeometryShader);
 	DYNAMIC_CAST_D3D11RESOURCE(UniformBuffer,Buffer);
-#if PLATFORM_XBOXONE
-	if (Buffer && Buffer->RingAllocation.IsValid())
-	{
-		StateCache.SetDynamicConstantBuffer<SF_Geometry>(BufferIndex, Buffer->RingAllocation);
-	}
-	else
-#endif
 	{
 		ID3D11Buffer* ConstantBuffer = Buffer ? Buffer->Resource : NULL;
 		StateCache.SetConstantBuffer<SF_Geometry>(ConstantBuffer, BufferIndex);
@@ -659,13 +631,6 @@ void FD3D11DynamicRHI::RHISetShaderUniformBuffer(FPixelShaderRHIParamRef PixelSh
 	VALIDATE_BOUND_SHADER(PixelShader);
 	DYNAMIC_CAST_D3D11RESOURCE(UniformBuffer,Buffer);
 
-#if PLATFORM_XBOXONE
-	if (Buffer && Buffer->RingAllocation.IsValid())
-	{
-		StateCache.SetDynamicConstantBuffer<SF_Pixel>(BufferIndex, Buffer->RingAllocation);
-	}
-	else
-#endif
 	{
 		ID3D11Buffer* ConstantBuffer = Buffer ? Buffer->Resource : NULL;
 		StateCache.SetConstantBuffer<SF_Pixel>(ConstantBuffer, BufferIndex);
@@ -679,13 +644,6 @@ void FD3D11DynamicRHI::RHISetShaderUniformBuffer(FComputeShaderRHIParamRef Compu
 {
 	//VALIDATE_BOUND_SHADER(ComputeShader);
 	DYNAMIC_CAST_D3D11RESOURCE(UniformBuffer,Buffer);
-#if PLATFORM_XBOXONE
-	if (Buffer && Buffer->RingAllocation.IsValid())
-	{
-		StateCache.SetDynamicConstantBuffer<SF_Compute>(BufferIndex, Buffer->RingAllocation);
-	}
-	else
-#endif
 	{
 		ID3D11Buffer* ConstantBuffer = Buffer ? Buffer->Resource : NULL;
 		StateCache.SetConstantBuffer<SF_Compute>(ConstantBuffer, BufferIndex);
diff --git a/Engine/Source/Runtime/Windows/D3D11RHI/Private/D3D11Device.cpp b/Engine/Source/Runtime/Windows/D3D11RHI/Private/D3D11Device.cpp
--- a/Engine/Source/Runtime/Windows/D3D11RHI/Private/D3D11Device.cpp
+++ b/Engine/Source/Runtime/Windows/D3D11RHI/Private/D3D11Device.cpp
@@ -243,7 +243,6 @@ void FD3D11DynamicRHI::RHIGetSupportedResolution( uint32 &Width, uint32 &Height
 		DXGI_ADAPTER_DESC AdapterDesc;
 		VERIFYD3D11RESULT(Adapter->GetDesc(&AdapterDesc));
 	  
-#ifndef PLATFORM_XBOXONE // No need for display mode enumeration on console
 		// Enumerate outputs for this adapter
 		// TODO: Cap at 1 for default output
 		for(uint32 o = 0;o < 1; o++)
@@ -291,7 +290,6 @@ void FD3D11DynamicRHI::RHIGetSupportedResolution( uint32 &Width, uint32 &Height
 
 			delete[] ModeList;
 		}
-#endif // PLATFORM_XBOXONE
 	}
 
 	check(InitializedMode);
diff --git a/Engine/Source/Runtime/Windows/D3D11RHI/Private/D3D11RHIPrivate.h b/Engine/Source/Runtime/Windows/D3D11RHI/Private/D3D11RHIPrivate.h
--- a/Engine/Source/Runtime/Windows/D3D11RHI/Private/D3D11RHIPrivate.h
+++ b/Engine/Source/Runtime/Windows/D3D11RHI/Private/D3D11RHIPrivate.h
@@ -495,10 +495,7 @@ protected:
 	/** Dynamic vertex and index buffers. */
 	TRefCountPtr<FD3D11DynamicBuffer> DynamicVB;
 	TRefCountPtr<FD3D11DynamicBuffer> DynamicIB;
-#if PLATFORM_XBOXONE
-	// New circular buffer system for faster constant uploads.  Avoids CopyResource and speeds things up considerably
-	FD3D11ConstantsRingBuffer DynamicCB;
-#endif
+
 	// State for begin/end draw primitive UP interface.
 	uint32 PendingNumVertices;
 	uint32 PendingVertexDataStride;
diff --git a/Engine/Source/Runtime/Windows/D3D11RHI/Private/D3D11StateCache.cpp b/Engine/Source/Runtime/Windows/D3D11RHI/Private/D3D11StateCache.cpp
--- a/Engine/Source/Runtime/Windows/D3D11RHI/Private/D3D11StateCache.cpp
+++ b/Engine/Source/Runtime/Windows/D3D11RHI/Private/D3D11StateCache.cpp
@@ -185,13 +185,7 @@ void FD3D11StateCacheBase::VerifyCacheState()
 		Direct3DDeviceIMContext->OMGetBlendState(BlendState.GetInitReference(), BlendFactor, &SampleMask);
 
 		check(BlendState.GetReference() == CurrentBlendState);
-#ifdef PLATFORM_XBOXONE
-		//MONO BUG WORKAROUND and a check to make sure it isn't fixed yet.
-		// mono's OMGetBlendState fails to set elements 2 and 3 (they don't change when you call OMGetBlendState)
-		check(FMemory::Memcmp(BlendFactor, CurrentBlendFactor, sizeof(float)*2) == 0);
-#else
 		check(FMemory::Memcmp(BlendFactor, CurrentBlendFactor, sizeof(CurrentBlendFactor)) == 0);
-#endif
 		check(SampleMask == CurrentBlendSampleMask);
 	}
 
@@ -232,10 +226,7 @@ void FD3D11StateCacheBase::VerifyCacheState()
 		for (uint32 Index = 0; Index < D3D11_IA_VERTEX_INPUT_RESOURCE_SLOT_COUNT; Index++)
 		{
 			check(VertexBuffers[Index] == CurrentVertexBuffers[Index].VertexBuffer);
-#ifndef PLATFORM_XBOXONE
-			//Broken on mono.  will return stride = "previous stride" after call to IASetVertexBuffers(0,1,0,0,0);
 			check(Strides[Index] == CurrentVertexBuffers[Index].Stride);
-#endif
 			check(Offsets[Index] == CurrentVertexBuffers[Index].Offset);
 			if (VertexBuffers[Index])
 			{
diff --git a/Engine/Source/Runtime/Windows/D3D11RHI/Private/D3D11UniformBuffer.cpp b/Engine/Source/Runtime/Windows/D3D11RHI/Private/D3D11UniformBuffer.cpp
--- a/Engine/Source/Runtime/Windows/D3D11RHI/Private/D3D11UniformBuffer.cpp
+++ b/Engine/Source/Runtime/Windows/D3D11RHI/Private/D3D11UniformBuffer.cpp
@@ -116,14 +116,6 @@ FUniformBufferRHIRef FD3D11DynamicRHI::RHICreateUniformBuffer(const void* Conten
 			TRefCountPtr<ID3D11Buffer> UniformBufferResource;
 			FRingAllocation RingAllocation;
 
-#if PLATFORM_XBOXONE
-			if (Usage == UniformBuffer_SingleDraw)
-			{
-				RingAllocation = DynamicCB.MapAndFill(GetDeviceContext(),NumBytes,Contents);
-				check(RingAllocation.IsValid());
-			}
-#endif
-
 			if (!RingAllocation.IsValid())
 			{
 				// Find the appropriate bucket based on size
diff --git a/Engine/Source/Runtime/Windows/D3D11RHI/Private/D3D11Viewport.cpp b/Engine/Source/Runtime/Windows/D3D11RHI/Private/D3D11Viewport.cpp
--- a/Engine/Source/Runtime/Windows/D3D11RHI/Private/D3D11Viewport.cpp
+++ b/Engine/Source/Runtime/Windows/D3D11RHI/Private/D3D11Viewport.cpp
@@ -185,12 +185,9 @@ void FD3D11Viewport::Resize(uint32 InSizeX,uint32 InSizeY,bool bInIsFullscreen)
 	{
 		check(BackBuffer->GetRefCount() == 1);
 
-		// IUnknown is GONE from Mono Drivers, and is replaced with IGraphicsUnknown, making this incompatible (no COM)
-#ifndef PLATFORM_XBOXONE
 		checkComRefCount(BackBuffer->GetResource(),1);
 		checkComRefCount(BackBuffer->GetRenderTargetView(0, -1),1);
 		checkComRefCount(BackBuffer->GetShaderResourceView(),1);
-#endif
 	}
 	BackBuffer.SafeRelease();
 
