CVar for setting maximum frame latency on DX11

Merging CL 2337430

//depot/UE4-UT/...

to //depot/UE4/...

diff --git a/Engine/Source/Runtime/Windows/D3D11RHI/Private/D3D11Viewport.cpp b/Engine/Source/Runtime/Windows/D3D11RHI/Private/D3D11Viewport.cpp
--- a/Engine/Source/Runtime/Windows/D3D11RHI/Private/D3D11Viewport.cpp
+++ b/Engine/Source/Runtime/Windows/D3D11RHI/Private/D3D11Viewport.cpp
@@ -80,6 +80,15 @@ namespace RHIConsoleVariables
 		TEXT("Number of consecutive 'fast' frames before vsync is enabled."),
 		ECVF_RenderThreadSafe
 		);
+
+	int32 MaximumFrameLatency = 3;
+	static FAutoConsoleVariableRef CVarMaximumFrameLatency(
+		TEXT("RHI.MaximumFrameLatency"),
+		MaximumFrameLatency,
+		TEXT("Number of frames that can be queued for render."),
+		ECVF_RenderThreadSafe
+		);
+
 };
 
 extern void D3D11TextureAllocated2D( FD3D11Texture2D& Texture );
@@ -418,6 +427,13 @@ bool FD3D11Viewport::Present(bool bLockToVsync)
 		}
 	}
 
+	if (MaximumFrameLatency != RHIConsoleVariables::MaximumFrameLatency)
+	{
+		MaximumFrameLatency = RHIConsoleVariables::MaximumFrameLatency;	
+		TRefCountPtr<IDXGIDevice1> DXGIDevice;
+		VERIFYD3D11RESULT(D3DRHI->GetDevice()->QueryInterface(IID_IDXGIDevice, (void**)DXGIDevice.GetInitReference()));
+		DXGIDevice->SetMaximumFrameLatency(MaximumFrameLatency);
+	}
 
 	// When desktop composition is enabled, locking to vsync via the Present
 	// call is unreliable. Instead, communicate with the desktop window manager
diff --git a/Engine/Source/Runtime/Windows/D3D11RHI/Private/Windows/WindowsD3D11Viewport.cpp b/Engine/Source/Runtime/Windows/D3D11RHI/Private/Windows/WindowsD3D11Viewport.cpp
--- a/Engine/Source/Runtime/Windows/D3D11RHI/Private/Windows/WindowsD3D11Viewport.cpp
+++ b/Engine/Source/Runtime/Windows/D3D11RHI/Private/Windows/WindowsD3D11Viewport.cpp
@@ -25,6 +25,7 @@ FD3D11Viewport::FD3D11Viewport(FD3D11DynamicRHI* InD3DRHI,HWND InWindowHandle,ui
 	SyncCounter(0),
 	bSyncedLastFrame(false),
 	WindowHandle(InWindowHandle),
+	MaximumFrameLatency(3),
 	SizeX(InSizeX),
 	SizeY(InSizeY),
 	bIsFullscreen(bInIsFullscreen),
diff --git a/Engine/Source/Runtime/Windows/D3D11RHI/Public/D3D11Viewport.h b/Engine/Source/Runtime/Windows/D3D11RHI/Public/D3D11Viewport.h
--- a/Engine/Source/Runtime/Windows/D3D11RHI/Public/D3D11Viewport.h
+++ b/Engine/Source/Runtime/Windows/D3D11RHI/Public/D3D11Viewport.h
@@ -100,6 +100,7 @@ private:
 	int32 SyncCounter;
 	bool bSyncedLastFrame;
 	HWND WindowHandle;
+	uint32 MaximumFrameLatency;
 	uint32 SizeX;
 	uint32 SizeY;
 	bool bIsFullscreen;
