#ue4 Reverting accidental local change to comment out assert

diff --git a/Engine/Source/Runtime/Windows/D3D11RHI/Private/D3D11Query.cpp b/Engine/Source/Runtime/Windows/D3D11RHI/Private/D3D11Query.cpp
--- a/Engine/Source/Runtime/Windows/D3D11RHI/Private/D3D11Query.cpp
+++ b/Engine/Source/Runtime/Windows/D3D11RHI/Private/D3D11Query.cpp
@@ -374,7 +374,7 @@ void FD3D11BufferedGPUTiming::PlatformStaticInitialize(void* UserData)
 		if(D3DResult == S_OK)
 		{
 			GTimingFrequency = FreqQueryData.Frequency;
-//			checkSlow(!FreqQueryData.Disjoint);
+			checkSlow(!FreqQueryData.Disjoint);
 		}
 	}
 
