Submitting Gil's fix for flickering meshes in the templates and Infiltrator.

diff --git a/Engine/Source/Runtime/Windows/D3D11RHI/Private/D3D11Commands.cpp b/Engine/Source/Runtime/Windows/D3D11RHI/Private/D3D11Commands.cpp
--- a/Engine/Source/Runtime/Windows/D3D11RHI/Private/D3D11Commands.cpp
+++ b/Engine/Source/Runtime/Windows/D3D11RHI/Private/D3D11Commands.cpp
@@ -1199,14 +1199,6 @@ static int32 PeriodicCheck = 0;
 
 void FD3D11DynamicRHI::CommitGraphicsResourceTables()
 {
-#if PLATFORM_HAS_THREADSAFE_RHIGetRenderQueryResult
-	if (GRHIThread && ++PeriodicCheck >= 10)
-	{
-		QUICK_SCOPE_CYCLE_COUNTER(STAT_CommitGraphicsResourceTables_CheckThreadsafeQueries);
-		PeriodicCheck = 0;
-		CheckThreadsafeQueries();
-	}
-#endif
 	FD3D11BoundShaderState* RESTRICT CurrentBoundShaderState = (FD3D11BoundShaderState*)BoundShaderStateHistory.GetLast();
 	check(CurrentBoundShaderState);
 
diff --git a/Engine/Source/Runtime/Windows/D3D11RHI/Private/D3D11Query.cpp b/Engine/Source/Runtime/Windows/D3D11RHI/Private/D3D11Query.cpp
--- a/Engine/Source/Runtime/Windows/D3D11RHI/Private/D3D11Query.cpp
+++ b/Engine/Source/Runtime/Windows/D3D11RHI/Private/D3D11Query.cpp
@@ -6,95 +6,6 @@
 
 #include "D3D11RHIPrivate.h"
 
-#if PLATFORM_HAS_THREADSAFE_RHIGetRenderQueryResult
-struct FThreadsafeQueryEmulator
-{
-	FThreadSafeCounter IsWaiting;
-	FCriticalSection OutstandingQueriesLock;
-	TArray<FRenderQueryRHIParamRef> OutstandingQueries;
-
-
-	FORCEINLINE_DEBUGGABLE void CheckThreadsafeQueries(FD3D11DynamicRHI &CallingRHI, bool bWait = false)
-	{
-		if (!IsWaiting.GetValue())
-		{
-			return;
-		}
-		QUICK_SCOPE_CYCLE_COUNTER(STAT_CheckThreadsafeQueries);
-		FScopeLock Lock(&OutstandingQueriesLock);
-		int32 LastOutstanding = OutstandingQueries.Num() - 1;
-		for (int32 Index = LastOutstanding; Index >= 0; Index--)
-		{
-			FRenderQueryRHIParamRef QueryRHI = OutstandingQueries[Index];
-
-			DYNAMIC_CAST_D3D11RESOURCE(OcclusionQuery,Query); // I think this macro is the worst thing in the whole world -- Gil
-
-			if(!Query->bResultIsCached)
-			{
-				bool bSuccess = CallingRHI.GetQueryData(Query->Resource,&Query->Result,sizeof(Query->Result),bWait, Query->QueryType);
-				FPlatformMisc::MemoryBarrier();
-				Query->bResultIsCached = bSuccess;
-				check(!bWait || Query->bResultIsCached);
-			}
-			if (Query->bResultIsCached)
-			{
-				LastOutstanding = Index - 1;
-			}
-			else
-			{
-				break;
-			}
-		}
-		int32 NumDeleted = OutstandingQueries.Num() - LastOutstanding - 1;
-		OutstandingQueries.RemoveAt(LastOutstanding + 1, NumDeleted, false);
-		check(!bWait || !OutstandingQueries.Num());
-	}
-	FORCEINLINE_DEBUGGABLE void Remove(FRenderQueryRHIParamRef QueryRHI)
-	{
-		FScopeLock Lock(&OutstandingQueriesLock);
-		OutstandingQueries.Remove(QueryRHI);
-	}
-	FORCEINLINE_DEBUGGABLE void Add(FRenderQueryRHIParamRef QueryRHI)
-	{
-		FScopeLock Lock(&OutstandingQueriesLock); // I am not sure this is actually required, but there should be no contention anyway
-		OutstandingQueries.Add(QueryRHI);
-	}
-	FORCEINLINE_DEBUGGABLE void WaitForQuery(FD3D11DynamicRHI &CallingRHI,FRenderQueryRHIParamRef QueryRHI)
-	{
-		DYNAMIC_CAST_D3D11RESOURCE(OcclusionQuery,Query);
-		verify(IsWaiting.Increment() == 1);
-		int32 Iteration = 0;
-		while (!Query->bResultIsCached)
-		{
-			if (Iteration++ > 0)
-			{
-				FPlatformProcess::Sleep(0.0001f);
-			}
-			if (!GRHICommandList.IsRHIThreadActive())
-			{
-				CheckThreadsafeQueries(CallingRHI, true);
-				break;
-			}
-			FPlatformMisc::MemoryBarrier();
-		}
-		verify(IsWaiting.Decrement() == 0);
-	}
-};
-
-static FThreadsafeQueryEmulator ThreadsafeQueryEmulator;
-
-void FD3D11DynamicRHI::CheckThreadsafeQueries(bool bWait)
-{
-	ThreadsafeQueryEmulator.CheckThreadsafeQueries(*this, bWait);
-}
-
-FD3D11OcclusionQuery::~FD3D11OcclusionQuery()
-{
-	// this more or less never happens until shutdown, but we will do a slow cleanup here for completeness.
-	ThreadsafeQueryEmulator.Remove(this);
-}
-
-#endif
 
 FRenderQueryRHIRef FD3D11DynamicRHI::RHICreateRenderQuery(ERenderQueryType QueryType)
 {
@@ -132,42 +43,6 @@ bool FD3D11DynamicRHI::RHIGetRenderQueryResult(FRenderQueryRHIParamRef QueryRHI,
 {
 	DYNAMIC_CAST_D3D11RESOURCE(OcclusionQuery,Query);
 
-#if PLATFORM_HAS_THREADSAFE_RHIGetRenderQueryResult
-	if (GRHIThread && !IsInRHIThread())
-	{
-		check(IsInRenderingThread());
-		if (!Query->bResultIsCached && !bWait)
-		{
-			return false;
-		}
-		if (!Query->bResultIsCached && bWait)
-		{
-			QUICK_SCOPE_CYCLE_COUNTER(STAT_GetRenderQueryResult_Spin);
-			ThreadsafeQueryEmulator.WaitForQuery(*this, QueryRHI);
-		}
-		if (Query->bResultIsCached)
-		{
-			if(Query->QueryType == RQT_AbsoluteTime)
-			{
-				// GetTimingFrequency is the number of ticks per second
-				uint64 Div = FGPUTiming::GetTimingFrequency() / (1000 * 1000);
-
-				// convert from GPU specific timestamp to micro sec (1 / 1 000 000 s) which seems a reasonable resolution
-				OutResult = Query->Result / Div;
-			}
-			else
-			{
-				OutResult = Query->Result;
-			}
-			return true;
-		}
-		if (!bWait)
-		{
-			return false;
-		}
-		 // we need to fall through in the case that we don't have outstanding queries...this happens when we transition into rhi thread because the last frame queries were not queued up
-	}
-#endif
 	bool bSuccess = true;
 	if(!Query->bResultIsCached)
 	{
@@ -188,9 +63,6 @@ bool FD3D11DynamicRHI::RHIGetRenderQueryResult(FRenderQueryRHIParamRef QueryRHI,
 	{
 		OutResult = Query->Result;
 	}
-#if PLATFORM_HAS_THREADSAFE_RHIGetRenderQueryResult
-	check(!bWait || bSuccess);
-#endif
 	return bSuccess;
 }
 
@@ -219,14 +91,6 @@ void FD3D11DynamicRHI::RHIEndRenderQuery(FRenderQueryRHIParamRef QueryRHI)
 	DYNAMIC_CAST_D3D11RESOURCE(OcclusionQuery,Query);
 #if PLATFORM_SUPPORTS_RHI_THREAD
 	Query->bResultIsCached = false; // for occlusion queries, this is redundant with the one in begin
-
-	#if PLATFORM_HAS_THREADSAFE_RHIGetRenderQueryResult
-		if (GRHIThread)
-		{
-			ThreadsafeQueryEmulator.Add(QueryRHI);
-		}
-	#endif
-
 #elif PLATFORM_HAS_THREADSAFE_RHIGetRenderQueryResult
 	#error "PLATFORM_HAS_THREADSAFE_RHIGetRenderQueryResult requires PLATFORM_SUPPORTS_RHI_THREAD"
 #endif
diff --git a/Engine/Source/Runtime/Windows/D3D11RHI/Private/D3D11RHIPrivate.h b/Engine/Source/Runtime/Windows/D3D11RHI/Private/D3D11RHIPrivate.h
--- a/Engine/Source/Runtime/Windows/D3D11RHI/Private/D3D11RHIPrivate.h
+++ b/Engine/Source/Runtime/Windows/D3D11RHI/Private/D3D11RHIPrivate.h
@@ -607,9 +607,6 @@ protected:
 #if PLATFORM_SUPPORTS_RHI_THREAD
 	void SetupRecursiveResources();
 #endif
-#if PLATFORM_HAS_THREADSAFE_RHIGetRenderQueryResult
-	void CheckThreadsafeQueries(bool bWait = false);
-#endif
 
 	friend struct FD3DGPUProfiler;
 };
diff --git a/Engine/Source/Runtime/Windows/D3D11RHI/Public/D3D11Resources.h b/Engine/Source/Runtime/Windows/D3D11RHI/Public/D3D11Resources.h
--- a/Engine/Source/Runtime/Windows/D3D11RHI/Public/D3D11Resources.h
+++ b/Engine/Source/Runtime/Windows/D3D11RHI/Public/D3D11Resources.h
@@ -579,9 +579,6 @@ public:
 		QueryType(InQueryType)
 	{}
 
-#if PLATFORM_HAS_THREADSAFE_RHIGetRenderQueryResult
-	~FD3D11OcclusionQuery();
-#endif
 };
 
 /** Updates tracked stats for a buffer. */
