Added the ability to specify a pixel format for RHI viewports

This allows us control over what back buffer format to use for specific RHI viewports. This is ignored on some RHIs (metal, gnm) where it doesn't make sense.

diff --git a/Engine/Source/Runtime/EmptyRHI/Private/EmptyViewport.cpp b/Engine/Source/Runtime/EmptyRHI/Private/EmptyViewport.cpp
--- a/Engine/Source/Runtime/EmptyRHI/Private/EmptyViewport.cpp
+++ b/Engine/Source/Runtime/EmptyRHI/Private/EmptyViewport.cpp
@@ -20,7 +20,7 @@ FEmptyViewport::~FEmptyViewport()
 /*=============================================================================
  *	The following RHI functions must be called from the main thread.
  *=============================================================================*/
-FViewportRHIRef FEmptyDynamicRHI::RHICreateViewport(void* WindowHandle,uint32 SizeX,uint32 SizeY,bool bIsFullscreen)
+FViewportRHIRef FEmptyDynamicRHI::RHICreateViewport(void* WindowHandle,uint32 SizeX,uint32 SizeY,bool bIsFullscreen,EPixelFormat PreferredPixelFormat)
 {
 	check( IsInGameThread() );
 	return new FEmptyViewport(WindowHandle, SizeX, SizeY, bIsFullscreen);
diff --git a/Engine/Source/Runtime/Engine/Private/UnrealClient.cpp b/Engine/Source/Runtime/Engine/Private/UnrealClient.cpp
--- a/Engine/Source/Runtime/Engine/Private/UnrealClient.cpp
+++ b/Engine/Source/Runtime/Engine/Private/UnrealClient.cpp
@@ -1275,8 +1275,8 @@ void FViewport::UpdateViewportRHI(bool bDestroyed,uint32 NewSizeX,uint32 NewSize
 					GetWindow(),
 					SizeX,
 					SizeY,
-					IsFullscreen()
-					);
+					IsFullscreen(),
+					EPixelFormat::PF_Unknown /* ie, use default format */);
 			}
 		
 			// Initialize the viewport's resources.
@@ -1369,7 +1369,8 @@ void FViewport::InitRHI()
 			GetWindow(),
 			SizeX,
 			SizeY,
-			IsFullscreen()
+			IsFullscreen(),
+			EPixelFormat::PF_Unknown
 			);
 		UpdateRenderTargetSurfaceRHIToCurrentBackBuffer();
 	}
diff --git a/Engine/Source/Runtime/IOS/MetalRHI/Private/MetalViewport.cpp b/Engine/Source/Runtime/IOS/MetalRHI/Private/MetalViewport.cpp
--- a/Engine/Source/Runtime/IOS/MetalRHI/Private/MetalViewport.cpp
+++ b/Engine/Source/Runtime/IOS/MetalRHI/Private/MetalViewport.cpp
@@ -20,7 +20,7 @@ FMetalViewport::~FMetalViewport()
 /*=============================================================================
  *	The following RHI functions must be called from the main thread.
  *=============================================================================*/
-FViewportRHIRef FMetalDynamicRHI::RHICreateViewport(void* WindowHandle,uint32 SizeX,uint32 SizeY,bool bIsFullscreen)
+FViewportRHIRef FMetalDynamicRHI::RHICreateViewport(void* WindowHandle,uint32 SizeX,uint32 SizeY,bool bIsFullscreen,EPixelFormat PreferredPixelFormat /* ignored */)
 {
 	check( IsInGameThread() );
 	return new FMetalViewport(WindowHandle, SizeX, SizeY, bIsFullscreen);
diff --git a/Engine/Source/Runtime/OpenGLDrv/Private/OpenGLViewport.cpp b/Engine/Source/Runtime/OpenGLDrv/Private/OpenGLViewport.cpp
--- a/Engine/Source/Runtime/OpenGLDrv/Private/OpenGLViewport.cpp
+++ b/Engine/Source/Runtime/OpenGLDrv/Private/OpenGLViewport.cpp
@@ -36,7 +36,7 @@ bool FOpenGLDynamicRHI::RHIGetAvailableResolutions(FScreenResolutionArray& Resol
 /*=============================================================================
  *	The following RHI functions must be called from the main thread.
  *=============================================================================*/
-FViewportRHIRef FOpenGLDynamicRHI::RHICreateViewport(void* WindowHandle,uint32 SizeX,uint32 SizeY,bool bIsFullscreen)
+FViewportRHIRef FOpenGLDynamicRHI::RHICreateViewport(void* WindowHandle,uint32 SizeX,uint32 SizeY,bool bIsFullscreen,EPixelFormat PreferredPixelFormat)
 {
 	check(IsInGameThread());
 
@@ -44,7 +44,13 @@ FViewportRHIRef FOpenGLDynamicRHI::RHICreateViewport(void* WindowHandle,uint32 S
 //	SCOPED_SUSPEND_RENDERING_THREAD(true);
 //#endif
 
-	return new FOpenGLViewport(this,WindowHandle,SizeX,SizeY,bIsFullscreen);
+	// Use a default pixel format if none was specified	
+	if (PreferredPixelFormat == EPixelFormat::PF_Unknown)
+	{
+		PreferredPixelFormat = EPixelFormat::PF_B8G8R8A8;
+	}
+
+	return new FOpenGLViewport(this,WindowHandle,SizeX,SizeY,bIsFullscreen,PreferredPixelFormat);
 }
 
 void FOpenGLDynamicRHI::RHIResizeViewport(FViewportRHIParamRef ViewportRHI,uint32 SizeX,uint32 SizeY,bool bIsFullscreen)
@@ -191,12 +197,13 @@ FTexture2DRHIRef FOpenGLDynamicRHI::RHIGetViewportBackBuffer(FViewportRHIParamRe
 	return Viewport->GetBackBuffer();
 }
 
-FOpenGLViewport::FOpenGLViewport(FOpenGLDynamicRHI* InOpenGLRHI,void* InWindowHandle,uint32 InSizeX,uint32 InSizeY,bool bInIsFullscreen)
+FOpenGLViewport::FOpenGLViewport(FOpenGLDynamicRHI* InOpenGLRHI,void* InWindowHandle,uint32 InSizeX,uint32 InSizeY,bool bInIsFullscreen,EPixelFormat PreferredPixelFormat)
 	: OpenGLRHI(InOpenGLRHI)
 	, OpenGLContext(NULL)
 	, SizeX(0)
 	, SizeY(0)
 	, bIsFullscreen(false)
+	, PixelFormat(PreferredPixelFormat)
 	, bIsValid(true)
 	, FrameSyncEvent(InOpenGLRHI)
 {
@@ -255,7 +262,7 @@ void FOpenGLViewport::Resize(uint32 InSizeX,uint32 InSizeY,bool bInIsFullscreen)
 	BackBuffer = (FOpenGLTexture2D*)PlatformCreateBuiltinBackBuffer(OpenGLRHI, InSizeX, InSizeY);
 	if (!BackBuffer)
 	{
-		BackBuffer = (FOpenGLTexture2D*)OpenGLRHI->CreateOpenGLTexture(InSizeX, InSizeY, false, false, PF_B8G8R8A8, 1, 1, 1, TexCreate_RenderTargetable);
+		BackBuffer = (FOpenGLTexture2D*)OpenGLRHI->CreateOpenGLTexture(InSizeX, InSizeY, false, false, PixelFormat, 1, 1, 1, TexCreate_RenderTargetable);
 	}
 
 	PlatformResizeGLContext(OpenGLRHI->PlatformDevice, OpenGLContext, InSizeX, InSizeY, bInIsFullscreen, bIsFullscreen, BackBuffer->Target, BackBuffer->Resource);
diff --git a/Engine/Source/Runtime/OpenGLDrv/Public/OpenGLResources.h b/Engine/Source/Runtime/OpenGLDrv/Public/OpenGLResources.h
--- a/Engine/Source/Runtime/OpenGLDrv/Public/OpenGLResources.h
+++ b/Engine/Source/Runtime/OpenGLDrv/Public/OpenGLResources.h
@@ -1305,7 +1305,7 @@ class FOpenGLViewport : public FRHIViewport
 {
 public:
 
-	FOpenGLViewport(class FOpenGLDynamicRHI* InOpenGLRHI,void* InWindowHandle,uint32 InSizeX,uint32 InSizeY,bool bInIsFullscreen);
+	FOpenGLViewport(class FOpenGLDynamicRHI* InOpenGLRHI,void* InWindowHandle,uint32 InSizeX,uint32 InSizeY,bool bInIsFullscreen,EPixelFormat PreferredPixelFormat);
 	~FOpenGLViewport();
 
 	void Resize(uint32 InSizeX,uint32 InSizeY,bool bInIsFullscreen);
@@ -1344,6 +1344,7 @@ private:
 	uint32 SizeX;
 	uint32 SizeY;
 	bool bIsFullscreen;
+	EPixelFormat PixelFormat;
 	bool bIsValid;
 	TRefCountPtr<FOpenGLTexture2D> BackBuffer;
 	FOpenGLEventQuery FrameSyncEvent;
diff --git a/Engine/Source/Runtime/RHI/Public/RHIMethods.h b/Engine/Source/Runtime/RHI/Public/RHIMethods.h
--- a/Engine/Source/Runtime/RHI/Public/RHIMethods.h
+++ b/Engine/Source/Runtime/RHI/Public/RHIMethods.h
@@ -1453,12 +1453,13 @@ DEFINE_RHIMETHOD_GLOBALTHREADSAFE_0(
 /**
  * The following RHI functions must be called from the main thread.
  */
-DEFINE_RHIMETHOD_GLOBALTHREADSAFE_4(
+DEFINE_RHIMETHOD_GLOBALTHREADSAFE_5(
 	FViewportRHIRef,CreateViewport,
 	void*,WindowHandle,
 	uint32,SizeX,
 	uint32,SizeY,
 	bool,bIsFullscreen,
+	EPixelFormat,PreferredPixelFormat,
 	return,return new FRHIViewport();
 	);
 DEFINE_RHIMETHOD_GLOBALTHREADSAFE_4(
diff --git a/Engine/Source/Runtime/SlateRHIRenderer/Private/SlateRHIRenderer.cpp b/Engine/Source/Runtime/SlateRHIRenderer/Private/SlateRHIRenderer.cpp
--- a/Engine/Source/Runtime/SlateRHIRenderer/Private/SlateRHIRenderer.cpp
+++ b/Engine/Source/Runtime/SlateRHIRenderer/Private/SlateRHIRenderer.cpp
@@ -263,7 +263,7 @@ void FSlateRHIRenderer::CreateViewport( const TSharedRef<SWindow> Window )
 		checkf(Width <= MAX_VIEWPORT_SIZE && Height <= MAX_VIEWPORT_SIZE, TEXT("Invalid window with Width=%u and Height=%u"), Width, Height);
 
 		bool bFullscreen = IsViewportFullscreen( *Window );
-		NewInfo->ViewportRHI = RHICreateViewport( NewInfo->OSWindow, Width, Height, bFullscreen );
+		NewInfo->ViewportRHI = RHICreateViewport( NewInfo->OSWindow, Width, Height, bFullscreen, EPixelFormat::PF_Unknown );
 		NewInfo->bFullscreen = bFullscreen;
 
 		WindowToViewportInfo.Add( &Window.Get(), NewInfo );
@@ -316,7 +316,7 @@ void FSlateRHIRenderer::ConditionalResizeViewport( FViewportInfo* ViewInfo, uint
 		}
 		else
 		{
-			ViewInfo->ViewportRHI = RHICreateViewport(ViewInfo->OSWindow, NewWidth, NewHeight, bFullscreen);
+			ViewInfo->ViewportRHI = RHICreateViewport(ViewInfo->OSWindow, NewWidth, NewHeight, bFullscreen, EPixelFormat::PF_Unknown);
 		}
 
 		// Safe to call here as the rendering thread has been suspended: game thread == render thread!
@@ -1119,7 +1119,7 @@ void* FSlateRHIRenderer::GetViewportResource( const SWindow& Window )
 
 			const bool bFullscreen = IsViewportFullscreen( Window );
 
-			ViewportInfo->ViewportRHI = RHICreateViewport( ViewportInfo->OSWindow, ViewportInfo->Width, ViewportInfo->Height, bFullscreen );
+			ViewportInfo->ViewportRHI = RHICreateViewport( ViewportInfo->OSWindow, ViewportInfo->Width, ViewportInfo->Height, bFullscreen, EPixelFormat::PF_Unknown );
 		}
 
 		return &ViewportInfo->ViewportRHI;
diff --git a/Engine/Source/Runtime/Windows/D3D11RHI/Private/D3D11Viewport.cpp b/Engine/Source/Runtime/Windows/D3D11RHI/Private/D3D11Viewport.cpp
--- a/Engine/Source/Runtime/Windows/D3D11RHI/Private/D3D11Viewport.cpp
+++ b/Engine/Source/Runtime/Windows/D3D11RHI/Private/D3D11Viewport.cpp
@@ -96,7 +96,7 @@ extern void D3D11TextureAllocated2D( FD3D11Texture2D& Texture );
 /**
  * Creates a FD3D11Surface to represent a swap chain's back buffer.
  */
-FD3D11Texture2D* GetSwapChainSurface(FD3D11DynamicRHI* D3DRHI,IDXGISwapChain* SwapChain)
+FD3D11Texture2D* GetSwapChainSurface(FD3D11DynamicRHI* D3DRHI, EPixelFormat PixelFormat, IDXGISwapChain* SwapChain)
 {
 	// Grab the back buffer
 	TRefCountPtr<ID3D11Texture2D> BackBufferResource;
@@ -138,7 +138,7 @@ FD3D11Texture2D* GetSwapChainSurface(FD3D11DynamicRHI* D3DRHI,IDXGISwapChain* Sw
 		1,
 		1,
 		1,
-		PF_A2B10G10R10,
+		PixelFormat,
 		false,
 		false,
 		false
@@ -172,7 +172,7 @@ DXGI_MODE_DESC FD3D11Viewport::SetupDXGI_MODE_DESC() const
 	Ret.Height = SizeY;
 	Ret.RefreshRate.Numerator = 0;	// illamas: use 0 to avoid a potential mismatch with hw
 	Ret.RefreshRate.Denominator = 0;	// illamas: ditto
-	Ret.Format = (DXGI_FORMAT)FD3D11Viewport::GetBackBufferFormat();
+	Ret.Format = GetRenderTargetFormat(PixelFormat);
 	Ret.ScanlineOrdering = DXGI_MODE_SCANLINE_ORDER_UNSPECIFIED;
 	Ret.Scaling = DXGI_MODE_SCALING_UNSPECIFIED;
 
@@ -209,7 +209,7 @@ void FD3D11Viewport::Resize(uint32 InSizeX,uint32 InSizeY,bool bInIsFullscreen)
 		check(SizeY > 0);
 
 		// Resize the swap chain.
-		VERIFYD3D11RESULT_EX(SwapChain->ResizeBuffers(1,SizeX,SizeY,(DXGI_FORMAT)FD3D11Viewport::GetBackBufferFormat(),DXGI_SWAP_CHAIN_FLAG_ALLOW_MODE_SWITCH), D3DRHI->GetDevice());
+		VERIFYD3D11RESULT_EX(SwapChain->ResizeBuffers(1,SizeX,SizeY,GetRenderTargetFormat(PixelFormat),DXGI_SWAP_CHAIN_FLAG_ALLOW_MODE_SWITCH), D3DRHI->GetDevice());
 
 		if(bInIsFullscreen)
 		{
@@ -233,7 +233,7 @@ void FD3D11Viewport::Resize(uint32 InSizeX,uint32 InSizeY,bool bInIsFullscreen)
 	}
 
 	// Create a RHI surface to represent the viewport's back buffer.
-	BackBuffer = GetSwapChainSurface(D3DRHI,SwapChain);
+	BackBuffer = GetSwapChainSurface(D3DRHI, PixelFormat, SwapChain);
 }
 
 /** Returns true if desktop composition is enabled. */
@@ -455,10 +455,17 @@ bool FD3D11Viewport::Present(bool bLockToVsync)
 /*=============================================================================
  *	The following RHI functions must be called from the main thread.
  *=============================================================================*/
-FViewportRHIRef FD3D11DynamicRHI::RHICreateViewport(void* WindowHandle,uint32 SizeX,uint32 SizeY,bool bIsFullscreen)
+FViewportRHIRef FD3D11DynamicRHI::RHICreateViewport(void* WindowHandle,uint32 SizeX,uint32 SizeY,bool bIsFullscreen,EPixelFormat PreferredPixelFormat)
 {
 	check( IsInGameThread() );
-	return new FD3D11Viewport(this,(HWND)WindowHandle,SizeX,SizeY,bIsFullscreen);
+
+	// Use a default pixel format if none was specified	
+	if (PreferredPixelFormat == EPixelFormat::PF_Unknown)
+	{
+		PreferredPixelFormat = EPixelFormat::PF_A2B10G10R10;
+	}
+
+	return new FD3D11Viewport(this,(HWND)WindowHandle,SizeX,SizeY,bIsFullscreen,PreferredPixelFormat);
 }
 
 void FD3D11DynamicRHI::RHIResizeViewport(FViewportRHIParamRef ViewportRHI,uint32 SizeX,uint32 SizeY,bool bIsFullscreen)
diff --git a/Engine/Source/Runtime/Windows/D3D11RHI/Private/WinRT/WinRTD3D11Viewport.cpp b/Engine/Source/Runtime/Windows/D3D11RHI/Private/WinRT/WinRTD3D11Viewport.cpp
--- a/Engine/Source/Runtime/Windows/D3D11RHI/Private/WinRT/WinRTD3D11Viewport.cpp
+++ b/Engine/Source/Runtime/Windows/D3D11RHI/Private/WinRT/WinRTD3D11Viewport.cpp
@@ -26,14 +26,9 @@ using namespace Windows::UI::Core;
 
 extern FD3D11Texture2D* GetSwapChainSurface(FD3D11DynamicRHI* D3DRHI,IDXGISwapChain* SwapChain);
 
-int32 FD3D11Viewport::GetBackBufferFormat()
-{
-	return DXGI_FORMAT_R10G10B10A2_UNORM;
-}
-
 extern void D3D11TextureAllocated2D( FD3D11Texture2D& Texture );
 
-FD3D11Viewport::FD3D11Viewport(FD3D11DynamicRHI* InD3DRHI,HWND InWindowHandle,uint32 InSizeX,uint32 InSizeY,bool bInIsFullscreen):
+FD3D11Viewport::FD3D11Viewport(FD3D11DynamicRHI* InD3DRHI,HWND InWindowHandle,uint32 InSizeX,uint32 InSizeY,bool bInIsFullscreen,EPixelFormat InPreferredPixelFormat):
 	D3DRHI(InD3DRHI),
 	LastFlipTime(0),
 	LastFrameComplete(0),
@@ -44,6 +39,7 @@ FD3D11Viewport::FD3D11Viewport(FD3D11DynamicRHI* InD3DRHI,HWND InWindowHandle,ui
 	SizeX(InSizeX),
 	SizeY(InSizeY),
 	bIsFullscreen(bInIsFullscreen),
+	PixelFormat(InPreferredPixelFormat),
 	bIsValid(true),
 	FrameSyncEvent(InD3DRHI)
 {
@@ -67,7 +63,7 @@ FD3D11Viewport::FD3D11Viewport(FD3D11DynamicRHI* InD3DRHI,HWND InWindowHandle,ui
 	DXGI_SWAP_CHAIN_DESC1 swapChainDesc = {0};
 	swapChainDesc.Width = SizeX;
 	swapChainDesc.Height = SizeY;
-	swapChainDesc.Format = (DXGI_FORMAT)FD3D11Viewport::GetBackBufferFormat();
+	swapChainDesc.Format = GetRenderTargetFormat(PixelFormat);
 	swapChainDesc.Stereo = false; 
 	swapChainDesc.SampleDesc.Count = 1;                          // don't use multi-sampling
 	swapChainDesc.SampleDesc.Quality = 0;
diff --git a/Engine/Source/Runtime/Windows/D3D11RHI/Private/Windows/WindowsD3D11Viewport.cpp b/Engine/Source/Runtime/Windows/D3D11RHI/Private/Windows/WindowsD3D11Viewport.cpp
--- a/Engine/Source/Runtime/Windows/D3D11RHI/Private/Windows/WindowsD3D11Viewport.cpp
+++ b/Engine/Source/Runtime/Windows/D3D11RHI/Private/Windows/WindowsD3D11Viewport.cpp
@@ -10,14 +10,9 @@
 #include "AllowWindowsPlatformTypes.h"
 #include <dwmapi.h>
 
-int32 FD3D11Viewport::GetBackBufferFormat()
-{
-	return DXGI_FORMAT_R10G10B10A2_UNORM;
-}
-
-extern FD3D11Texture2D* GetSwapChainSurface(FD3D11DynamicRHI* D3DRHI,IDXGISwapChain* SwapChain);
+extern FD3D11Texture2D* GetSwapChainSurface(FD3D11DynamicRHI* D3DRHI, EPixelFormat PixelFormat, IDXGISwapChain* SwapChain);
 
-FD3D11Viewport::FD3D11Viewport(FD3D11DynamicRHI* InD3DRHI,HWND InWindowHandle,uint32 InSizeX,uint32 InSizeY,bool bInIsFullscreen):
+FD3D11Viewport::FD3D11Viewport(FD3D11DynamicRHI* InD3DRHI,HWND InWindowHandle,uint32 InSizeX,uint32 InSizeY,bool bInIsFullscreen, EPixelFormat InPreferredPixelFormat):
 	D3DRHI(InD3DRHI),
 	LastFlipTime(0),
 	LastFrameComplete(0),
@@ -29,6 +24,7 @@ FD3D11Viewport::FD3D11Viewport(FD3D11DynamicRHI* InD3DRHI,HWND InWindowHandle,ui
 	SizeX(InSizeX),
 	SizeY(InSizeY),
 	bIsFullscreen(bInIsFullscreen),
+	PixelFormat(InPreferredPixelFormat),
 	bIsValid(true),
 	FrameSyncEvent(InD3DRHI)
 {
@@ -64,7 +60,7 @@ FD3D11Viewport::FD3D11Viewport(FD3D11DynamicRHI* InD3DRHI,HWND InWindowHandle,ui
 	D3DRHI->GetFactory()->MakeWindowAssociation(WindowHandle,DXGI_MWA_NO_WINDOW_CHANGES);
 
 	// Create a RHI surface to represent the viewport's back buffer.
-	BackBuffer = GetSwapChainSurface(D3DRHI,SwapChain);
+	BackBuffer = GetSwapChainSurface(D3DRHI, PixelFormat, SwapChain);
 
 	// Tell the window to redraw when they can.
 	// @todo: For Slate viewports, it doesn't make sense to post WM_PAINT messages (we swallow those.)
diff --git a/Engine/Source/Runtime/Windows/D3D11RHI/Public/D3D11Viewport.h b/Engine/Source/Runtime/Windows/D3D11RHI/Public/D3D11Viewport.h
--- a/Engine/Source/Runtime/Windows/D3D11RHI/Public/D3D11Viewport.h
+++ b/Engine/Source/Runtime/Windows/D3D11RHI/Public/D3D11Viewport.h
@@ -32,15 +32,28 @@ private:
 	TRefCountPtr<ID3D11Query> Query;
 };
 
+static DXGI_FORMAT GetRenderTargetFormat(EPixelFormat PixelFormat)
+{
+	DXGI_FORMAT	DXFormat = (DXGI_FORMAT)GPixelFormats[PixelFormat].PlatformFormat;
+	switch(DXFormat)
+	{
+		case DXGI_FORMAT_B8G8R8A8_TYPELESS:		return DXGI_FORMAT_B8G8R8A8_UNORM;
+		case DXGI_FORMAT_BC1_TYPELESS:			return DXGI_FORMAT_BC1_UNORM;
+		case DXGI_FORMAT_BC2_TYPELESS:			return DXGI_FORMAT_BC2_UNORM;
+		case DXGI_FORMAT_BC3_TYPELESS:			return DXGI_FORMAT_BC3_UNORM;
+		case DXGI_FORMAT_R16_TYPELESS:			return DXGI_FORMAT_R16_UNORM;
+		case DXGI_FORMAT_R8G8B8A8_TYPELESS:		return DXGI_FORMAT_R8G8B8A8_UNORM;
+		default: 								return DXFormat;
+	}
+}
+
 class FD3D11Viewport : public FRHIViewport
 {
 public:
 
-	FD3D11Viewport(class FD3D11DynamicRHI* InD3DRHI,HWND InWindowHandle,uint32 InSizeX,uint32 InSizeY,bool bInIsFullscreen);
+	FD3D11Viewport(class FD3D11DynamicRHI* InD3DRHI,HWND InWindowHandle,uint32 InSizeX,uint32 InSizeY,bool bInIsFullscreen, EPixelFormat InPreferredPixelFormat);
 	~FD3D11Viewport();
 
-	static int32 GetBackBufferFormat();
-
 	void Resize(uint32 InSizeX,uint32 InSizeY,bool bInIsFullscreen);
 
 	/**
@@ -104,6 +117,7 @@ private:
 	uint32 SizeX;
 	uint32 SizeY;
 	bool bIsFullscreen;
+	EPixelFormat PixelFormat;
 	bool bIsValid;
 	TRefCountPtr<IDXGISwapChain> SwapChain;
 	TRefCountPtr<FD3D11Texture2D> BackBuffer;
