Fix foliage flickering when dragging a single instance. Dragging repeatedly renders a single unbuilt instance after rendering the built tree, and the rendering code assumed 1 instance can use the DrawIndexed path even though FirstInstance > 0.

diff --git a/Engine/Source/Runtime/Windows/D3D11RHI/Private/D3D11Commands.cpp b/Engine/Source/Runtime/Windows/D3D11RHI/Private/D3D11Commands.cpp
--- a/Engine/Source/Runtime/Windows/D3D11RHI/Private/D3D11Commands.cpp
+++ b/Engine/Source/Runtime/Windows/D3D11RHI/Private/D3D11Commands.cpp
@@ -1325,7 +1325,7 @@ void FD3D11DynamicRHI::RHIDrawIndexedPrimitive(FIndexBufferRHIParamRef IndexBuff
 	StateCache.SetIndexBuffer(IndexBuffer->Resource, Format, 0);
 	StateCache.SetPrimitiveTopology(GetD3D11PrimitiveType(PrimitiveType,bUsingTessellation));
 
-	if(NumInstances > 1)
+	if (NumInstances > 1 || FirstInstance != 0)
 	{
 		Direct3DDeviceIMContext->DrawIndexedInstanced(IndexCount, NumInstances, StartIndex, BaseVertexIndex, FirstInstance);
 	}
