UE4 - UE-8971 - Pull request for 1 bpp D3D formats
PR #795: Added support for 1 byte-per-pixel formats to ComputeBytesPerPixel (D3D11) (Contributed by [REDACTED_PERSON])
https://github.com/EpicGames/UnrealEngine/pull/795

diff --git a/Engine/Source/Runtime/Windows/D3D11RHI/Private/D3D11RenderTarget.cpp b/Engine/Source/Runtime/Windows/D3D11RHI/Private/D3D11RenderTarget.cpp
--- a/Engine/Source/Runtime/Windows/D3D11RHI/Private/D3D11RenderTarget.cpp
+++ b/Engine/Source/Runtime/Windows/D3D11RHI/Private/D3D11RenderTarget.cpp
@@ -608,6 +608,13 @@ static uint32 ComputeBytesPerPixel(DXGI_FORMAT Format)
 		case DXGI_FORMAT_R32G32B32A32_FLOAT:
 			BytesPerPixel = 16;
 			break;
+		case DXGI_FORMAT_R8_TYPELESS:
+		case DXGI_FORMAT_R8_UNORM:
+		case DXGI_FORMAT_R8_UINT:
+		case DXGI_FORMAT_R8_SNORM:
+		case DXGI_FORMAT_R8_SINT:
+			BytesPerPixel = 1;
+			break;
 	}
 
 	// format not supported yet
