Move some utility shaders to a new module - UtilityShaders (which the RHIs that need them now depend on)

diff --git a/Engine/Source/Runtime/Engine/Private/ClearReplacementShaders.cpp b/Engine/Source/Runtime/Engine/Private/ClearReplacementShaders.cpp
deleted file mode 100644
diff --git a/Engine/Source/Runtime/Engine/Private/OneColorShader.cpp b/Engine/Source/Runtime/Engine/Private/OneColorShader.cpp
deleted file mode 100644
diff --git a/Engine/Source/Runtime/Engine/Public/ClearReplacementShaders.h b/Engine/Source/Runtime/Engine/Public/ClearReplacementShaders.h
deleted file mode 100644
diff --git a/Engine/Source/Runtime/Engine/Public/OneColorShader.h b/Engine/Source/Runtime/Engine/Public/OneColorShader.h
deleted file mode 100644
diff --git a/Engine/Source/Runtime/IOS/MetalRHI/MetalRHI.Build.cs b/Engine/Source/Runtime/IOS/MetalRHI/MetalRHI.Build.cs
--- a/Engine/Source/Runtime/IOS/MetalRHI/MetalRHI.Build.cs
+++ b/Engine/Source/Runtime/IOS/MetalRHI/MetalRHI.Build.cs
@@ -13,7 +13,8 @@ public class MetalRHI : ModuleRules
 				"Engine",
 				"RHI",
 				"RenderCore",
-				"ShaderCore"
+				"ShaderCore",
+				"UtilityShaders",
 			}
 			);
 			
diff --git a/Engine/Source/Runtime/OpenGLDrv/OpenGLDrv.Build.cs b/Engine/Source/Runtime/OpenGLDrv/OpenGLDrv.Build.cs
--- a/Engine/Source/Runtime/OpenGLDrv/OpenGLDrv.Build.cs
+++ b/Engine/Source/Runtime/OpenGLDrv/OpenGLDrv.Build.cs
@@ -15,7 +15,8 @@ public class OpenGLDrv : ModuleRules
 				"Engine", 
 				"RHI", 
 				"RenderCore", 
-				"ShaderCore", 
+				"ShaderCore",
+				"UtilityShaders",
 			}
 			);
 
diff --git a/Engine/Source/Runtime/OpenGLDrv/Private/OpenGLDrv.cpp b/Engine/Source/Runtime/OpenGLDrv/Private/OpenGLDrv.cpp
--- a/Engine/Source/Runtime/OpenGLDrv/Private/OpenGLDrv.cpp
+++ b/Engine/Source/Runtime/OpenGLDrv/Private/OpenGLDrv.cpp
@@ -363,31 +363,6 @@ float FOpenGLEventNode::GetTiming()
 	return Result;
 }
 
-class FOpenGLRHILongGPUTaskPS : public FGlobalShader
-{
-	DECLARE_SHADER_TYPE(FOpenGLRHILongGPUTaskPS,Global);
-public:
-	FOpenGLRHILongGPUTaskPS( )	{ }
-	FOpenGLRHILongGPUTaskPS(const ShaderMetaType::CompiledShaderInitializerType& Initializer)
-		:	FGlobalShader( Initializer )
-	{
-	}
-
-	// FShader interface.
-	virtual bool Serialize(FArchive& Ar)
-	{
-		bool bShaderHasOutdatedParameters = FGlobalShader::Serialize(Ar);
-		return bShaderHasOutdatedParameters;
-	}
-
-	static bool ShouldCache(EShaderPlatform Platform)
-	{
-		return IsFeatureLevelSupported(Platform, ERHIFeatureLevel::SM4);
-	}
-};
-
-IMPLEMENT_SHADER_TYPE(,FOpenGLRHILongGPUTaskPS,TEXT("OneColorShader"),TEXT("MainLongGPUTask"),SF_Pixel);
-
 static FGlobalBoundShaderState LongGPUTaskBoundShaderState;
 
 void FOpenGLDynamicRHI::IssueLongGPUTask()
@@ -420,7 +395,7 @@ void FOpenGLDynamicRHI::IssueLongGPUTask()
 
 		auto ShaderMap = GetGlobalShaderMap(FeatureLevel);
 		TShaderMapRef<TOneColorVS<true> > VertexShader(ShaderMap);
-		TShaderMapRef<FOpenGLRHILongGPUTaskPS> PixelShader(ShaderMap);
+		TShaderMapRef<FLongGPUTaskPS> PixelShader(ShaderMap);
 
 		SetGlobalBoundShaderState(RHICmdList, FeatureLevel, LongGPUTaskBoundShaderState, GOpenGLVector4VertexDeclaration.VertexDeclarationRHI, *VertexShader, *PixelShader, 0);
 
diff --git a/Engine/Source/Runtime/Renderer/Renderer.Build.cs b/Engine/Source/Runtime/Renderer/Renderer.Build.cs
--- a/Engine/Source/Runtime/Renderer/Renderer.Build.cs
+++ b/Engine/Source/Runtime/Renderer/Renderer.Build.cs
@@ -29,7 +29,8 @@ public class Renderer : ModuleRules
 				"CoreUObject", 
 				"RenderCore", 
 				"RHI", 
-				"ShaderCore"
+				"ShaderCore",
+				"UtilityShaders",
 			}
 			);
 
diff --git a/Engine/Source/Runtime/UtilityShaders/Private/ClearReplacementShaders.cpp b/Engine/Source/Runtime/UtilityShaders/Private/ClearReplacementShaders.cpp
new file mode 100644
--- /dev/null
+++ b/Engine/Source/Runtime/UtilityShaders/Private/ClearReplacementShaders.cpp
@@ -0,0 +1,27 @@
+// Copyright 1998-2015 Epic Games, Inc. All Rights Reserved.
+
+
+#include "UtilityShadersPrivatePCH.h"
+#include "ClearReplacementShaders.h"
+#include "ShaderParameterUtils.h"
+
+IMPLEMENT_SHADER_TYPE(,FClearReplacementVS,TEXT("ClearReplacementShaders"),TEXT("ClearVS"),SF_Vertex);
+IMPLEMENT_SHADER_TYPE(,FClearReplacementPS,TEXT("ClearReplacementShaders"),TEXT("ClearPS"),SF_Pixel);
+
+IMPLEMENT_SHADER_TYPE(,FClearTexture2DReplacementCS,TEXT("ClearReplacementShaders"),TEXT("ClearTexture2DCS"),SF_Compute);
+
+IMPLEMENT_SHADER_TYPE(,FClearBufferReplacementCS,TEXT("ClearReplacementShaders"),TEXT("ClearBufferCS"),SF_Compute);
+
+void FClearTexture2DReplacementCS::SetParameters( FRHICommandList& RHICmdList, FUnorderedAccessViewRHIParamRef TextureRW, FLinearColor Value )
+{
+	FComputeShaderRHIParamRef ComputeShaderRHI = GetComputeShader();
+	SetShaderValue(RHICmdList, ComputeShaderRHI, ClearColor, Value);
+	RHICmdList.SetUAVParameter(ComputeShaderRHI, ClearTextureRW.GetBaseIndex(), TextureRW);
+}
+
+void FClearBufferReplacementCS::SetParameters( FRHICommandList& RHICmdList, FUnorderedAccessViewRHIParamRef BufferRW, uint32 Dword )
+{
+	FComputeShaderRHIParamRef ComputeShaderRHI = GetComputeShader();
+	SetShaderValue(RHICmdList, ComputeShaderRHI, ClearDword, Dword);
+	RHICmdList.SetUAVParameter(ComputeShaderRHI, ClearBufferRW.GetBaseIndex(), BufferRW);
+}
diff --git a/Engine/Source/Runtime/UtilityShaders/Private/OneColorShader.cpp b/Engine/Source/Runtime/UtilityShaders/Private/OneColorShader.cpp
new file mode 100644
--- /dev/null
+++ b/Engine/Source/Runtime/UtilityShaders/Private/OneColorShader.cpp
@@ -0,0 +1,23 @@
+// Copyright 1998-2015 Epic Games, Inc. All Rights Reserved.
+
+
+#include "UtilityShadersPrivatePCH.h"
+#include "OneColorShader.h"
+
+IMPLEMENT_SHADER_TYPE(template<>,TOneColorVS<true>,TEXT("OneColorShader"),TEXT("MainVertexShader"),SF_Vertex);
+IMPLEMENT_SHADER_TYPE(template<>,TOneColorVS<false>,TEXT("OneColorShader"),TEXT("MainVertexShader"),SF_Vertex);
+IMPLEMENT_SHADER_TYPE(,FOneColorPS,TEXT("OneColorShader"),TEXT("MainPixelShader"),SF_Pixel);
+// Compiling a version for every number of MRT's
+// On AMD PC hardware, outputting to a color index in the shader without a matching render target set has a significant performance hit
+IMPLEMENT_SHADER_TYPE(template<>,TOneColorPixelShaderMRT<1>,TEXT("OneColorShader"),TEXT("MainPixelShaderMRT"),SF_Pixel);
+IMPLEMENT_SHADER_TYPE(template<>,TOneColorPixelShaderMRT<2>,TEXT("OneColorShader"),TEXT("MainPixelShaderMRT"),SF_Pixel);
+IMPLEMENT_SHADER_TYPE(template<>,TOneColorPixelShaderMRT<3>,TEXT("OneColorShader"),TEXT("MainPixelShaderMRT"),SF_Pixel);
+IMPLEMENT_SHADER_TYPE(template<>,TOneColorPixelShaderMRT<4>,TEXT("OneColorShader"),TEXT("MainPixelShaderMRT"),SF_Pixel);
+IMPLEMENT_SHADER_TYPE(template<>,TOneColorPixelShaderMRT<5>,TEXT("OneColorShader"),TEXT("MainPixelShaderMRT"),SF_Pixel);
+IMPLEMENT_SHADER_TYPE(template<>,TOneColorPixelShaderMRT<6>,TEXT("OneColorShader"),TEXT("MainPixelShaderMRT"),SF_Pixel);
+IMPLEMENT_SHADER_TYPE(template<>,TOneColorPixelShaderMRT<7>,TEXT("OneColorShader"),TEXT("MainPixelShaderMRT"),SF_Pixel);
+IMPLEMENT_SHADER_TYPE(template<>,TOneColorPixelShaderMRT<8>,TEXT("OneColorShader"),TEXT("MainPixelShaderMRT"),SF_Pixel);
+
+IMPLEMENT_SHADER_TYPE(,FFillTextureCS,TEXT("OneColorShader"),TEXT("MainFillTextureCS"),SF_Compute);
+
+IMPLEMENT_SHADER_TYPE(,FLongGPUTaskPS,TEXT("OneColorShader"),TEXT("MainLongGPUTask"),SF_Pixel);
diff --git a/Engine/Source/Runtime/UtilityShaders/Private/ResolveShader.cpp b/Engine/Source/Runtime/UtilityShaders/Private/ResolveShader.cpp
new file mode 100644
--- /dev/null
+++ b/Engine/Source/Runtime/UtilityShaders/Private/ResolveShader.cpp
@@ -0,0 +1,16 @@
+// Copyright 1998-2015 Epic Games, Inc. All Rights Reserved.
+
+
+#include "UtilityShadersPrivatePCH.h"
+#include "ResolveShader.h"
+#include "ShaderParameterUtils.h"
+
+IMPLEMENT_SHADER_TYPE(, FResolveDepthPS, TEXT("ResolvePixelShader"), TEXT("MainDepth"), SF_Pixel);
+IMPLEMENT_SHADER_TYPE(, FResolveDepthNonMSPS, TEXT("ResolvePixelShader"), TEXT("MainDepthNonMS"), SF_Pixel);
+IMPLEMENT_SHADER_TYPE(, FResolveSingleSamplePS, TEXT("ResolvePixelShader"), TEXT("MainSingleSample"), SF_Pixel);
+IMPLEMENT_SHADER_TYPE(, FResolveVS, TEXT("ResolveVertexShader"), TEXT("Main"), SF_Vertex);
+
+void FResolveSingleSamplePS::SetParameters(FRHICommandList& RHICmdList, uint32 SingleSampleIndexValue)
+{
+	SetShaderValue(RHICmdList, GetPixelShader(),SingleSampleIndex,SingleSampleIndexValue);
+}
diff --git a/Engine/Source/Runtime/UtilityShaders/Private/UtilityShaders.cpp b/Engine/Source/Runtime/UtilityShaders/Private/UtilityShaders.cpp
new file mode 100644
--- /dev/null
+++ b/Engine/Source/Runtime/UtilityShaders/Private/UtilityShaders.cpp
@@ -0,0 +1,7 @@
+// Copyright 1998-2015 Epic Games, Inc. All Rights Reserved.
+
+#include "UtilityShadersPrivatePCH.h"
+#include "ModuleManager.h"
+
+
+IMPLEMENT_MODULE(FDefaultModuleImpl, UtilityShaders);
diff --git a/Engine/Source/Runtime/UtilityShaders/Private/UtilityShadersPrivatePCH.h b/Engine/Source/Runtime/UtilityShaders/Private/UtilityShadersPrivatePCH.h
new file mode 100644
--- /dev/null
+++ b/Engine/Source/Runtime/UtilityShaders/Private/UtilityShadersPrivatePCH.h
@@ -0,0 +1,7 @@
+// Copyright 1998-2015 Epic Games, Inc. All Rights Reserved.
+
+#pragma once
+
+#include "Core.h"
+#include "Engine.h"
+
diff --git a/Engine/Source/Runtime/UtilityShaders/Public/ClearReplacementShaders.h b/Engine/Source/Runtime/UtilityShaders/Public/ClearReplacementShaders.h
new file mode 100644
--- /dev/null
+++ b/Engine/Source/Runtime/UtilityShaders/Public/ClearReplacementShaders.h
@@ -0,0 +1,138 @@
+// Copyright 1998-2015 Epic Games, Inc. All Rights Reserved.
+
+
+#pragma once
+
+#include "GlobalShader.h"
+#include "ShaderParameters.h"
+
+class FClearReplacementVS : public FGlobalShader
+{
+	DECLARE_EXPORTED_SHADER_TYPE(FClearReplacementVS, Global, UTILITYSHADERS_API);
+public:
+	FClearReplacementVS() {}
+	FClearReplacementVS(const ShaderMetaType::CompiledShaderInitializerType& Initializer)
+	: FGlobalShader( Initializer )
+	{
+	}
+	
+	// FShader interface.
+	virtual bool Serialize(FArchive& Ar)
+	{
+		bool bShaderHasOutdatedParameters = FGlobalShader::Serialize(Ar);
+		return bShaderHasOutdatedParameters;
+	}
+	
+	static bool ShouldCache(EShaderPlatform Platform)
+	{
+		return IsFeatureLevelSupported(Platform, ERHIFeatureLevel::SM4);
+	}
+};
+
+class FClearReplacementPS : public FGlobalShader
+{
+	DECLARE_EXPORTED_SHADER_TYPE(FClearReplacementPS, Global, UTILITYSHADERS_API);
+public:
+	FClearReplacementPS() {}
+	FClearReplacementPS(const ShaderMetaType::CompiledShaderInitializerType& Initializer)
+	: FGlobalShader( Initializer )
+	{
+		ClearColor.Bind(Initializer.ParameterMap, TEXT("ClearColor"), SPF_Mandatory);
+	}
+	
+	// FShader interface.
+	virtual bool Serialize(FArchive& Ar)
+	{
+		bool bShaderHasOutdatedParameters = FGlobalShader::Serialize(Ar);
+		Ar << ClearColor;
+		return bShaderHasOutdatedParameters;
+	}
+	
+	static bool ShouldCache(EShaderPlatform Platform)
+	{
+		return IsFeatureLevelSupported(Platform, ERHIFeatureLevel::SM4);
+	}
+	
+	static void ModifyCompilationEnvironment(EShaderPlatform Platform, FShaderCompilerEnvironment& OutEnvironment)
+	{
+		FGlobalShader::ModifyCompilationEnvironment(Platform, OutEnvironment);
+		OutEnvironment.SetRenderTargetOutputFormat(0, PF_A32B32G32R32F);
+	}
+	
+protected:
+	FShaderParameter ClearColor;
+};
+
+class FClearTexture2DReplacementCS : public FGlobalShader
+{
+	DECLARE_EXPORTED_SHADER_TYPE(FClearTexture2DReplacementCS, Global, UTILITYSHADERS_API);
+public:
+	FClearTexture2DReplacementCS() {}
+	FClearTexture2DReplacementCS(const ShaderMetaType::CompiledShaderInitializerType& Initializer)
+	: FGlobalShader( Initializer )
+	{
+		ClearColor.Bind(Initializer.ParameterMap, TEXT("ClearColor"), SPF_Mandatory);
+		ClearTextureRW.Bind(Initializer.ParameterMap, TEXT("ClearTextureRW"), SPF_Mandatory);
+	}
+	
+	// FShader interface.
+	virtual bool Serialize(FArchive& Ar)
+	{
+		bool bShaderHasOutdatedParameters = FGlobalShader::Serialize(Ar);
+		Ar << ClearColor << ClearTextureRW;
+		return bShaderHasOutdatedParameters;
+	}
+	
+	void SetParameters(FRHICommandList& RHICmdList, FUnorderedAccessViewRHIParamRef TextureRW, FLinearColor Value);
+	
+	static bool ShouldCache(EShaderPlatform Platform)
+	{
+		return IsFeatureLevelSupported(Platform, ERHIFeatureLevel::SM5);
+	}
+	
+	const FShaderParameter& GetClearColorParameter()
+	{
+		return ClearColor;
+	}
+	
+	const FShaderResourceParameter& GetClearTextureRWParameter()
+	{
+		return ClearTextureRW;
+	}
+	
+protected:
+	FShaderParameter ClearColor;
+	FShaderResourceParameter ClearTextureRW;
+};
+
+class FClearBufferReplacementCS : public FGlobalShader
+{
+	DECLARE_EXPORTED_SHADER_TYPE(FClearBufferReplacementCS, Global, UTILITYSHADERS_API);
+public:
+	FClearBufferReplacementCS() {}
+	FClearBufferReplacementCS(const ShaderMetaType::CompiledShaderInitializerType& Initializer)
+	: FGlobalShader( Initializer )
+	{
+		ClearDword.Bind(Initializer.ParameterMap, TEXT("ClearDword"), SPF_Mandatory);
+		ClearBufferRW.Bind(Initializer.ParameterMap, TEXT("ClearBufferRW"), SPF_Mandatory);
+	}
+	
+	void SetParameters(FRHICommandList& RHICmdList, FUnorderedAccessViewRHIParamRef BufferRW, uint32 Dword);
+	
+	// FShader interface.
+	virtual bool Serialize(FArchive& Ar)
+	{
+		bool bShaderHasOutdatedParameters = FGlobalShader::Serialize(Ar);
+		Ar << ClearDword << ClearBufferRW;
+		return bShaderHasOutdatedParameters;
+	}
+	
+	static bool ShouldCache(EShaderPlatform Platform)
+	{
+		return IsFeatureLevelSupported(Platform, ERHIFeatureLevel::SM5);
+	}
+	
+	//protected:
+	FShaderParameter ClearDword;
+	FShaderResourceParameter ClearBufferRW;
+};
diff --git a/Engine/Source/Runtime/UtilityShaders/Public/OneColorShader.h b/Engine/Source/Runtime/UtilityShaders/Public/OneColorShader.h
new file mode 100644
--- /dev/null
+++ b/Engine/Source/Runtime/UtilityShaders/Public/OneColorShader.h
@@ -0,0 +1,158 @@
+// Copyright 1998-2015 Epic Games, Inc. All Rights Reserved.
+
+
+#pragma once
+
+#include "GlobalShader.h"
+#include "ShaderParameters.h"
+
+/**
+ * Vertex shader for rendering a single, constant color.
+ */
+template<bool bUsingNDCPositions=true>
+class TOneColorVS : public FGlobalShader
+{
+	DECLARE_EXPORTED_SHADER_TYPE(TOneColorVS, Global, UTILITYSHADERS_API);
+
+	/** Default constructor. */
+	TOneColorVS() {}
+
+public:
+
+	TOneColorVS(const ShaderMetaType::CompiledShaderInitializerType& Initializer)
+		: FGlobalShader(Initializer)
+	{}
+
+	static void ModifyCompilationEnvironment(EShaderPlatform Platform, FShaderCompilerEnvironment& OutEnvironment)
+	{
+		FGlobalShader::ModifyCompilationEnvironment(Platform, OutEnvironment);
+		OutEnvironment.SetDefine(TEXT("USING_NDC_POSITIONS"), (uint32)(bUsingNDCPositions ? 1 : 0));
+	}
+
+	static bool ShouldCache(EShaderPlatform Platform)
+	{
+		return true;
+	}
+};
+
+/**
+ * Pixel shader for rendering a single, constant color.
+ */
+class FOneColorPS : public FGlobalShader
+{
+	DECLARE_EXPORTED_SHADER_TYPE(FOneColorPS,Global,UTILITYSHADERS_API);
+public:
+	FOneColorPS( )	{ }
+	FOneColorPS(const ShaderMetaType::CompiledShaderInitializerType& Initializer)
+	:	FGlobalShader( Initializer )
+	{
+		ColorParameter.Bind( Initializer.ParameterMap, TEXT("DrawColorMRT"), SPF_Mandatory);
+	}
+
+	// FShader interface.
+	virtual bool Serialize(FArchive& Ar)
+	{
+		bool bShaderHasOutdatedParameters = FGlobalShader::Serialize(Ar);
+		Ar << ColorParameter;
+		return bShaderHasOutdatedParameters;
+	}
+	static bool ShouldCache(EShaderPlatform Platform)
+	{
+		return true;
+	}
+
+	/** The parameter to use for setting the draw Color. */
+	FShaderParameter ColorParameter;
+};
+
+/**
+ * Pixel shader for rendering a single, constant color to MRTs.
+ */
+template<int32 NumOutputs>
+class TOneColorPixelShaderMRT : public FOneColorPS
+{
+	DECLARE_EXPORTED_SHADER_TYPE(TOneColorPixelShaderMRT,Global,UTILITYSHADERS_API);
+public:
+	TOneColorPixelShaderMRT( )	{ }
+	TOneColorPixelShaderMRT(const ShaderMetaType::CompiledShaderInitializerType& Initializer)
+	:	FOneColorPS( Initializer )
+	{
+	}
+
+	static bool ShouldCache(EShaderPlatform Platform)
+	{
+		if( NumOutputs > 1 )
+		{
+			return IsFeatureLevelSupported( Platform, ERHIFeatureLevel::SM4 );
+		}
+		return true;
+	}
+
+	static void ModifyCompilationEnvironment(EShaderPlatform Platform, FShaderCompilerEnvironment& OutEnvironment)
+	{
+		FOneColorPS::ModifyCompilationEnvironment(Platform, OutEnvironment);
+		OutEnvironment.SetDefine(TEXT("NUM_OUTPUTS"), NumOutputs);
+	}
+};
+
+/**
+ * Compute shader for writing values
+ */
+class FFillTextureCS : public FGlobalShader
+{
+	DECLARE_EXPORTED_SHADER_TYPE(FFillTextureCS,Global,UTILITYSHADERS_API);
+public:
+	FFillTextureCS( )	{ }
+	FFillTextureCS(const ShaderMetaType::CompiledShaderInitializerType& Initializer)
+	:	FGlobalShader( Initializer )
+	{
+		FillValue.Bind( Initializer.ParameterMap, TEXT("FillValue"), SPF_Mandatory);
+		Params0.Bind( Initializer.ParameterMap, TEXT("Params0"), SPF_Mandatory);
+		Params1.Bind( Initializer.ParameterMap, TEXT("Params1"), SPF_Mandatory);
+		Params2.Bind( Initializer.ParameterMap, TEXT("Params2"), SPF_Optional);
+		FillTexture.Bind( Initializer.ParameterMap, TEXT("FillTexture"), SPF_Mandatory);
+	}
+
+	// FShader interface.
+	virtual bool Serialize(FArchive& Ar)
+	{
+		bool bShaderHasOutdatedParameters = FGlobalShader::Serialize(Ar);
+		Ar << FillValue << Params0 << Params1 << Params2 << FillTexture;
+		return bShaderHasOutdatedParameters;
+	}
+
+	static bool ShouldCache(EShaderPlatform Platform)
+	{
+		return IsFeatureLevelSupported(Platform, ERHIFeatureLevel::SM5);
+	}
+
+	FShaderParameter FillValue;
+	FShaderParameter Params0;	// Texture Width,Height (.xy); Use Exclude Rect 1 : 0 (.z)
+	FShaderParameter Params1;	// Include X0,Y0 (.xy) - X1,Y1 (.zw)
+	FShaderParameter Params2;	// ExcludeRect X0,Y0 (.xy) - X1,Y1 (.zw)
+	FShaderResourceParameter FillTexture;
+};
+
+class FLongGPUTaskPS : public FGlobalShader
+{
+	DECLARE_EXPORTED_SHADER_TYPE(FLongGPUTaskPS,Global,UTILITYSHADERS_API);
+public:
+	FLongGPUTaskPS( )	{ }
+	FLongGPUTaskPS(const ShaderMetaType::CompiledShaderInitializerType& Initializer)
+	:	FGlobalShader( Initializer )
+	{
+	}
+	
+	// FShader interface.
+	virtual bool Serialize(FArchive& Ar)
+	{
+		bool bShaderHasOutdatedParameters = FGlobalShader::Serialize(Ar);
+		return bShaderHasOutdatedParameters;
+	}
+	
+	static bool ShouldCache(EShaderPlatform Platform)
+	{
+		return IsFeatureLevelSupported(Platform, ERHIFeatureLevel::SM4);
+	}
+};
+
diff --git a/Engine/Source/Runtime/UtilityShaders/Public/ResolveShader.h b/Engine/Source/Runtime/UtilityShaders/Public/ResolveShader.h
new file mode 100644
--- /dev/null
+++ b/Engine/Source/Runtime/UtilityShaders/Public/ResolveShader.h
@@ -0,0 +1,117 @@
+// Copyright 1998-2015 Epic Games, Inc. All Rights Reserved.
+
+
+#pragma once
+
+#include "GlobalShader.h"
+#include "ShaderParameters.h"
+
+struct FDummyResolveParameter {};
+
+class FResolveDepthPS : public FGlobalShader
+{
+	DECLARE_EXPORTED_SHADER_TYPE(FResolveDepthPS, Global, UTILITYSHADERS_API);
+public:
+	
+	typedef FDummyResolveParameter FParameter;
+	
+	static bool ShouldCache(EShaderPlatform Platform) { return Platform == SP_PCD3D_SM5; }
+	
+	FResolveDepthPS(const ShaderMetaType::CompiledShaderInitializerType& Initializer):
+	FGlobalShader(Initializer)
+	{
+		UnresolvedSurface.Bind(Initializer.ParameterMap,TEXT("UnresolvedSurface"), SPF_Mandatory);
+	}
+	FResolveDepthPS() {}
+	
+	void SetParameters(FRHICommandList& RHICmdList, FParameter)
+	{
+	}
+	
+	virtual bool Serialize(FArchive& Ar)
+	{
+		bool bShaderHasOutdatedParameters = FGlobalShader::Serialize(Ar);
+		Ar << UnresolvedSurface;
+		return bShaderHasOutdatedParameters;
+	}
+	
+	FShaderResourceParameter UnresolvedSurface;
+};
+
+
+class FResolveDepthNonMSPS : public FGlobalShader
+{
+	DECLARE_EXPORTED_SHADER_TYPE(FResolveDepthNonMSPS, Global, UTILITYSHADERS_API);
+public:
+	
+	typedef FDummyResolveParameter FParameter;
+	
+	static bool ShouldCache(EShaderPlatform Platform) { return GetMaxSupportedFeatureLevel(Platform) <= ERHIFeatureLevel::SM4; }
+	
+	FResolveDepthNonMSPS(const ShaderMetaType::CompiledShaderInitializerType& Initializer):
+	FGlobalShader(Initializer)
+	{
+		UnresolvedSurface.Bind(Initializer.ParameterMap,TEXT("UnresolvedSurfaceNonMS"), SPF_Mandatory);
+	}
+	FResolveDepthNonMSPS() {}
+	
+	void SetParameters(FRHICommandList& RHICmdList, FParameter)
+	{
+	}
+	
+	virtual bool Serialize(FArchive& Ar)
+	{
+		bool bShaderHasOutdatedParameters = FGlobalShader::Serialize(Ar);
+		Ar << UnresolvedSurface;
+		return bShaderHasOutdatedParameters;
+	}
+	
+	FShaderResourceParameter UnresolvedSurface;
+};
+
+class FResolveSingleSamplePS : public FGlobalShader
+{
+	DECLARE_EXPORTED_SHADER_TYPE(FResolveSingleSamplePS, Global, UTILITYSHADERS_API);
+public:
+	
+	typedef uint32 FParameter;
+	
+	static bool ShouldCache(EShaderPlatform Platform) { return Platform == SP_PCD3D_SM5; }
+	
+	FResolveSingleSamplePS(const ShaderMetaType::CompiledShaderInitializerType& Initializer):
+	FGlobalShader(Initializer)
+	{
+		UnresolvedSurface.Bind(Initializer.ParameterMap,TEXT("UnresolvedSurface"), SPF_Mandatory);
+		SingleSampleIndex.Bind(Initializer.ParameterMap,TEXT("SingleSampleIndex"), SPF_Mandatory);
+	}
+	FResolveSingleSamplePS() {}
+	
+	UTILITYSHADERS_API void SetParameters(FRHICommandList& RHICmdList, uint32 SingleSampleIndexValue);
+	
+	virtual bool Serialize(FArchive& Ar)
+	{
+		bool bShaderHasOutdatedParameters = FGlobalShader::Serialize(Ar);
+		Ar << UnresolvedSurface;
+		Ar << SingleSampleIndex;
+		return bShaderHasOutdatedParameters;
+	}
+	
+	FShaderResourceParameter UnresolvedSurface;
+	FShaderParameter SingleSampleIndex;
+};
+
+/**
+ * A vertex shader for rendering a textured screen element.
+ */
+class FResolveVS : public FGlobalShader
+{
+	DECLARE_EXPORTED_SHADER_TYPE(FResolveVS, Global, UTILITYSHADERS_API);
+public:
+	
+	static bool ShouldCache(EShaderPlatform Platform) { return true; }
+	
+	FResolveVS(const ShaderMetaType::CompiledShaderInitializerType& Initializer):
+	FGlobalShader(Initializer)
+	{}
+	FResolveVS() {}
+};
diff --git a/Engine/Source/Runtime/UtilityShaders/UtilityShaders.Build.cs b/Engine/Source/Runtime/UtilityShaders/UtilityShaders.Build.cs
new file mode 100644
--- /dev/null
+++ b/Engine/Source/Runtime/UtilityShaders/UtilityShaders.Build.cs
@@ -0,0 +1,19 @@
+// Copyright 1998-2015 Epic Games, Inc. All Rights Reserved.
+
+using UnrealBuildTool;
+
+public class UtilityShaders : ModuleRules
+{
+	public UtilityShaders(TargetInfo Target)
+	{
+		PrivateDependencyModuleNames.AddRange(
+			new string[] {
+				"Core",
+				"Engine",
+				"RHI",
+				"RenderCore",
+				"ShaderCore",
+			}
+		);
+	}
+}
diff --git a/Engine/Source/Runtime/Windows/D3D11RHI/D3D11RHI.Build.cs b/Engine/Source/Runtime/Windows/D3D11RHI/D3D11RHI.Build.cs
--- a/Engine/Source/Runtime/Windows/D3D11RHI/D3D11RHI.Build.cs
+++ b/Engine/Source/Runtime/Windows/D3D11RHI/D3D11RHI.Build.cs
@@ -15,6 +15,7 @@ public class D3D11RHI : ModuleRules
 				"RHI",
 				"RenderCore",
 				"ShaderCore",
+				"UtilityShaders",
 			}
 			);
 
diff --git a/Engine/Source/Runtime/Windows/D3D11RHI/Private/D3D11RHI.cpp b/Engine/Source/Runtime/Windows/D3D11RHI/Private/D3D11RHI.cpp
--- a/Engine/Source/Runtime/Windows/D3D11RHI/Private/D3D11RHI.cpp
+++ b/Engine/Source/Runtime/Windows/D3D11RHI/Private/D3D11RHI.cpp
@@ -264,31 +264,6 @@ void FD3D11DynamicRHI::ClearAllShaderResources()
 	ClearAllShaderResourcesForFrequency<SF_Compute>();
 }
 
-class FLongGPUTaskPS : public FGlobalShader
-{
-	DECLARE_SHADER_TYPE(FLongGPUTaskPS,Global);
-public:
-	FLongGPUTaskPS( )	{ }
-	FLongGPUTaskPS(const ShaderMetaType::CompiledShaderInitializerType& Initializer)
-		:	FGlobalShader( Initializer )
-	{
-	}
-
-	// FShader interface.
-	virtual bool Serialize(FArchive& Ar)
-	{
-		bool bShaderHasOutdatedParameters = FGlobalShader::Serialize(Ar);
-		return bShaderHasOutdatedParameters;
-	}
-
-	static bool ShouldCache(EShaderPlatform Platform)
-	{
-		return IsFeatureLevelSupported(Platform, ERHIFeatureLevel::SM4);
-	}
-};
-
-IMPLEMENT_SHADER_TYPE(,FLongGPUTaskPS,TEXT("OneColorShader"),TEXT("MainLongGPUTask"),SF_Pixel);
-
 FGlobalBoundShaderState LongGPUTaskBoundShaderState;
 
 void FD3D11DynamicRHI::IssueLongGPUTask()
diff --git a/Engine/Source/Runtime/Windows/D3D11RHI/Private/D3D11RenderTarget.cpp b/Engine/Source/Runtime/Windows/D3D11RHI/Private/D3D11RenderTarget.cpp
--- a/Engine/Source/Runtime/Windows/D3D11RHI/Private/D3D11RenderTarget.cpp
+++ b/Engine/Source/Runtime/Windows/D3D11RHI/Private/D3D11RenderTarget.cpp
@@ -8,126 +8,7 @@
 #include "BatchedElements.h"
 #include "ScreenRendering.h"
 #include "RHIStaticStates.h"
-
-namespace
-{
-	struct FDummyResolveParameter {};
-
-	class FResolveDepthPS : public FGlobalShader
-	{
-		DECLARE_SHADER_TYPE(FResolveDepthPS,Global);
-	public:
-
-		typedef FDummyResolveParameter FParameter;
-
-		static bool ShouldCache(EShaderPlatform Platform) { return Platform == SP_PCD3D_SM5; }
-
-		FResolveDepthPS(const ShaderMetaType::CompiledShaderInitializerType& Initializer):
-			FGlobalShader(Initializer)
-		{
-			UnresolvedSurface.Bind(Initializer.ParameterMap,TEXT("UnresolvedSurface"), SPF_Mandatory);
-		}
-		FResolveDepthPS() {}
-
-		void SetParameters(FRHICommandList& RHICmdList, ID3D11DeviceContext* Direct3DDeviceContext,FParameter)
-		{
-		}
-
-		virtual bool Serialize(FArchive& Ar)
-		{
-			bool bShaderHasOutdatedParameters = FGlobalShader::Serialize(Ar);
-			Ar << UnresolvedSurface;
-			return bShaderHasOutdatedParameters;
-		}
-
-		FShaderResourceParameter UnresolvedSurface;
-	};
-	IMPLEMENT_SHADER_TYPE(,FResolveDepthPS,TEXT("ResolvePixelShader"),TEXT("MainDepth"),SF_Pixel);
-
-
-	class FResolveDepthNonMSPS : public FGlobalShader
-	{
-		DECLARE_SHADER_TYPE(FResolveDepthNonMSPS,Global);
-	public:
-
-		typedef FDummyResolveParameter FParameter;
-
-		static bool ShouldCache(EShaderPlatform Platform) { return GetMaxSupportedFeatureLevel(Platform) <= ERHIFeatureLevel::SM4; }
-
-		FResolveDepthNonMSPS(const ShaderMetaType::CompiledShaderInitializerType& Initializer):
-		FGlobalShader(Initializer)
-		{
-			UnresolvedSurface.Bind(Initializer.ParameterMap,TEXT("UnresolvedSurfaceNonMS"), SPF_Mandatory);
-		}
-		FResolveDepthNonMSPS() {}
-
-		void SetParameters(FRHICommandList& RHICmdList, ID3D11DeviceContext* Direct3DDeviceContext,FParameter)
-		{
-		}
-
-		virtual bool Serialize(FArchive& Ar)
-		{
-			bool bShaderHasOutdatedParameters = FGlobalShader::Serialize(Ar);
-			Ar << UnresolvedSurface;
-			return bShaderHasOutdatedParameters;
-		}
-
-		FShaderResourceParameter UnresolvedSurface;
-	};
-	IMPLEMENT_SHADER_TYPE(,FResolveDepthNonMSPS,TEXT("ResolvePixelShader"),TEXT("MainDepthNonMS"),SF_Pixel);
-
-	class FResolveSingleSamplePS : public FGlobalShader
-	{
-		DECLARE_SHADER_TYPE(FResolveSingleSamplePS,Global);
-	public:
-
-		typedef uint32 FParameter;
-
-		static bool ShouldCache(EShaderPlatform Platform) { return Platform == SP_PCD3D_SM5; }
-
-		FResolveSingleSamplePS(const ShaderMetaType::CompiledShaderInitializerType& Initializer):
-			FGlobalShader(Initializer)
-		{
-			UnresolvedSurface.Bind(Initializer.ParameterMap,TEXT("UnresolvedSurface"), SPF_Mandatory);
-			SingleSampleIndex.Bind(Initializer.ParameterMap,TEXT("SingleSampleIndex"), SPF_Mandatory);
-		}
-		FResolveSingleSamplePS() {}
-
-		void SetParameters(FRHICommandList& RHICmdList, ID3D11DeviceContext* Direct3DDeviceContext,uint32 SingleSampleIndexValue)
-		{
-			SetShaderValue(RHICmdList, GetPixelShader(),SingleSampleIndex,SingleSampleIndexValue);
-		}
-
-		virtual bool Serialize(FArchive& Ar)
-		{
-			bool bShaderHasOutdatedParameters = FGlobalShader::Serialize(Ar);
-			Ar << UnresolvedSurface;
-			Ar << SingleSampleIndex;
-			return bShaderHasOutdatedParameters;
-		}
-
-		FShaderResourceParameter UnresolvedSurface;
-		FShaderParameter SingleSampleIndex;
-	};
-	IMPLEMENT_SHADER_TYPE(,FResolveSingleSamplePS,TEXT("ResolvePixelShader"),TEXT("MainSingleSample"),SF_Pixel);
-
-	/**
-	 * A vertex shader for rendering a textured screen element.
-	 */
-	class FResolveVS : public FGlobalShader
-	{
-		DECLARE_SHADER_TYPE(FResolveVS,Global);
-	public:
-
-		static bool ShouldCache(EShaderPlatform Platform) { return true; }
-
-		FResolveVS(const ShaderMetaType::CompiledShaderInitializerType& Initializer):
-			FGlobalShader(Initializer)
-		{}
-		FResolveVS() {}
-	};
-	IMPLEMENT_SHADER_TYPE(,FResolveVS,TEXT("ResolveVertexShader"),TEXT("Main"),SF_Vertex);
-}
+#include "ResolveShader.h"
 
 static inline DXGI_FORMAT ConvertTypelessToUnorm(DXGI_FORMAT Format)
 {
@@ -254,7 +135,7 @@ void FD3D11DynamicRHI::ResolveTextureUsingShader(
 	TShaderMapRef<TPixelShader> ResolvePixelShader(ShaderMap);
 	SetGlobalBoundShaderState(RHICmdList, GMaxRHIFeatureLevel, ResolveBoundShaderState, GScreenVertexDeclaration.VertexDeclarationRHI, *ResolveVertexShader, *ResolvePixelShader);
 
-	ResolvePixelShader->SetParameters(RHICmdList, Direct3DDeviceContext,PixelShaderParameter);
+	ResolvePixelShader->SetParameters(RHICmdList, PixelShaderParameter);
 	RHICmdList.Flush(); // always call flush when using a command list in RHI implementations before doing anything else. This is super hazardous.
 
 	// Set the source texture.
