Add ensure for setting a non-RT as an RT.  Avoid a crash when CheckSRVs is enabled.

diff --git a/Engine/Source/Runtime/Windows/D3D11RHI/Private/D3D11Commands.cpp b/Engine/Source/Runtime/Windows/D3D11RHI/Private/D3D11Commands.cpp
--- a/Engine/Source/Runtime/Windows/D3D11RHI/Private/D3D11Commands.cpp
+++ b/Engine/Source/Runtime/Windows/D3D11RHI/Private/D3D11Commands.cpp
@@ -871,12 +871,16 @@ void FD3D11DynamicRHI::RHISetRenderTargets(
 			FD3D11TextureBase* NewRenderTarget = GetD3D11TextureFromRHITexture(NewRenderTargetsRHI[RenderTargetIndex].Texture);
 			RenderTargetView = NewRenderTarget->GetRenderTargetView(RTMipIndex, RTSliceIndex);
 
-#if CHECK_SRV_TRANSITIONS
-			// remember this target as having been bound for write.
-			ID3D11Resource* RTVResource;
-			RenderTargetView->GetResource(&RTVResource);			
-			UnresolvedTargets.Add(RTVResource, FUnresolvedRTInfo(NewRenderTargetsRHI[RenderTargetIndex].Texture->GetName(), RTMipIndex, 1, RTSliceIndex, 1));
-			RTVResource->Release();
+			ensureMsg(RenderTargetView, TEXT("Texture being set as render target has no RTV"));
+#if CHECK_SRV_TRANSITIONS			
+			if (RenderTargetView)
+			{
+				// remember this target as having been bound for write.
+				ID3D11Resource* RTVResource;
+				RenderTargetView->GetResource(&RTVResource);
+				UnresolvedTargets.Add(RTVResource, FUnresolvedRTInfo(NewRenderTargetsRHI[RenderTargetIndex].Texture->GetName(), RTMipIndex, 1, RTSliceIndex, 1));
+				RTVResource->Release();
+			}
 #endif
 			
 			// Unbind any shader views of the render target that are bound.
