#vr Fix for taking screenshots in VR with "shot"

diff --git a/Engine/Source/Runtime/Windows/D3D11RHI/Private/D3D11RenderTarget.cpp b/Engine/Source/Runtime/Windows/D3D11RHI/Private/D3D11RenderTarget.cpp
--- a/Engine/Source/Runtime/Windows/D3D11RHI/Private/D3D11RenderTarget.cpp
+++ b/Engine/Source/Runtime/Windows/D3D11RHI/Private/D3D11RenderTarget.cpp
@@ -468,6 +468,7 @@ static uint32 ComputeBytesPerPixel(DXGI_FORMAT Format)
 		case DXGI_FORMAT_B8G8R8A8_TYPELESS:
 		case DXGI_FORMAT_B8G8R8A8_UNORM:
 		case DXGI_FORMAT_R8G8B8A8_TYPELESS:
+		case DXGI_FORMAT_R8G8B8A8_UNORM:
 		case DXGI_FORMAT_R24G8_TYPELESS:
 		case DXGI_FORMAT_R10G10B10A2_UNORM:
 		case DXGI_FORMAT_R11G11B10_FLOAT:
@@ -660,7 +661,7 @@ static void ConvertRAWSurfaceDataToFColor(DXGI_FORMAT Format, uint32 Width, uint
 			}
 		}
 	}
-	else if(Format == DXGI_FORMAT_R8G8B8A8_TYPELESS)
+	else if(Format == DXGI_FORMAT_R8G8B8A8_TYPELESS || Format == DXGI_FORMAT_R8G8B8A8_UNORM)
 	{
 		// Read the data out of the buffer, converting it from ABGR to ARGB.
 		for(uint32 Y = 0; Y < Height; Y++)
