Add RHIReducePlatformGPUOverhead() which gives us the 4.5% Kinect GPU reserve on Xbox One. Added stubs for other platforms.

diff --git a/Engine/Source/Runtime/EmptyRHI/Private/EmptyCommands.cpp b/Engine/Source/Runtime/EmptyRHI/Private/EmptyCommands.cpp
--- a/Engine/Source/Runtime/EmptyRHI/Private/EmptyCommands.cpp
+++ b/Engine/Source/Runtime/EmptyRHI/Private/EmptyCommands.cpp
@@ -414,6 +414,11 @@ bool FEmptyDynamicRHI::RHIIsRenderingSuspended()
 	return false;
 }
 
+void FEmptyDynamicRHI::RHIReducePlatformGPUOverhead(bool)
+{
+	// Not supported
+}
+
 void FEmptyDynamicRHI::RHIBlockUntilGPUIdle()
 {
 
diff --git a/Engine/Source/Runtime/IOS/MetalRHI/Private/MetalCommands.cpp b/Engine/Source/Runtime/IOS/MetalRHI/Private/MetalCommands.cpp
--- a/Engine/Source/Runtime/IOS/MetalRHI/Private/MetalCommands.cpp
+++ b/Engine/Source/Runtime/IOS/MetalRHI/Private/MetalCommands.cpp
@@ -863,6 +863,11 @@ bool FMetalDynamicRHI::RHIIsRenderingSuspended()
 	return false;
 }
 
+void FMetalDynamicRHI::RHIReducePlatformGPUOverhead(bool)
+{
+	// Not supported
+}
+
 void FMetalDynamicRHI::RHIBlockUntilGPUIdle()
 {
 	NOT_SUPPORTED("RHIBlockUntilGPUIdle");
diff --git a/Engine/Source/Runtime/OpenGLDrv/Private/OpenGLCommands.cpp b/Engine/Source/Runtime/OpenGLDrv/Private/OpenGLCommands.cpp
--- a/Engine/Source/Runtime/OpenGLDrv/Private/OpenGLCommands.cpp
+++ b/Engine/Source/Runtime/OpenGLDrv/Private/OpenGLCommands.cpp
@@ -3395,6 +3395,11 @@ bool FOpenGLDynamicRHI::RHIIsRenderingSuspended()
 	return false;
 }
 
+void FOpenGLDynamicRHI::RHIReducePlatformGPUOverhead(bool)
+{
+	// Not supported
+}
+
 // Blocks the CPU until the GPU catches up and goes idle.
 void FOpenGLDynamicRHI::RHIBlockUntilGPUIdle()
 {
diff --git a/Engine/Source/Runtime/RHI/Public/RHIMethods.h b/Engine/Source/Runtime/RHI/Public/RHIMethods.h
--- a/Engine/Source/Runtime/RHI/Public/RHIMethods.h
+++ b/Engine/Source/Runtime/RHI/Public/RHIMethods.h
@@ -2095,6 +2095,12 @@ DEFINE_RHIMETHOD_0(
 	return,return false;
 	);
 
+// Operation to get platform specific GPU reserve from the system
+DEFINE_RHIMETHOD_1(
+	void,ReducePlatformGPUOverhead,
+	bool,bEnable,
+	return,return);
+
 /**
  *	Retrieve available screen resolutions.
  *
diff --git a/Engine/Source/Runtime/Windows/D3D11RHI/Private/D3D11Commands.cpp b/Engine/Source/Runtime/Windows/D3D11RHI/Private/D3D11Commands.cpp
--- a/Engine/Source/Runtime/Windows/D3D11RHI/Private/D3D11Commands.cpp
+++ b/Engine/Source/Runtime/Windows/D3D11RHI/Private/D3D11Commands.cpp
@@ -1942,6 +1942,11 @@ bool FD3D11DynamicRHI::RHIIsRenderingSuspended()
 	return false;
 }
 
+void FD3D11DynamicRHI::RHIReducePlatformGPUOverhead(bool)
+{
+	// Not supported
+}
+
 // Blocks the CPU until the GPU catches up and goes idle.
 void FD3D11DynamicRHI::RHIBlockUntilGPUIdle()
 {
