Fixing XBoxOne and PS4 Builds

diff --git a/Engine/Source/Runtime/Windows/XAudio2/Private/XAudio2Device.cpp b/Engine/Source/Runtime/Windows/XAudio2/Private/XAudio2Device.cpp
--- a/Engine/Source/Runtime/Windows/XAudio2/Private/XAudio2Device.cpp
+++ b/Engine/Source/Runtime/Windows/XAudio2/Private/XAudio2Device.cpp
@@ -163,10 +163,10 @@ bool FXAudio2Device::InitializeHardware()
 #else	//XAUDIO_SUPPORTS_DEVICE_DETAILS
 	// Create the final output voice
 	if (!ValidateAPICall(TEXT("CreateMasteringVoice"),
-			FXAudioDeviceProperties::XAudio2->CreateMasteringVoice(&FXAudioDeviceProperties::MasteringVoice, UE4_XAUDIO2_NUMCHANNELS, UE4_XAUDIO2_SAMPLERATE, 0, 0, NULL )))
+		DeviceProperties->XAudio2->CreateMasteringVoice(&DeviceProperties->MasteringVoice, UE4_XAUDIO2_NUMCHANNELS, UE4_XAUDIO2_SAMPLERATE, 0, 0, NULL )))
 	{
 		UE_LOG(LogInit, Warning, TEXT( "Failed to create the mastering voice for XAudio2" ) );
-		FXAudioDeviceProperties::XAudio2 = NULL;
+		DeviceProperties->XAudio2 = nullptr;
 		return false;
 	}
 #endif	//XAUDIO_SUPPORTS_DEVICE_DETAILS
