Ensure was incorrectly checking the same bound twice

diff --git a/Engine/Source/Runtime/Windows/D3D11RHI/Private/D3D11RHI.cpp b/Engine/Source/Runtime/Windows/D3D11RHI/Private/D3D11RHI.cpp
--- a/Engine/Source/Runtime/Windows/D3D11RHI/Private/D3D11RHI.cpp
+++ b/Engine/Source/Runtime/Windows/D3D11RHI/Private/D3D11RHI.cpp
@@ -168,7 +168,7 @@ void FD3D11DynamicRHI::CheckIfSRVIsResolved(ID3D11ShaderResourceView* SRV)
 		ensureMsgf(NumMips != -1, TEXT("SRV is set to read all mips.  Target %s is unresolved for mip %i"), *RTInfo.ResourceName.ToString(), RTInfo.MipLevel);
 
 		int32 RTLastSlice = RTInfo.ArraySlice + RTInfo.ArraySize - 1;
-		ensureMsgf((ArraySlice == -1 || ArraySlice == -1) || (LastSlice < RTInfo.ArraySlice || ArraySlice > RTLastSlice), TEXT("SRV is set to read slices in range %i to %i.  Target %s is unresolved for mip %i"), ArraySlice, LastSlice, *RTInfo.ResourceName.ToString(), RTInfo.ArraySlice);
+		ensureMsgf((ArraySlice == -1 || LastSlice == -1) || (LastSlice < RTInfo.ArraySlice || ArraySlice > RTLastSlice), TEXT("SRV is set to read slices in range %i to %i.  Target %s is unresolved for mip %i"), ArraySlice, LastSlice, *RTInfo.ResourceName.ToString(), RTInfo.ArraySlice);
 		ensureMsgf(ArraySlice == -1 || NumSlices != -1, TEXT("SRV is set to read all slices.  Target %s is unresolved for slice %i"));
 	}
 	SRVResource->Release();
