- Updated to using DXGI 1.1 (DXGIFactory1) for SteamVR support.
- Everything but the factory continues to use the 1.0 interfaces.
- This requires Vista SP2+, but we should already be falling back to OpenGL at that point, so hopefully this doesn't have any compatibility regressions.

diff --git a/Engine/Source/Runtime/Windows/D3D11RHI/Private/D3D11Device.cpp b/Engine/Source/Runtime/Windows/D3D11RHI/Private/D3D11Device.cpp
--- a/Engine/Source/Runtime/Windows/D3D11RHI/Private/D3D11Device.cpp
+++ b/Engine/Source/Runtime/Windows/D3D11RHI/Private/D3D11Device.cpp
@@ -34,8 +34,8 @@ TAutoConsoleVariable<int32> CVarD3D11ZeroBufferSizeInMB(
 	ECVF_ReadOnly
 	);
 
-FD3D11DynamicRHI::FD3D11DynamicRHI(IDXGIFactory* InDXGIFactory,D3D_FEATURE_LEVEL InFeatureLevel, int32 InChosenAdapter) :
-	DXGIFactory(InDXGIFactory),
+FD3D11DynamicRHI::FD3D11DynamicRHI(IDXGIFactory1* InDXGIFactory1,D3D_FEATURE_LEVEL InFeatureLevel, int32 InChosenAdapter) :
+	DXGIFactory1(InDXGIFactory1),
 	bDeviceRemoved(false),
 	FeatureLevel(InFeatureLevel),
 	bUsingTessellation(false),
@@ -235,7 +235,7 @@ void FD3D11DynamicRHI::RHIGetSupportedResolution( uint32 &Width, uint32 &Height
 	{
 		HRESULT hr = S_OK;
 		TRefCountPtr<IDXGIAdapter> Adapter;
-		hr = DXGIFactory->EnumAdapters(ChosenAdapter,Adapter.GetInitReference());
+		hr = DXGIFactory1->EnumAdapters(ChosenAdapter,Adapter.GetInitReference());
 		if( DXGI_ERROR_NOT_FOUND == hr )
 		{
 			return;
diff --git a/Engine/Source/Runtime/Windows/D3D11RHI/Private/D3D11RHIPrivate.h b/Engine/Source/Runtime/Windows/D3D11RHI/Private/D3D11RHIPrivate.h
--- a/Engine/Source/Runtime/Windows/D3D11RHI/Private/D3D11RHIPrivate.h
+++ b/Engine/Source/Runtime/Windows/D3D11RHI/Private/D3D11RHIPrivate.h
@@ -313,7 +313,7 @@ public:
 	TMap<FD3D11LockedKey,FD3D11LockedData> OutstandingLocks;
 
 	/** Initialization constructor. */
-	FD3D11DynamicRHI(IDXGIFactory* InDXGIFactory,D3D_FEATURE_LEVEL InFeatureLevel,int32 InChosenAdapter);
+	FD3D11DynamicRHI(IDXGIFactory1* InDXGIFactory1,D3D_FEATURE_LEVEL InFeatureLevel,int32 InChosenAdapter);
 
 	/** Destructor */
 	virtual ~FD3D11DynamicRHI() {}
@@ -350,9 +350,9 @@ public:
 		return Direct3DDeviceIMContext;
 	}
 
-	IDXGIFactory* GetFactory() const
+	IDXGIFactory1* GetFactory() const
 	{
-		return DXGIFactory;
+		return DXGIFactory1;
 	}
 private:
 	template <EShaderFrequency ShaderFrequency>
@@ -401,7 +401,7 @@ public:
 
 protected:
 	/** The global D3D interface. */
-	TRefCountPtr<IDXGIFactory> DXGIFactory;
+	TRefCountPtr<IDXGIFactory1> DXGIFactory1;
 
 	/** The global D3D device's immediate context */
 	TRefCountPtr<FD3D11DeviceContext> Direct3DDeviceIMContext;
diff --git a/Engine/Source/Runtime/Windows/D3D11RHI/Private/Windows/WindowsD3D11Device.cpp b/Engine/Source/Runtime/Windows/D3D11RHI/Private/Windows/WindowsD3D11Device.cpp
--- a/Engine/Source/Runtime/Windows/D3D11RHI/Private/Windows/WindowsD3D11Device.cpp
+++ b/Engine/Source/Runtime/Windows/D3D11RHI/Private/Windows/WindowsD3D11Device.cpp
@@ -61,16 +61,16 @@ static bool IsDelayLoadException(PEXCEPTION_POINTERS ExceptionPointers)
 	MSVC_PRAGMA(warning(disable:6322))
 #endif	// USING_CODE_ANALYSIS
 /**
- * Since CreateDXGIFactory is a delay loaded import from the D3D11 DLL, if the user
- * doesn't have Vista/DX10, calling CreateDXGIFactory will throw an exception.
+ * Since CreateDXGIFactory1 is a delay loaded import from the D3D11 DLL, if the user
+ * doesn't have VistaSP2/DX10, calling CreateDXGIFactory1 will throw an exception.
  * We use SEH to detect that case and fail gracefully.
  */
-static void SafeCreateDXGIFactory(IDXGIFactory** DXGIFactory)
+static void SafeCreateDXGIFactory(IDXGIFactory1** DXGIFactory1)
 {
 #if !D3D11_CUSTOM_VIEWPORT_CONSTRUCTOR
 	__try
 	{
-		CreateDXGIFactory(__uuidof(IDXGIFactory),(void**)DXGIFactory);
+		CreateDXGIFactory1(__uuidof(IDXGIFactory1),(void**)DXGIFactory1);
 	}
 	__except(IsDelayLoadException(GetExceptionInformation()))
 	{
@@ -225,10 +225,10 @@ void FD3D11DynamicRHIModule::FindAdapter()
 	// Once we chosen one we don't need to do it again.
 	check(!ChosenAdapter.IsValid());
 
-	// Try to create the DXGIFactory.  This will fail if we're not running Vista.
-	TRefCountPtr<IDXGIFactory> DXGIFactory;
-	SafeCreateDXGIFactory(DXGIFactory.GetInitReference());
-	if(!DXGIFactory)
+	// Try to create the DXGIFactory1.  This will fail if we're not running Vista SP2 or higher.
+	TRefCountPtr<IDXGIFactory1> DXGIFactory1;
+	SafeCreateDXGIFactory(DXGIFactory1.GetInitReference());
+	if(!DXGIFactory1)
 	{
 		return;
 	}
@@ -254,7 +254,7 @@ void FD3D11DynamicRHIModule::FindAdapter()
 	bool bIsAnyNVIDIA = false;
 
 	// Enumerate the DXGIFactory's adapters.
-	for(uint32 AdapterIndex = 0; DXGIFactory->EnumAdapters(AdapterIndex,TempAdapter.GetInitReference()) != DXGI_ERROR_NOT_FOUND; ++AdapterIndex)
+	for(uint32 AdapterIndex = 0; DXGIFactory1->EnumAdapters(AdapterIndex,TempAdapter.GetInitReference()) != DXGI_ERROR_NOT_FOUND; ++AdapterIndex)
 	{
 		// Check that if adapter supports D3D11.
 		if(TempAdapter)
@@ -355,10 +355,10 @@ void FD3D11DynamicRHIModule::FindAdapter()
 
 FDynamicRHI* FD3D11DynamicRHIModule::CreateRHI()
 {
-	TRefCountPtr<IDXGIFactory> DXGIFactory;
-	SafeCreateDXGIFactory(DXGIFactory.GetInitReference());
-	check(DXGIFactory);
-	return new FD3D11DynamicRHI(DXGIFactory,ChosenAdapter.MaxSupportedFeatureLevel,ChosenAdapter.AdapterIndex);
+	TRefCountPtr<IDXGIFactory1> DXGIFactory1;
+	SafeCreateDXGIFactory(DXGIFactory1.GetInitReference());
+	check(DXGIFactory1);
+	return new FD3D11DynamicRHI(DXGIFactory1,ChosenAdapter.MaxSupportedFeatureLevel,ChosenAdapter.AdapterIndex);
 }
 
 void FD3D11DynamicRHI::Init()
@@ -432,7 +432,7 @@ void FD3D11DynamicRHI::InitD3DDevice()
 
 		TRefCountPtr<IDXGIAdapter> EnumAdapter;
 
-		if(DXGIFactory->EnumAdapters(ChosenAdapter,EnumAdapter.GetInitReference()) != DXGI_ERROR_NOT_FOUND)
+		if(DXGIFactory1->EnumAdapters(ChosenAdapter,EnumAdapter.GetInitReference()) != DXGI_ERROR_NOT_FOUND)
 		{
 			if (EnumAdapter)// && EnumAdapter->CheckInterfaceSupport(__uuidof(ID3D11Device),NULL) == S_OK)
 			{
@@ -702,7 +702,7 @@ bool FD3D11DynamicRHI::RHIGetAvailableResolutions(FScreenResolutionArray& Resolu
 
 	HRESULT hr = S_OK;
 	TRefCountPtr<IDXGIAdapter> Adapter;
-	hr = DXGIFactory->EnumAdapters(ChosenAdapter,Adapter.GetInitReference());
+	hr = DXGIFactory1->EnumAdapters(ChosenAdapter,Adapter.GetInitReference());
 
 	if( DXGI_ERROR_NOT_FOUND == hr )
 		return false;
