Add support for RQT_Undefined in Xbox Query code which waits for the GPU to catch up.
Also renamed FD3D11OcclusionQuery to FD3D11RenderQuery to better fit its function and match the other RHIs.

diff --git a/Engine/Source/Runtime/Windows/D3D11RHI/Private/D3D11Query.cpp b/Engine/Source/Runtime/Windows/D3D11RHI/Private/D3D11Query.cpp
--- a/Engine/Source/Runtime/Windows/D3D11RHI/Private/D3D11Query.cpp
+++ b/Engine/Source/Runtime/Windows/D3D11RHI/Private/D3D11Query.cpp
@@ -34,13 +34,13 @@ FRenderQueryRHIRef FD3D11DynamicRHI::RHICreateRenderQuery(ERenderQueryType Query
 
 	Desc.MiscFlags = 0;
 	VERIFYD3D11RESULT(Direct3DDevice->CreateQuery(&Desc,Query.GetInitReference()));
-	return new FD3D11OcclusionQuery(Query, QueryType);
+	return new FD3D11RenderQuery(Query, QueryType);
 }
 
 bool FD3D11DynamicRHI::RHIGetRenderQueryResult(FRenderQueryRHIParamRef QueryRHI,uint64& OutResult,bool bWait)
 {
 	check(IsInRenderingThread());
-	DYNAMIC_CAST_D3D11RESOURCE(OcclusionQuery,Query);
+	DYNAMIC_CAST_D3D11RESOURCE(RenderQuery,Query);
 
 	bool bSuccess = true;
 	if(!Query->bResultIsCached)
@@ -69,7 +69,7 @@ bool FD3D11DynamicRHI::RHIGetRenderQueryResult(FRenderQueryRHIParamRef QueryRHI,
 // Occlusion/Timer queries.
 void FD3D11DynamicRHI::RHIBeginRenderQuery(FRenderQueryRHIParamRef QueryRHI)
 {
-	DYNAMIC_CAST_D3D11RESOURCE(OcclusionQuery,Query);
+	DYNAMIC_CAST_D3D11RESOURCE(RenderQuery,Query);
 
 	if(Query->QueryType == RQT_Occlusion)
 	{
@@ -85,14 +85,13 @@ void FD3D11DynamicRHI::RHIBeginRenderQuery(FRenderQueryRHIParamRef QueryRHI)
 
 void FD3D11DynamicRHI::RHIEndRenderQuery(FRenderQueryRHIParamRef QueryRHI)
 {
-	DYNAMIC_CAST_D3D11RESOURCE(OcclusionQuery,Query);
+	DYNAMIC_CAST_D3D11RESOURCE(RenderQuery,Query);
 	Query->bResultIsCached = false; // for occlusion queries, this is redundant with the one in begin
 	Direct3DDeviceIMContext->End(Query->Resource);
 
 	//@todo - d3d debug spews warnings about OQ's that are being issued but not polled, need to investigate
 }
 
-
 bool FD3D11DynamicRHI::GetQueryData(ID3D11Query* Query,void* Data,SIZE_T DataSize,bool bWait, ERenderQueryType QueryType)
 {
 	// Request the data from the query.
diff --git a/Engine/Source/Runtime/Windows/D3D11RHI/Public/D3D11Resources.h b/Engine/Source/Runtime/Windows/D3D11RHI/Public/D3D11Resources.h
--- a/Engine/Source/Runtime/Windows/D3D11RHI/Public/D3D11Resources.h
+++ b/Engine/Source/Runtime/Windows/D3D11RHI/Public/D3D11Resources.h
@@ -554,8 +554,8 @@ inline FD3D11TextureBase* GetD3D11TextureFromRHITexture(FRHITexture* Texture)
 	}
 }
 
-/** D3D11 occlusion query */
-class FD3D11OcclusionQuery : public FRHIRenderQuery
+/** D3D11 render query */
+class FD3D11RenderQuery : public FRHIRenderQuery
 {
 public:
 
@@ -572,7 +572,7 @@ public:
 	ERenderQueryType QueryType;
 
 	/** Initialization constructor. */
-	FD3D11OcclusionQuery(ID3D11Query* InResource, ERenderQueryType InQueryType):
+	FD3D11RenderQuery(ID3D11Query* InResource, ERenderQueryType InQueryType):
 		Resource(InResource),
 		Result(0),
 		bResultIsCached(false),
