Implemented GPU profiler interface for PS4
* Now we get proper formatting and other gpu profiler functionality that already works for other RHI's (r.profilegpu.root)

diff --git a/Engine/Source/Runtime/RHI/Private/GPUProfiler.cpp b/Engine/Source/Runtime/RHI/Private/GPUProfiler.cpp
--- a/Engine/Source/Runtime/RHI/Private/GPUProfiler.cpp
+++ b/Engine/Source/Runtime/RHI/Private/GPUProfiler.cpp
@@ -194,7 +194,7 @@ void FGPUProfilerEventNodeFrame::DumpEventTree()
 		// Display a warning if this is a GPU profile and the GPU was profiled with v-sync enabled
 		FText VsyncEnabledWarningText = FText::GetEmpty();
 		static IConsoleVariable* CVSyncVar = IConsoleManager::Get().FindConsoleVariable(TEXT("r.VSync"));
-		if (CVSyncVar->GetInt() != 0)
+		if (CVSyncVar->GetInt() != 0 && !PlatformDisablesVSync())
 		{
 			VsyncEnabledWarningText = LOCTEXT("GpuProfileVsyncEnabledWarning",
 				"WARNING: This GPU profile was captured with v-sync enabled.  V-sync wait time may show up in any bucket, and as a result the data in this profile may be skewed. Please profile with v-sync disabled to obtain the most accurate data.");
diff --git a/Engine/Source/Runtime/RHI/Public/GPUProfiler.h b/Engine/Source/Runtime/RHI/Public/GPUProfiler.h
--- a/Engine/Source/Runtime/RHI/Public/GPUProfiler.h
+++ b/Engine/Source/Runtime/RHI/Public/GPUProfiler.h
@@ -83,6 +83,8 @@ struct RHI_API FGPUProfilerEventNodeFrame
 
 	/** D3D11 Hack */
 	virtual void LogDisjointQuery() {}
+
+	virtual bool PlatformDisablesVSync() const { return false; }
 };
 
 /**
diff --git a/Engine/Source/Runtime/Windows/D3D11RHI/Private/D3D11RHI.cpp b/Engine/Source/Runtime/Windows/D3D11RHI/Private/D3D11RHI.cpp
--- a/Engine/Source/Runtime/Windows/D3D11RHI/Private/D3D11RHI.cpp
+++ b/Engine/Source/Runtime/Windows/D3D11RHI/Private/D3D11RHI.cpp
@@ -551,73 +551,6 @@ void FD3DGPUProfiler::PopEvent()
 	FGPUProfiler::PopEvent();
 }
 
-int32 FD3DGPUProfiler::RecordEventTimestamp(ID3D11Device* Direct3DDevice, ID3D11DeviceContext* Direct3DDeviceIMContext)
-{
-	check(CurrentGPUProfile);
-
-	D3D11_QUERY_DESC TimestampQueryDesc;
-	TimestampQueryDesc.Query = D3D11_QUERY_TIMESTAMP;
-	TimestampQueryDesc.MiscFlags = 0;
-
-	TRefCountPtr<ID3D11Query> TimestampQuery;
-	VERIFYD3D11RESULT(Direct3DDevice->CreateQuery(&TimestampQueryDesc,TimestampQuery.GetInitReference()));
-
-	Direct3DDeviceIMContext->End(TimestampQuery);
-
-	return CurrentGPUProfile->EventTimestampQueries.Add(TimestampQuery);
-}
-
-bool FD3DGPUProfiler::FinishProfiling(TMap<int32, double>& OutEventIdToTimestampMap, ID3D11DeviceContext* Direct3DDeviceIMContext)
-{
-	bool bSuccess = true;
-
-	check(CurrentGPUProfile);
-
-	// Get the disjoint timestamp query result.
-	CurrentGPUProfile->DisjointQuery.EndTracking();
-	D3D11_QUERY_DATA_TIMESTAMP_DISJOINT DisjointQueryResult = CurrentGPUProfile->DisjointQuery.GetResult();
-	if(DisjointQueryResult.Disjoint)
-	{
-		bSuccess = false;
-	}
-	else
-	{
-		// Read back the timestamp query results.
-		for(int32 EventIndex = 0;EventIndex < CurrentGPUProfile->EventTimestampQueries.Num();++EventIndex)
-		{
-			TRefCountPtr<ID3D11Query> EventTimestampQuery = CurrentGPUProfile->EventTimestampQueries[EventIndex];
-
-			uint64 EventTimestamp;
-			const double StartTime = FPlatformTime::Seconds();
-			HRESULT QueryGetDataResult;
-			while(true)
-			{
-				QueryGetDataResult = Direct3DDeviceIMContext->GetData(EventTimestampQuery,&EventTimestamp,sizeof(EventTimestamp),0);
-				if(QueryGetDataResult == S_FALSE && (FPlatformTime::Seconds() - StartTime) < 0.5)
-				{
-					FPlatformProcess::Sleep(0.005f);
-				}
-				else
-				{
-					break;
-				}
-			}
-			if(QueryGetDataResult != S_OK)
-			{
-				bSuccess = false;
-				break;
-			}
-			else
-			{
-				OutEventIdToTimestampMap.Add(EventIndex,(double)EventTimestamp / (double)DisjointQueryResult.Frequency);
-			}
-		}
-	}
-
-	CurrentGPUProfile.Reset();
-	return bSuccess;
-}
-
 /** Start this frame of per tracking */
 void FD3D11EventNodeFrame::StartFrame()
 {
diff --git a/Engine/Source/Runtime/Windows/D3D11RHI/Private/D3D11RHIPrivate.h b/Engine/Source/Runtime/Windows/D3D11RHI/Private/D3D11RHIPrivate.h
--- a/Engine/Source/Runtime/Windows/D3D11RHI/Private/D3D11RHIPrivate.h
+++ b/Engine/Source/Runtime/Windows/D3D11RHI/Private/D3D11RHIPrivate.h
@@ -231,27 +231,6 @@ public:
 	FD3D11DisjointTimeStampQuery DisjointQuery;
 };
 
-/** The state of a pending GPU timestamp profile. */
-class FD3D11GPUProfile
-{
-public:
-
-	FD3D11GPUProfile(FD3D11DynamicRHI* D3D11RHI)
-	: DisjointQuery(D3D11RHI)
-	{
-		  DisjointQuery.InitResource();
-	}
-
-	~FD3D11GPUProfile()
-	{
-		  DisjointQuery.ReleaseResource();
-	}
-
-	FD3D11DisjointTimeStampQuery DisjointQuery;
-	TArray<TRefCountPtr<ID3D11Query> > EventTimestampQueries;
-};
-
-
 /** 
  * Encapsulates GPU profiling logic and data. 
  * There's only one global instance of this struct so it should only contain global data, nothing specific to a frame.
@@ -261,9 +240,6 @@ struct FD3DGPUProfiler : public FGPUProfiler
 	/** Used to measure GPU time per frame. */
 	FD3D11BufferedGPUTiming FrameTiming;
 
-	/** Current GPU profiling state. */
-	TScopedPointer<FD3D11GPUProfile> CurrentGPUProfile;
-
 	class FD3D11DynamicRHI* D3D11RHI;
 
 	/** GPU hitch profile histories */
@@ -290,16 +266,6 @@ struct FD3DGPUProfiler : public FGPUProfiler
 	void BeginFrame(class FD3D11DynamicRHI* InRHI);
 
 	void EndFrame();
-
-	void StartProfiling(FD3D11DynamicRHI* D3D11RHI)
-	{
-		CurrentGPUProfile.Reset(new FD3D11GPUProfile(D3D11RHI));
-		CurrentGPUProfile->DisjointQuery.StartTracking();
-	}
-
-	int32 RecordEventTimestamp(ID3D11Device* Direct3DDevice, ID3D11DeviceContext* Direct3DDeviceIMContext);
-
-	bool FinishProfiling(TMap<int32, double>& OutEventIdToTimestampMap, ID3D11DeviceContext* Direct3DDeviceIMContext);
 };
 
 /** The interface which is implemented by the dynamically bound RHI. */
