Unreal Audio: Low level audio I/O

- First check in for new audio system
- Creation of new UnrealAudio module interface and implementation
- UnrealAudioDevice module interface (for platform-dependent audio device I/O)
- Intialize/Shutdown of new UnrealAudio module

- Windows WASAPI implementation of audio device interface
  - Support for querying audio devices
  - Support for streaming audio from input devices
  - Support for streaming audio to output devices
  - Support for mutliple output channels (up to 7.1 for now)

- Beginnings of XAudio2 implementation of audio device interface

- A bunch of unit tests testing audio input/output streaming

- A commandlet to run said unit tests
  - To run commandlet use "[UE4Editor executable] -command UnrealEd.AudioTestCommandlet [arguments]"
  - arguments to use (and their test) are:
    - "device query"      - Prints out information about every connected audio device
    - "device out"          - Plays a decaying sinusoid in harmonic multiples at different rates on available output channels
    - "device out_fm"    - Plays an FM synthesizer bank with randomized parameters on all available output channels
    - "device out_pan"  - Plays a single noise source (with LFO LP filter) "3D" panning clockwise around output channels
    - "device in"             - Plays mic input into all output channels (simple passthrough)
    - "device in_delay"  - Plays mic input through a multi-tap delay line with random delays through all output channels
    - "device all"            - Runs all the above with 10 seconds for each test

diff --git a/Engine/Source/Editor/UnrealEd/Classes/Commandlets/AudioTestCommandlet.h b/Engine/Source/Editor/UnrealEd/Classes/Commandlets/AudioTestCommandlet.h
new file mode 100644
--- /dev/null
+++ b/Engine/Source/Editor/UnrealEd/Classes/Commandlets/AudioTestCommandlet.h
@@ -0,0 +1,15 @@
+// Copyright 1998-2015 Epic Games, Inc. All Rights Reserved.
+
+#pragma once
+#include "Commandlets/Commandlet.h"
+#include "AudioTestCommandlet.generated.h"
+
+UCLASS()
+class UAudioTestCommandlet : public UCommandlet
+{
+	GENERATED_UCLASS_BODY()
+
+	// Begin UCommandlet Interface
+	virtual int32 Main(const FString& Params) override;
+	// End UCommandlet Interface
+};
diff --git a/Engine/Source/Editor/UnrealEd/Private/Commandlets/AudioTestCommandlet.cpp b/Engine/Source/Editor/UnrealEd/Private/Commandlets/AudioTestCommandlet.cpp
new file mode 100644
--- /dev/null
+++ b/Engine/Source/Editor/UnrealEd/Private/Commandlets/AudioTestCommandlet.cpp
@@ -0,0 +1,214 @@
+// Copyright 1998-2015 Epic Games, Inc. All Rights Reserved.
+
+#include "UnrealEd.h"
+#include "UnrealAudioModule.h"
+#include "UnrealAudioTests.h"
+
+DEFINE_LOG_CATEGORY_STATIC(AudioTestCommandlet, Log, All);
+
+#if ENABLE_UNREAL_AUDIO
+
+/** Unreal audio module ptr */
+static UAudio::IUnrealAudioModule* UnrealAudioModule = nullptr;
+
+static bool UnrealAudioLoad()
+{
+	UnrealAudioModule = FModuleManager::LoadModulePtr<UAudio::IUnrealAudioModule>(FName("UnrealAudio"));
+	if (UnrealAudioModule)
+	{
+		UnrealAudioModule->Initialize();
+	}
+	return UnrealAudioModule != nullptr;
+}
+
+static bool UnrealAudioUnload()
+{
+	if (UnrealAudioModule)
+	{
+		UnrealAudioModule->Shutdown();
+		return true;
+	}
+	return false;
+}
+
+/**
+* Test static functions which call into module test code
+*/
+
+static bool TestAudioDeviceAll()
+{
+	if (!UAudio::TestDeviceQuery())
+	{
+		return false;
+	}
+
+	if (!UAudio::TestDeviceOutputSimple(10))
+	{
+		return false;
+	}
+
+	if (!UAudio::TestDeviceOutputRandomizedFm(10))
+	{
+		return false;
+	}
+
+	if (!UAudio::TestDeviceOutputNoisePan(10))
+	{
+		return false;
+	}
+
+	if (!UAudio::TestDeviceInputSimple(10))
+	{
+		return false;
+	}
+
+	if (!UAudio::TestDeviceInputRandomizedDelay(10))
+	{
+		return false;
+	}
+	return true;
+}
+
+static bool TestAudioDeviceQuery()
+{
+	return UAudio::TestDeviceQuery();
+}
+
+static bool TestAudioDeviceOutputSimple()
+{
+	return UAudio::TestDeviceOutputSimple(-1);
+}
+
+static bool TestAudioDeviceInputSimple()
+{
+	return UAudio::TestDeviceInputSimple(-1);
+}
+
+static bool TestAudioDeviceOutputFm()
+{
+	return UAudio::TestDeviceOutputRandomizedFm(-1);
+}
+
+static bool TestAudioDeviceOutputPan()
+{
+	return UAudio::TestDeviceOutputNoisePan(-1);
+}
+
+static bool TestAudioDeviceInputDelay()
+{
+	return UAudio::TestDeviceInputRandomizedDelay(-1);
+}
+
+/**
+* Setting up commandlet code
+*/
+
+typedef bool AudioTestFunction();
+
+struct AudioTestInfo
+{
+	FString CategoryName;
+	FString TestName;
+	bool (*TestFunction)();
+};
+
+enum EAudioTests
+{
+	AUDIO_TEST_DEVICE_ALL,
+	AUDIO_TEST_DEVICE_QUERY,
+	AUDIO_TEST_DEVICE_OUTPUT_SIMPLE,
+	AUDIO_TEST_DEVICE_OUTPUT_FM,
+	AUDIO_TEST_DEVICE_OUTPUT_PAN,
+	AUDIO_TEST_DEVICE_INPUT_SIMPLE,
+	AUDIO_TEST_DEVICE_INPUT_DELAY,
+	AUDIO_TESTS
+};
+
+static AudioTestInfo AudioTestInfoList[] =
+{
+	{ "device", "all",			TestAudioDeviceAll			}, // AUDIO_TEST_DEVICE_ALL
+	{ "device", "query",		TestAudioDeviceQuery		}, // AUDIO_TEST_DEVICE_QUERY
+	{ "device", "out",			TestAudioDeviceOutputSimple }, // AUDIO_TEST_DEVICE_OUTPUT_SIMPLE
+	{ "device", "out_fm",		TestAudioDeviceOutputFm		}, // AUDIO_TEST_DEVICE_OUTPUT_FM
+	{ "device", "out_pan",		TestAudioDeviceOutputPan	}, // AUDIO_TEST_DEVICE_OUTPUT_PAN
+	{ "device", "in",			TestAudioDeviceInputSimple	}, // AUDIO_TEST_DEVICE_INPUT_SIMPLE
+	{ "device", "in_delay",		TestAudioDeviceInputDelay	}, // AUDIO_TEST_DEVICE_INPUT_DELAY
+};
+static_assert(ARRAY_COUNT(AudioTestInfoList) == AUDIO_TESTS, "Mismatched info list and test enumeration");
+
+static void PrintUsage()
+{
+	UE_LOG(AudioTestCommandlet, Display, TEXT("AudioTestCommandlet Usage: {Editor}.exe UnrealEd.AudioTestCommandlet {testcategory} {test}"));
+	UE_LOG(AudioTestCommandlet, Display, TEXT("Possible Tests: [testcategory]: [test]"));
+	for (uint32 Index = 0; Index < AUDIO_TESTS; ++Index)
+	{
+		UE_LOG(AudioTestCommandlet, Display, TEXT("    %s: %s"), *AudioTestInfoList[Index].CategoryName, *AudioTestInfoList[Index].TestName);
+	}
+}
+
+#endif // #if ENABLE_UNREAL_AUDIO
+
+// -- UAudioTestCommandlet Functions -------------------
+
+UAudioTestCommandlet::UAudioTestCommandlet(const FObjectInitializer& ObjectInitializer)
+	: Super(ObjectInitializer)
+{
+}
+
+int32 UAudioTestCommandlet::Main(const FString& InParams)
+{
+#if ENABLE_UNREAL_AUDIO
+	// Parse command line.
+	TArray<FString> Tokens;
+	TArray<FString> Switches;
+	UCommandlet::ParseCommandLine(*InParams, Tokens, Switches);
+
+	if (Tokens.Num() != 3)
+	{
+		PrintUsage();
+		return 0;
+	}
+
+	if (!UnrealAudioLoad())
+	{
+		UE_LOG(AudioTestCommandlet, Display, TEXT("Failed to load unreal audio module. Exiting."));
+		return 0;
+	}
+
+	check(UnrealAudioModule != nullptr);
+
+	bool FoundTest = false;
+	for (uint32 Index = 0; Index < AUDIO_TESTS; ++Index)
+	{
+		if (AudioTestInfoList[Index].CategoryName == Tokens[1])
+		{
+			if (AudioTestInfoList[Index].TestName == Tokens[2])
+			{
+				FoundTest = true;
+				if (AudioTestInfoList[Index].TestFunction())
+				{
+					UE_LOG(AudioTestCommandlet, Display, TEXT("Test %s succeeded."), *AudioTestInfoList[Index].TestName);
+				}
+				else
+				{
+					UE_LOG(AudioTestCommandlet, Display, TEXT("Test %s failed."), *AudioTestInfoList[Index].TestName);
+				}
+				break;
+			}
+		}
+	}
+
+	if (!FoundTest)
+	{
+		UE_LOG(AudioTestCommandlet, Display, TEXT("Unknown category or test. Exiting."));
+		return 0;
+	}
+
+	UnrealAudioUnload();
+#else
+	UE_LOG(AudioTestCommandlet, Display, TEXT("Unreal Audio Module Not Enabled For This Platform"));
+#endif // #if ENABLE_UNREAL_AUDIO
+
+
+	return 0;
+}
diff --git a/Engine/Source/Editor/UnrealEd/UnrealEd.Build.cs b/Engine/Source/Editor/UnrealEd/UnrealEd.Build.cs
--- a/Engine/Source/Editor/UnrealEd/UnrealEd.Build.cs
+++ b/Engine/Source/Editor/UnrealEd/UnrealEd.Build.cs
@@ -76,6 +76,7 @@ public class UnrealEd : ModuleRules
                 "AIModule",
 				"BlueprintGraph",
                 "Http",
+				"UnrealAudio",
 			}
 		);
 
@@ -227,6 +228,7 @@ public class UnrealEd : ModuleRules
 			(Target.Platform == UnrealTargetPlatform.Win32))
 		{
 			PublicDependencyModuleNames.Add("XAudio2");
+			PublicDependencyModuleNames.Add("UnrealAudioWasapi");
 
 			AddThirdPartyPrivateStaticDependencies(Target, 
 				"UEOgg",
diff --git a/Engine/Source/Runtime/UnrealAudio/Private/Tests/UnrealAudioDeviceTests.cpp b/Engine/Source/Runtime/UnrealAudio/Private/Tests/UnrealAudioDeviceTests.cpp
new file mode 100644
--- /dev/null
+++ b/Engine/Source/Runtime/UnrealAudio/Private/Tests/UnrealAudioDeviceTests.cpp
@@ -0,0 +1,380 @@
+// Copyright 1998-2015 Epic Games, Inc. All Rights Reserved.
+
+#pragma once 
+
+#include "UnrealAudioPrivate.h"
+#include "UnrealAudioModule.h"
+#include "UnrealAudioDeviceModule.h"
+#include "UnrealAudioTests.h"
+#include "UnrealAudioTestGenerators.h"
+
+#if ENABLE_UNREAL_AUDIO
+
+namespace UAudio
+{
+	static FUnrealAudioModule* UnrealAudioModule = nullptr;
+	static bool bTestInitialized = false;
+
+	void FUnrealAudioModule::InitializeTests(FUnrealAudioModule* Module)
+	{
+		UnrealAudioModule = Module;
+		bTestInitialized = true;
+	}
+
+	bool TestDeviceQuery()
+	{
+		if (!bTestInitialized)
+		{
+			return false;
+		}
+
+		UE_LOG(LogUnrealAudio, Display, TEXT("Running audio device query test..."));
+
+		IUnrealAudioDeviceModule* UnrealAudioDevice = UnrealAudioModule->GetDeviceModule();
+
+		if (!UnrealAudioDevice)
+		{
+			UE_LOG(LogUnrealAudio, Display, TEXT("Failed: No Audio Device Object."));
+			return false;
+		}
+
+		
+		UE_LOG(LogUnrealAudio, Display, TEXT("Querying which audio device API we're using..."), );
+		EDeviceApi::Type DeviceApi;
+		if (!UnrealAudioDevice->GetDevicePlatformApi(DeviceApi))
+		{
+			UE_LOG(LogUnrealAudio, Display, TEXT("Failed."));
+			return false;
+		}
+
+		const TCHAR* ApiName = EDeviceApi::ToString(DeviceApi);
+		UE_LOG(LogUnrealAudio, Display, TEXT("Success: Using %s"), ApiName);
+
+		if (DeviceApi == EDeviceApi::DUMMY)
+		{
+			UE_LOG(LogUnrealAudio, Display, TEXT("This is a dummy API. Platform not implemented yet."));
+			return true;
+		}
+
+		UE_LOG(LogUnrealAudio, Display, TEXT("Querying the number of output devices for current system..."), );
+		uint32 NumOutputDevices = 0;
+		if (!UnrealAudioDevice->GetNumOutputDevices(NumOutputDevices))
+		{
+			UE_LOG(LogUnrealAudio, Display, TEXT("Failed."));
+			return false;
+		}
+		UE_LOG(LogUnrealAudio, Display, TEXT("Success: %d Number of Output Devices"), NumOutputDevices);
+
+		UE_LOG(LogUnrealAudio, Display, TEXT("Retrieving output device info for each output device..."), );
+		for (uint32 DeviceIndex = 0; DeviceIndex < NumOutputDevices; ++DeviceIndex)
+		{
+			FDeviceInfo DeviceInfo;
+			if (!UnrealAudioDevice->GetOutputDeviceInfo(DeviceIndex, DeviceInfo))
+			{
+				UE_LOG(LogUnrealAudio, Display, TEXT("Failed for index %d."), DeviceIndex);
+				return false;
+			}
+
+			UE_LOG(LogUnrealAudio, Display, TEXT("Device Index: %d"), DeviceIndex);
+			UE_LOG(LogUnrealAudio, Display, TEXT("Device Name: %s"), *DeviceInfo.FriendlyName);
+			UE_LOG(LogUnrealAudio, Display, TEXT("Device FrameRate: %d"), DeviceInfo.FrameRate);
+			UE_LOG(LogUnrealAudio, Display, TEXT("Device NumChannels: %d"), DeviceInfo.NumChannels);
+			UE_LOG(LogUnrealAudio, Display, TEXT("Device Default?: %s"), DeviceInfo.bIsSystemDefault ? TEXT("yes") : TEXT("no"));
+			UE_LOG(LogUnrealAudio, Display, TEXT("Device Native Format: %s"), EStreamFormat::ToString(DeviceInfo.StreamFormat));
+
+			UE_LOG(LogUnrealAudio, Display, TEXT("Speaker Array:"));
+
+			for (int32 Channel = 0; Channel < DeviceInfo.Speakers.Num(); ++Channel)
+			{
+				UE_LOG(LogUnrealAudio, Display, TEXT("    %s"), ESpeaker::ToString(DeviceInfo.Speakers[Channel]));
+			}
+
+
+		}
+
+		UE_LOG(LogUnrealAudio, Display, TEXT("Querying the number of output devices for current system..."), );
+		uint32 NumInputDevices = 0;
+		if (!UnrealAudioDevice->GetNumInputDevices(NumInputDevices))
+		{
+			UE_LOG(LogUnrealAudio, Display, TEXT("Failed."));
+			return false;
+		}
+		UE_LOG(LogUnrealAudio, Display, TEXT("Success: %d Number of Input Devices"), NumInputDevices);
+
+		UE_LOG(LogUnrealAudio, Display, TEXT("Retrieving input device info for each input device..."), );
+		for (uint32 DeviceIndex = 0; DeviceIndex < NumInputDevices; ++DeviceIndex)
+		{
+			FDeviceInfo DeviceInfo;
+			if (!UnrealAudioDevice->GetInputDeviceInfo(DeviceIndex, DeviceInfo))
+			{
+				UE_LOG(LogUnrealAudio, Display, TEXT("Failed for index %d."), DeviceIndex);
+				return false;
+			}
+
+			UE_LOG(LogUnrealAudio, Display, TEXT("Device Index: %d"), DeviceIndex);
+			UE_LOG(LogUnrealAudio, Display, TEXT("Device Name: %s"), *DeviceInfo.FriendlyName);
+			UE_LOG(LogUnrealAudio, Display, TEXT("Device FrameRate: %d"), DeviceInfo.FrameRate);
+			UE_LOG(LogUnrealAudio, Display, TEXT("Device NumChannels: %d"), DeviceInfo.NumChannels);
+			UE_LOG(LogUnrealAudio, Display, TEXT("Device Default?: %s"), DeviceInfo.bIsSystemDefault ? TEXT("yes") : TEXT("no"));
+			UE_LOG(LogUnrealAudio, Display, TEXT("Device Native Format: %s"), EStreamFormat::ToString(DeviceInfo.StreamFormat));
+		}
+
+		return true;
+	}
+
+	/**
+	* TestCallback
+	* Callback function used for all device test functions. 
+	* Function takes input params and passes to the ITestGenerator object to generate output buffers.
+	*/
+	static bool TestCallback(FCallbackInfo& CallbackInfo)	
+	{
+		static int32 CurrentSecond = -1;
+		int32 StreamSecond = (int32)CallbackInfo.StreamTime;
+		if (StreamSecond != CurrentSecond)
+		{
+			if (CurrentSecond == -1)
+			{
+				UE_LOG(LogUnrealAudio, Display, TEXT("Stream Time (seconds):"));
+			}
+			CurrentSecond = StreamSecond;
+			UE_LOG(LogUnrealAudio, Display, TEXT("%d"), CurrentSecond);
+		}
+
+		// Sets any data used by lots of different generators in static data struct
+		Test::UpdateCallbackData(CallbackInfo);
+
+		// Cast the user data object to our own FSimpleFM object (since we know we used that when we created the stream)
+		Test::IGenerator* Generator = (Test::IGenerator*)CallbackInfo.UserData;
+
+		// Get the next buffer of output data
+		Generator->GetNextBuffer(CallbackInfo);
+		return true;
+	}
+
+	static bool DoOutputTest(const FString& TestName, double LifeTime, Test::IGenerator* Generator)
+	{
+		check(Generator);
+
+		if (!bTestInitialized)
+		{
+			return false;
+		}
+
+		UE_LOG(LogUnrealAudio, Display, TEXT("Running audio test: \"%s\"..."), *TestName);
+
+		if (LifeTime > 0.0)
+		{
+			UE_LOG(LogUnrealAudio, Display, TEXT("Time of test: %d (seconds)"), LifeTime);
+		}
+		else
+		{
+			UE_LOG(LogUnrealAudio, Display, TEXT("Time of test: (indefinite)"));
+		}
+
+		IUnrealAudioDeviceModule* UnrealAudioDevice = UnrealAudioModule->GetDeviceModule();
+
+		EDeviceApi::Type DeviceApi;
+		if (!UnrealAudioDevice->GetDevicePlatformApi(DeviceApi))
+		{
+			UE_LOG(LogUnrealAudio, Display, TEXT("Failed."));
+			return false;
+		}
+		if (DeviceApi == EDeviceApi::DUMMY)
+		{
+			UE_LOG(LogUnrealAudio, Display, TEXT("This is a dummy API. Platform not implemented yet."));
+			return true;
+		}
+
+		uint32 DefaultDeviceIndex = 0;
+		if (!UnrealAudioDevice->GetDefaultOutputDeviceIndex(DefaultDeviceIndex))
+		{
+			UE_LOG(LogUnrealAudio, Display, TEXT("Failed to get default device index."));
+			return false;
+		}
+
+		FDeviceInfo DeviceInfo;
+		if (!UnrealAudioDevice->GetOutputDeviceInfo(DefaultDeviceIndex, DeviceInfo))
+		{
+			UE_LOG(LogUnrealAudio, Display, TEXT("Failed to get default device info."));
+			return false;
+		}
+
+		FCreateStreamParams CreateStreamParams;
+		CreateStreamParams.OutputDeviceIndex = DefaultDeviceIndex;
+		CreateStreamParams.CallbackFunction = &TestCallback;
+		CreateStreamParams.UserData = (void*)Generator;
+		CreateStreamParams.CallbackBlockSize = 1024;
+
+		UE_LOG(LogUnrealAudio, Display, TEXT("Creating the stream..."));
+		if (!UnrealAudioDevice->CreateStream(CreateStreamParams))
+		{
+			UE_LOG(LogUnrealAudio, Display, TEXT("Failed to create a stream."));
+			return false;
+		}
+
+		UE_LOG(LogUnrealAudio, Display, TEXT("Starting the stream."));
+		UnrealAudioDevice->StartStream();
+
+		// Block until the synthesizer is done
+		while (!Generator->IsDone())
+		{
+			Sleep(1);
+		}
+
+		// Stop the output stream (which blocks until its actually freed)
+		UE_LOG(LogUnrealAudio, Display, TEXT("Shutting down the stream."));
+		if (!UnrealAudioDevice->StopStream())
+		{
+			UE_LOG(LogUnrealAudio, Display, TEXT("Failed to stop the device stream."));
+			return false;
+		}
+
+		if (!UnrealAudioDevice->ShutdownStream())
+		{
+			UE_LOG(LogUnrealAudio, Display, TEXT("Failed to shut down the device stream."));
+			return false;
+		}
+
+		// At this point audio device I/O is done
+		UE_LOG(LogUnrealAudio, Display, TEXT("Success!"));
+		return true;
+	}
+
+	static bool DoInputTest(const FString& TestName, double LifeTime, Test::IGenerator* Generator)
+	{
+		check(Generator);
+
+		if (!bTestInitialized)
+		{
+			return false;
+		}
+
+		UE_LOG(LogUnrealAudio, Display, TEXT("Running audio test: \"%s\"..."), *TestName);
+
+		if (LifeTime > 0.0)
+		{
+			UE_LOG(LogUnrealAudio, Display, TEXT("Time of test: %d (seconds)"), LifeTime);
+		}
+		else
+		{
+			UE_LOG(LogUnrealAudio, Display, TEXT("Time of test: (indefinite)"));
+		}
+
+		IUnrealAudioDeviceModule* UnrealAudioDevice = UnrealAudioModule->GetDeviceModule();
+
+		EDeviceApi::Type DeviceApi;
+		if (!UnrealAudioDevice->GetDevicePlatformApi(DeviceApi))
+		{
+			UE_LOG(LogUnrealAudio, Display, TEXT("Failed."));
+			return false;
+		}
+		if (DeviceApi == EDeviceApi::DUMMY)
+		{
+			UE_LOG(LogUnrealAudio, Display, TEXT("This is a dummy API. Platform not implemented yet."));
+			return true;
+		}
+
+		uint32 DefaultOutputDeviceIndex = 0;
+		if (!UnrealAudioDevice->GetDefaultOutputDeviceIndex(DefaultOutputDeviceIndex))
+		{
+			UE_LOG(LogUnrealAudio, Display, TEXT("Failed to get default output device index."));
+			return false;
+		}
+
+		FDeviceInfo OutputDeviceInfo;
+		if (!UnrealAudioDevice->GetOutputDeviceInfo(DefaultOutputDeviceIndex, OutputDeviceInfo))
+		{
+			UE_LOG(LogUnrealAudio, Display, TEXT("Failed to get default output device info."));
+			return false;
+		}
+
+		uint32 DefaultInputDeviceIndex = 0;
+		if (!UnrealAudioDevice->GetDefaultInputDeviceIndex(DefaultInputDeviceIndex))
+		{
+			UE_LOG(LogUnrealAudio, Display, TEXT("Failed to get default input device index."));
+			return false;
+		}
+
+		FDeviceInfo InputDeviceInfo;
+		if (!UnrealAudioDevice->GetInputDeviceInfo(DefaultInputDeviceIndex, InputDeviceInfo))
+		{
+			UE_LOG(LogUnrealAudio, Display, TEXT("Failed to get default input device info."));
+			return false;
+		}
+
+		FCreateStreamParams CreateStreamParams;
+		CreateStreamParams.InputDeviceIndex = DefaultInputDeviceIndex;
+		CreateStreamParams.OutputDeviceIndex = DefaultOutputDeviceIndex;
+		CreateStreamParams.CallbackFunction = &TestCallback;
+		CreateStreamParams.UserData = (void*)Generator;
+
+		UE_LOG(LogUnrealAudio, Display, TEXT("Creating the stream..."));
+		if (!UnrealAudioDevice->CreateStream(CreateStreamParams))
+		{
+			UE_LOG(LogUnrealAudio, Display, TEXT("Failed to create a stream."));
+			return false;
+		}
+
+		UE_LOG(LogUnrealAudio, Display, TEXT("Starting the stream."));
+		UnrealAudioDevice->StartStream();
+
+		// Block until the synthesizer is done
+		while (!Generator->IsDone())
+		{
+			Sleep(1);
+		}
+
+		// Stop the output stream (which blocks until its actually freed)
+		UE_LOG(LogUnrealAudio, Display, TEXT("Shutting down the stream."));
+		if (!UnrealAudioDevice->StopStream())
+		{
+			UE_LOG(LogUnrealAudio, Display, TEXT("Failed to stop the device stream."));
+			return false;
+		}
+
+		if (!UnrealAudioDevice->ShutdownStream())
+		{
+			UE_LOG(LogUnrealAudio, Display, TEXT("Failed to shut down the device stream."));
+			return false;
+		}
+
+		// At this point audio device I/O is done
+		UE_LOG(LogUnrealAudio, Display, TEXT("Success!"));
+
+		return true;
+	}
+
+	bool TestDeviceOutputSimple(double LifeTime)
+	{
+		Test::FSimpleOutput SimpleOutput(LifeTime);
+		return DoOutputTest("output randomized FM synthesis", LifeTime, &SimpleOutput);
+	}
+
+	bool TestDeviceInputSimple(double LifeTime)
+	{
+		Test::FPassThrough SimplePassThrough(LifeTime);
+		return DoInputTest("input pass through", LifeTime, &SimplePassThrough);
+	}
+
+	bool TestDeviceOutputRandomizedFm(double LifeTime)
+	{
+		Test::FPhaseModulator RandomizedFmGenerator(LifeTime);
+		return DoOutputTest("output randomized FM synthesis", LifeTime, &RandomizedFmGenerator);
+	}
+
+	bool TestDeviceOutputNoisePan(double LifeTime)
+	{
+		Test::FNoisePan SimpleWhiteNoisePan(LifeTime);
+		return DoOutputTest("output white noise panner", LifeTime, &SimpleWhiteNoisePan);
+	}
+
+	bool UNREALAUDIO_API TestDeviceInputRandomizedDelay(double LifeTime)
+	{
+		Test::FInputDelay SimpleDelayer(LifeTime);
+		return DoInputTest("input randomized delay", LifeTime, &SimpleDelayer);
+	}
+
+}
+
+#endif // #if ENABLE_UNREAL_AUDIO
diff --git a/Engine/Source/Runtime/UnrealAudio/Private/Tests/UnrealAudioTestGenerators.cpp b/Engine/Source/Runtime/UnrealAudio/Private/Tests/UnrealAudioTestGenerators.cpp
new file mode 100644
--- /dev/null
+++ b/Engine/Source/Runtime/UnrealAudio/Private/Tests/UnrealAudioTestGenerators.cpp
@@ -0,0 +1,1124 @@
+// Copyright 1998-2015 Epic Games, Inc. All Rights Reserved.
+
+#pragma once 
+
+#include "UnrealAudioPrivate.h"
+#include "UnrealAudioModule.h"
+#include "UnrealAudioTestGenerators.h"
+
+#if ENABLE_UNREAL_AUDIO
+
+#ifndef TWO_PI
+#define TWO_PI	6.28318530718
+#endif
+
+#ifndef PI_OVER_TWO
+#define PI_OVER_TWO 1.57079632679
+#endif
+
+namespace UAudio
+{
+namespace Test
+{
+	/************************************************************************/
+	/* Static Data															*/
+	/************************************************************************/
+
+	struct FCallbackData
+	{
+		float FrameRate;
+		int32 NumFrames;
+		int32 NumChannels;
+		double Time;
+		TArray<ESpeaker::Type> Speakers;
+		bool bInitialized;
+
+		FCallbackData()
+			: bInitialized(false)
+		{
+		}
+
+	} CallbackData;
+
+	/************************************************************************/
+	/* Static Helper Functions                                              */
+	/************************************************************************/
+
+	static float WrapTwoPi(float Value)
+	{
+		while (Value >= TWO_PI)
+		{
+			return Value -= TWO_PI;
+		}
+
+		while (Value < 0)
+		{
+			return Value += TWO_PI;
+		}
+		return Value;
+	}
+
+	static float Lerp(float StartX, float EndX, float StartY, float EndY, float Alpha)
+	{
+		if (EndX - StartX <= 0.0f)
+		{
+			return StartY;
+		}
+		float Delta = (Alpha - StartX) / (EndX - StartX);
+		return Delta * EndY + (1.0f - Delta)*StartY;
+	}
+
+	void UpdateCallbackData(FCallbackInfo& CallbackInfo)
+	{
+		// Only update static data once
+		if (!CallbackData.bInitialized)
+		{
+			CallbackData.bInitialized = true;
+			CallbackData.FrameRate = (float)CallbackInfo.FrameRate;
+			CallbackData.NumFrames = CallbackInfo.NumFrames;
+			CallbackData.NumChannels = CallbackInfo.NumChannels;
+			CallbackData.Speakers = CallbackInfo.OutputSpeakers;
+		}
+
+		// Update dynamic data every time
+		CallbackData.Time = CallbackInfo.StreamTime;
+	}
+
+	/************************************************************************/
+	/* FTimer																*/
+	/************************************************************************/
+
+	FTimer::FTimer(double TotalTime)
+		: TotalTime(TotalTime)
+		, StartTime(0.0f)
+		, CurrentTime(0.0f)
+		, LastTime(0.0f)
+	{
+	}
+
+	FTimer::~FTimer()
+	{
+	}
+
+	bool FTimer::Update()
+	{
+		LastTime = CallbackData.Time;
+		CurrentTime = LastTime - StartTime;
+		return IsDone();
+	}
+
+	bool FTimer::IsDone() const
+	{
+		bool bIsDone = TotalTime >= 0.0f && CurrentTime >= TotalTime;
+		if (bIsDone)
+		{
+			printf("test");
+		}
+		return bIsDone;
+	}
+
+	void FTimer::Start(double InTotalTime)
+	{
+		StartTime = CallbackData.Time;
+		TotalTime = InTotalTime;
+		CurrentTime = 0.0;
+	}
+
+	void FTimer::Reset()
+	{
+		StartTime = CallbackData.Time;
+	}
+
+	float FTimer::GetTotalTime() const
+	{
+		return TotalTime;
+	}
+
+	/************************************************************************/
+	/* FFader																*/
+	/************************************************************************/
+
+	FFader::FFader(float FadeInEase, float FadeOutEase)
+		: FadeInEase(FadeInEase)
+		, FadeOutEase(FadeOutEase)
+		, FadeValue(0.0f)
+		, FadeState(EFadeState::OFF)
+	{
+	}
+
+	FFader::~FFader()
+	{
+	}
+
+	float FFader::Update()
+	{
+		if (FadeState == EFadeState::FADE_IN)
+		{
+			FadeValue += (1.0f - FadeValue) * FadeInEase;
+			if (1.0f - FadeValue < 0.0001f)
+			{
+				FadeState = EFadeState::ON;
+				FadeValue = 1.0f;
+			}
+		}
+		else if (FadeState == EFadeState::FADE_OUT)
+		{
+			FadeValue += (0.0f - FadeValue) * FadeOutEase;
+			if (FadeValue < 0.0001f)
+			{
+				FadeValue = 0.0f;
+				FadeState = EFadeState::OFF;
+			}
+		}
+		return FadeValue;
+
+	}
+
+	void FFader::Fade(bool bFadeIn)
+	{
+		if (bFadeIn)
+		{
+			FadeState = EFadeState::FADE_IN;
+		}
+		else
+		{
+			FadeState = EFadeState::FADE_OUT;
+		}
+	}
+
+	void FFader::Toggle()
+	{
+		Fade(FadeState == EFadeState::FADE_OUT || FadeState == EFadeState::OFF);
+	}
+
+	/************************************************************************/
+	/* FBiquad																*/
+	/************************************************************************/
+
+	FBiquad::FBiquad()
+	{
+		memset((void *)this, 0, sizeof(FBiquad));
+	}
+
+	FBiquad::~FBiquad()
+	{
+	}
+
+	float FBiquad::Update(float Value)
+	{
+		// A biquad has difference equation:
+		// y(n) = a0 * x(n) + a1 * x(n-1) + a2 * x(n-2) - b1 * y(n-1) - b2 * y(n-2)
+
+		float Output = A0 * Value + A1 * X1 + A2 * X2 - B1 * Y1 - B2 * Y2;
+		if ((Output > 0.0f && Output < FLT_MIN) || (Output < 0.0f && Output > -FLT_MIN))
+		{
+			Output = 0.0f;
+		}
+		Y2 = Y1;
+		Y1 = Output;
+		X2 = X1;
+		X1 = Value;
+		return Output;
+	}
+
+	/************************************************************************/
+	/* FLowPass																*/
+	/************************************************************************/
+
+	FLowPass::FLowPass()
+		: Frequency(0.0f)
+		, Quality(1.0f)
+	{
+	}
+
+	FLowPass::~FLowPass()
+	{
+	}
+
+	void FLowPass::SetFrequency(float InFrequency)
+	{
+		SetParams(InFrequency, Quality);
+	}
+
+	void FLowPass::SetQuality(float InQuality)
+	{
+		SetParams(Frequency, InQuality);
+	}
+
+	void FLowPass::SetParams(float InFrequency, float InQuality)
+	{
+		Frequency = InFrequency;
+		Quality = InQuality;
+
+		check(CallbackData.FrameRate > 0.0f);
+
+		float Theta = TWO_PI * Frequency / CallbackData.FrameRate;
+		float InverseQ = 0.5f / Quality;
+
+		float Temp = InverseQ * FMath::Sin(Theta);
+		float Beta = 0.5f * (1.0f - Temp) / (1.0f + Temp);
+		float Gamma = (0.5f + Beta) * FMath::Cos(Theta);
+		float Alpha = 0.5f * (0.5f + Beta - Gamma);
+
+		A0 = Alpha;
+		A1 = 2.0f * Alpha;
+		A2 = A0;
+		B1 = -2.0f * Gamma;
+		B2 = 2.0f * Beta;
+	}
+
+	/************************************************************************/
+	/* FRamp																*/
+	/************************************************************************/
+
+	FRamp::FRamp()
+		: CurrValue(0.0f)
+		, DeltaValue(0.0f)
+		, Frequency(0.0f)
+		, Scale(1.0f)
+		, Add(0.0f)
+	{
+	}
+
+	FRamp::~FRamp()
+	{
+	}
+
+	void FRamp::SetFrequency(float InFrequency)
+	{
+		check(CallbackData.FrameRate > 0.0f);
+		Frequency = InFrequency;
+		DeltaValue = Frequency / CallbackData.FrameRate;
+	}
+
+	void FRamp::SetScaleAdd(float InScale, float InAdd)
+	{
+		Scale = InScale;
+		Add = InAdd;
+	}
+
+	float FRamp::Update()
+	{
+		CurrValue += DeltaValue;
+		if (CurrValue >= 1.0f)
+		{
+			CurrValue -= 1.0f;
+		}
+		else if (CurrValue <= 0.0f)
+		{
+			CurrValue += 1.0f;
+		}
+		return Scale * CurrValue + Add;
+	}
+
+	/************************************************************************/
+	/* FSineOsc																*/
+	/************************************************************************/
+
+	FSineOsc::FSineOsc()
+		: Frequency(0.0f)
+		, Phase(0.0f)
+		, PhaseDelta(0.0f)
+		, TargetPhaseDelta(0.0f)
+		, Scale(1.0f)
+		, Add(0.0f)
+		, bNewValue(false)
+	{
+	}
+
+	FSineOsc::~FSineOsc()
+	{
+	}
+
+	void FSineOsc::SetFrequency(float InFrequency)
+	{
+		bool bIsInit = (Frequency == 0.0f);
+		Frequency = InFrequency;
+		if (bIsInit)
+		{
+			check(CallbackData.FrameRate > 0.0f);
+			PhaseDelta = TWO_PI * Frequency / CallbackData.FrameRate;
+			TargetPhaseDelta = PhaseDelta;
+			PhaseDeltaDelta = 0.0f;
+			Phase = 0.0f;
+			bNewValue = false;
+		}
+		else
+		{
+			TargetPhaseDelta = TWO_PI * Frequency / CallbackData.FrameRate;
+			PhaseDeltaDelta = (TargetPhaseDelta - PhaseDelta) / 100.0f;
+			bNewValue = true;
+		}
+	}
+
+	void FSineOsc::SetScaleAdd(float InScale, float InAdd)
+	{
+		Scale = InScale;
+		Add = InAdd;
+	}
+
+	void FSineOsc::SetOutputRange(float Min, float Max)
+	{
+		Add = Min + 1.0f;
+		Scale = 0.5f * (Max - Min);
+	}
+
+	float FSineOsc::Update()
+	{
+		Phase += PhaseDelta;
+
+		if (bNewValue)
+		{
+			if (FMath::Abs<float>(PhaseDelta - TargetPhaseDelta) < 0.00001f)
+			{
+				PhaseDelta = TargetPhaseDelta;
+				bNewValue = false;
+			}
+			else
+			{
+				PhaseDelta += PhaseDeltaDelta;
+			}
+		}
+
+		Phase = WrapTwoPi(Phase);
+		return Scale * FMath::Sin(Phase) + Add;
+	}
+
+	/************************************************************************/
+	/* FPan																	*/
+	/************************************************************************/
+
+	FPan::FPan()
+		: Pan(0.0f)
+		, LFEIndex(-1)
+		, PrevSpeakerIndex(-1)
+		, NumNonLfeSpeakers(0)
+	{
+	}
+
+	FPan::~FPan()
+	{
+	}
+
+	void FPan::Init(float InPan)
+	{
+		Pan = InPan;
+		SpeakerMap.Init(0, CallbackData.NumChannels);
+
+		NumNonLfeSpeakers = 0;
+		for (int32 Speaker = 0; Speaker < CallbackData.NumChannels; ++Speaker)
+		{
+			if (CallbackData.Speakers[Speaker] == ESpeaker::LOW_FREQUENCY)
+			{
+				LFEIndex = Speaker;
+			}
+			else
+			{
+				NumNonLfeSpeakers++;
+			}
+		}
+	}
+
+	void FPan::SetPan(float InPan)
+	{
+		Pan = InPan;
+	}
+
+	void FPan::Spatialize(float Value, float* Frame)
+	{
+		// Scale the pan value by the total number of (non-lfe) speakers
+		float SpeakerFractionTotal = Pan * (float)NumNonLfeSpeakers;
+
+		// Cast to integer to get the speaker on the left side of the sound
+		int32 SpeakerIndex = (int)SpeakerFractionTotal % NumNonLfeSpeakers;
+
+		// Get fraction of pan from this speaker to next
+		float SpeakerFraction = SpeakerFractionTotal - (float)SpeakerIndex;
+
+		// Use simple equal power panning to compute the amount of pan between the speakers
+		float LeftAmount = FMath::Cos(SpeakerFraction * PI_OVER_TWO);
+		float RightAmount = FMath::Sin(SpeakerFraction * PI_OVER_TWO);
+
+		// Now Build a speaker map
+
+		// Zero out the current speaker map if the speaker indices have changed
+		if (PrevSpeakerIndex != SpeakerIndex)
+		{
+			if (SpeakerMap.Num() == 0)
+			{
+				SpeakerMap.Init(0, NumNonLfeSpeakers);
+			}
+			else
+			{
+				// Zero out the speaker map if the speaker indices changed
+				memset((void*)SpeakerMap.GetData(), 0, sizeof(float)*NumNonLfeSpeakers);
+			}
+		}
+		PrevSpeakerIndex = SpeakerIndex;
+
+		// Now map the values to the correct index
+		SpeakerMap[SpeakerIndex] = LeftAmount;
+
+		int32 NextIndex = (SpeakerIndex + 1) % NumNonLfeSpeakers;
+		SpeakerMap[NextIndex] = RightAmount;
+
+		// Use speaker map to mix in the value of the synth
+		int32 MapIndex = 0;
+		for (int32 Channel = 0; Channel < CallbackData.NumChannels; ++Channel)
+		{
+			if (Channel == LFEIndex)
+			{
+				Frame[Channel] += Value;
+			}
+			else
+			{
+				float PanScale = SpeakerMap[MapIndex];
+				if (PanScale > 0.0f)
+				{
+					// Get the actual output speaker for the given map index
+					int32 ChannelIndex = GetOutputSpeaker(MapIndex);
+					Frame[ChannelIndex] += PanScale * Value;
+				}
+				++MapIndex;
+			}
+		}
+	}
+
+	int32 FPan::GetOutputSpeaker(int32 MapIndex)
+	{
+		if (CallbackData.NumChannels == 2)
+		{
+			return StereoSpeakerIndexMap[MapIndex];
+		}
+		else if (CallbackData.NumChannels == 4)
+		{
+			return QuadSpeakerIndexMap[MapIndex];
+		}
+		else if (CallbackData.NumChannels == 6)
+		{
+			return FiveOneSpeakerIndexMap[MapIndex];
+		}
+		else if (CallbackData.NumChannels == 8)
+		{
+			return SevenOneSpeakerIndexMap[MapIndex];
+		}
+		return 0;
+	}
+
+	int32 FPan::StereoSpeakerIndexMap[] =
+	{
+		0, // LEFT CHANNEL
+		1, // RIGHT CHANNEL
+	};
+
+	int32 FPan::QuadSpeakerIndexMap[] =
+	{
+		0, // LEFT CHANNEL
+		1, // RIGHT CHANNEL
+		3, // BACK RIGHT CHANNEL
+		2, // BACK LEFT CHANNEL
+	};
+
+	int32 FPan::FiveOneSpeakerIndexMap[] =
+	{
+		0,	// LEFT CHANNEL
+		2,	// CENTER CHANNEL
+		1,	// RIGHT CHANNEL
+		5,	// RIGHT BACK CHANNEL
+		4	// LEFT BACK CHANNEL
+	};
+
+	int32 FPan::SevenOneSpeakerIndexMap[] =
+	{
+		0,  // LEFT CHANNEL
+		2,  // CENTER CHANNEL
+		1,	// RIGHT CHANNEL
+		7,	// RIGHT SIDE
+		5,	// RIGHT BACK
+		4,	// LEFT BACK
+		6,	// LEFT SIDE
+	};
+
+	/************************************************************************/
+	/* FDelay																*/
+	/************************************************************************/
+
+	FDelay::FDelay()
+		: MaxLengthSec(0.0f)
+		, LengthFrames(0)
+		, WriteIndex(0)
+	{
+	}
+
+	FDelay::~FDelay()
+	{
+	}
+
+	void FDelay::Init(int32 NumTaps, float InMaxLengthSec)
+	{
+		check(CallbackData.FrameRate > 0);
+
+		if (LengthFrames != 0)
+		{
+			return;
+		}
+
+		MaxLengthSec = InMaxLengthSec;
+		LengthFrames = (int32)(CallbackData.FrameRate * InMaxLengthSec) + 1;
+
+		DelayBuffer.Init(0.0f, LengthFrames);
+		Taps.Init(FTap(), NumTaps);
+	}
+
+	void FDelay::SetDelayLength(int32 Tap, float LengthSec)
+	{
+		LengthSec = FMath::Clamp(LengthSec, 0.0f, MaxLengthSec);
+
+		check(Tap < Taps.Num());
+		Taps[Tap].DelayFrames = LengthSec * CallbackData.FrameRate;
+		Taps[Tap].ReadIndex = WriteIndex - (int32)Taps[Tap].DelayFrames;
+		if (Taps[Tap].ReadIndex < 0)
+		{
+			Taps[Tap].ReadIndex += LengthFrames;
+		}
+	}
+
+	void FDelay::SetWet(int32 Tap, float WetLevel)
+	{
+		check(Tap < Taps.Num());
+		Taps[Tap].Wet = FMath::Clamp(WetLevel, 0.0f, 1.0f);
+	}
+
+	void FDelay::SetFeedback(int32 Tap, float FeedbackLevel)
+	{
+		check(Tap < Taps.Num());
+		Taps[Tap].Feedback = FMath::Clamp(FeedbackLevel, 0.0f, 1.0f);
+	}
+
+	void FDelay::GetOutput(float InSample, TArray<float>& TapOutput)
+	{
+		float Xn = InSample;
+
+		check(TapOutput.Num() == Taps.Num());
+
+		// previous output read in from delay line
+		for (int32 TapIndex = 0; TapIndex < Taps.Num(); ++TapIndex)
+		{
+			FTap& Tap = Taps[TapIndex];
+			float Yn = DelayBuffer[Tap.ReadIndex];
+			if (Tap.ReadIndex == WriteIndex && Tap.DelayFrames < 1.0f)
+			{
+				Yn = Xn;
+			}
+
+			// Get sample index less than current read index
+			int32 ReadIndexPrev = Tap.ReadIndex - 1;
+
+			// Wrap if needed
+			if (ReadIndexPrev < 0)
+			{
+				ReadIndexPrev = LengthFrames - 1;
+			}
+
+			float YnPrev = DelayBuffer[ReadIndexPrev];
+			float Alpha = Tap.DelayFrames - (int)Tap.DelayFrames;
+			float Interp = Lerp(0.0f, 1.0f, Yn, YnPrev, Alpha);
+			Yn = Interp;
+
+			// Write input value into the delay buffer (with previous output scaled according to feedback)
+			DelayBuffer[WriteIndex] = Xn + Tap.Feedback * Yn;
+
+			Tap.ReadIndex = (Tap.ReadIndex + 1) % LengthFrames;
+			TapOutput[TapIndex] = Tap.Wet * Yn + (1.0f - Tap.Wet) * Xn;
+		}
+
+		// Update read/write indices
+		WriteIndex = (WriteIndex + 1) % LengthFrames;
+	}
+
+	/************************************************************************/
+	/* IGenerator															*/
+	/************************************************************************/
+
+	IGenerator::IGenerator(double LifeTime)
+		: LifeTimer(LifeTime)
+		, bIsDone(false)
+		, bIsInit(false)
+	{
+	}
+
+	bool IGenerator::IsDone() const
+	{
+		return bIsDone;
+	}
+
+	bool IGenerator::UpdateTimers()
+	{
+		if (bIsDone)
+		{
+			return false;
+		}
+
+		if (!bIsInit)
+		{
+			bIsInit = true;
+			LifeTimer.Reset();
+		}
+
+		if (LifeTimer.Update())
+		{
+			bIsDone = true;
+			return false;
+		}
+
+		return true;
+	}
+
+
+	/************************************************************************/
+	/* FSimpleOutput														*/
+	/************************************************************************/
+
+	FSimpleOutput::FSimpleOutput(double LifeTime)
+		: IGenerator(LifeTime)
+		, BaseFreqHz(220.0f)
+		, PhaseDelta(0.0f)
+		, Amplitude(0.5f)
+		, NumNonLfeSpeakers(0)
+		, LFEIndex(-1)
+	{
+	}
+
+	FSimpleOutput::~FSimpleOutput()
+	{
+	}
+
+	bool FSimpleOutput::IsInit()
+	{
+		return NumNonLfeSpeakers == 0;
+	}
+
+	bool FSimpleOutput::GetNextBuffer(FCallbackInfo& CallbackInfo)
+	{
+		if (!UpdateTimers())
+		{
+			return true;
+		}
+
+		// phase delta is radians per frame, or amount to increment the current phase each frame
+		if (IsInit())
+		{
+			check(CallbackData.FrameRate > 0);
+			PhaseDelta = (TWO_PI * BaseFreqHz) / CallbackData.FrameRate;
+			NumNonLfeSpeakers = 0;
+			for (int32 Chan = 0; Chan < CallbackInfo.NumChannels; ++Chan)
+			{
+				if (CallbackInfo.OutputSpeakers[Chan] == ESpeaker::LOW_FREQUENCY)
+				{
+					LFEIndex = Chan;
+				}
+				else
+				{
+					++NumNonLfeSpeakers;
+				}
+			}
+
+			SpeakerInfo.Init(FSpeakerInfo(), NumNonLfeSpeakers);
+			for (int32 Index = 0; Index < NumNonLfeSpeakers; ++Index)
+			{
+				SpeakerInfo[Index].MaxHarmonic = FMath::RandRange(3, 12);
+				SpeakerInfo[Index].Time = (float) FMath::RandRange(1, 8) / 16.0f;
+			}
+		}
+
+		for (int32 Index = 0; Index < SpeakerInfo.Num(); ++Index)
+		{
+			FSpeakerInfo& Info = SpeakerInfo[Index++];
+			if (Info.Timer.Update())
+			{
+				Info.bOn = !Info.bOn;
+				if (Info.bOn)
+				{
+					if (Info.bDirection)
+					{
+						Info.Harmonic++;
+					}
+					else
+					{
+						Info.Harmonic--;
+					}
+					if (Info.Harmonic >= Info.MaxHarmonic)
+					{
+						Info.bDirection = false;
+					}
+					else if (Info.Harmonic == 1)
+					{
+						Info.bDirection = true;
+					}
+				}
+				Info.Timer.Start(Info.Time);
+			}
+		}
+
+		uint32 FrameIndex = 0;
+		for (uint32 Frame = 0; Frame < CallbackInfo.NumFrames; ++Frame, FrameIndex += CallbackInfo.NumChannels)
+		{
+			int32 SpeakerIndex = 0;
+			for (int32 Channel = 0; Channel < CallbackInfo.NumChannels; ++Channel)
+			{
+				// Ignore the LFE speaker
+				if (Channel == LFEIndex)
+				{
+					continue;
+				}
+
+				FSpeakerInfo& Info = SpeakerInfo[SpeakerIndex++];
+				float Value = Amplitude * FMath::Sin(Info.Phase);
+
+				Info.Fader.Fade(Info.bOn);
+
+				CallbackInfo.OutBuffer[FrameIndex + Channel] = Info.Fader.Update()*Value;
+
+				// Increment the phase, an octave up for each channel
+				Info.Phase += (PhaseDelta * (float)Info.Harmonic);
+
+				// Wrap the phase
+				Info.Phase = WrapTwoPi(Info.Phase);
+			}
+		}
+		return true;
+	}
+
+	/************************************************************************/
+	/* FPhaseModulator														*/
+	/************************************************************************/
+
+	FPhaseModulator::FPhaseModulator(double LifeTime)
+		: IGenerator(LifeTime)
+		, Amplitude(0.4f)
+		, TotalNumSynthesizers(4)
+		, CurrNumSynthesisers(0)
+	{
+		SynthesisData.Init(FSynth(this), TotalNumSynthesizers);
+		CurrNumSynthesisers = SynthesisData.Num();
+	}
+
+	FPhaseModulator::~FPhaseModulator()
+	{
+	}
+
+	bool FPhaseModulator::GetNextBuffer(FCallbackInfo& CallbackInfo)
+	{
+		if (!UpdateTimers())
+		{
+			return true;
+		}
+
+		uint32 FrameIndex = 0;
+		for (uint32 Frame = 0; Frame < CallbackInfo.NumFrames; ++Frame, FrameIndex += CallbackInfo.NumChannels)
+		{
+			for (int32 DataIndex = 0; DataIndex < CurrNumSynthesisers; ++DataIndex)
+			{
+				SynthesisData[DataIndex].GetNextFrame(&CallbackInfo.OutBuffer[FrameIndex]);
+			}
+		}
+		return true;
+	}
+
+	FPhaseModulator::OscData::OscData(float InMaxAmp)
+		: Phase(FMath::FRandRange(0.0f, TWO_PI))
+		, Delta(0.01f)
+		, TargetDelta(Delta)
+		, DeltaEase(0.001f)
+		, SweepPhase(FMath::FRandRange(0.0f, TWO_PI))
+		, SweepDelta(0.0f)
+		, Amp(0.0f)
+		, MaxAmp(InMaxAmp)
+		, bOscilateSweep(true)
+	{
+		Reset();
+	}
+
+	void FPhaseModulator::OscData::Reset()
+	{
+		static float Scale[] = { 1.0f, 1.25f, 1.5f };
+		float Root = 0.001f;
+		TargetDelta = Root * Scale[FMath::RandRange(0, 2)];
+		DeltaEase = 0.0001f * FMath::FRand();
+
+		if (FMath::RandBool())
+		{
+			Delta = FMath::FRandRange(0.01f, 0.2f);
+		}
+		else
+		{
+			Delta = FMath::FRandRange(0.01f, 0.2f) / 1000.0f;
+		}
+
+		Amp = MaxAmp * FMath::FRandRange(0.1f, 1.0f);
+		bOscilateSweep = FMath::RandBool();
+		SweepDelta = FMath::FRandRange(-0.000001f, 0.000001f);
+	}
+
+	void FPhaseModulator::OscData::Update(bool bIsCarrier /* = false */)
+	{
+		Phase += Delta;
+		Phase = WrapTwoPi(Phase);
+
+		if (bIsCarrier)
+		{
+			Delta += (TargetDelta - Delta) * DeltaEase;
+		}
+		else
+		{
+			if (bOscilateSweep)
+			{
+				float Sweep = SweepDelta * FMath::Sin(SweepPhase);
+				SweepPhase += SweepDelta;
+				SweepPhase = WrapTwoPi(SweepPhase);
+				Delta += Sweep;
+			}
+			else
+			{
+				Delta += SweepDelta;
+			}
+		}
+	}
+
+	float FPhaseModulator::OscData::GetAmp() const
+	{
+		return Amp;
+	}
+
+	float FPhaseModulator::OscData::GetValue() const
+	{
+		return Amp * FMath::Sin(Phase);
+	}
+
+	float FPhaseModulator::OscData::GetPhase() const
+	{
+		return Phase;
+	}
+
+	FPhaseModulator::FSynth::FSynth(FPhaseModulator* InParent)
+		: Parent(InParent)
+		, Carrier(0.5f)
+		, ModIndex(1.0f)
+		, PrevSpeakerIndex(-1)
+		, bInit(false)
+	{
+	}
+
+	FPhaseModulator::FSynth::~FSynth()
+	{
+	}
+
+	void FPhaseModulator::FSynth::Init()
+	{
+		if (!bInit)
+		{
+			bInit = true;
+			Panner.Init(FMath::FRandRange(0.0f, 1.0f));
+			PanRamp.SetFrequency(FMath::FRandRange(-3.0f, 3.0f));
+			Mods.Init(OscData(FMath::FRandRange(0.25f, 1.0f)), FMath::RandRange(1, 4));
+			Carrier.Reset();
+			ModIndex.Reset();
+			FilterLFO.SetFrequency(FMath::FRandRange(0.1f, 20.0f));
+			FilterLFO.SetScaleAdd(1000.0f, 2500.0f);
+			LowPass.SetParams(1500.0f, 1.0f);
+			for (int32 ModIndex = 0; ModIndex < Mods.Num(); ++ModIndex)
+			{
+				Mods[ModIndex].Reset();
+			}
+		}
+	}
+
+	void FPhaseModulator::FSynth::GetNextFrame(float* Frame)
+	{
+		Init();
+
+		// Compute the mod value
+		float ModValue = 0.0f;
+		for (int32 Mod = 0; Mod < Mods.Num(); ++Mod)
+		{
+			Mods[Mod].Update();
+			ModValue += Mods[Mod].GetValue();
+		}
+
+		ModIndex.Update();
+		float ModIndexValue = ModIndex.GetValue();
+
+		// Plug the mod value into the carrier to generate the value of this synth
+		float Value = Carrier.GetAmp() * ModIndexValue * FMath::Sin(Carrier.GetPhase() + TWO_PI * ModValue * ModIndexValue);
+
+		// Scale by the total amplitude
+		Value *= Parent->Amplitude;
+
+		// apply the filter
+		LowPass.SetFrequency(FilterLFO.Update());
+		Value = LowPass.Update(Value);
+
+		// Update the carrier frequency data
+		Carrier.Update(true);
+
+		Panner.SetPan(PanRamp.Update());
+		Panner.Spatialize(Value, Frame);
+	}
+
+	/************************************************************************/
+	/* FPassThrough															*/
+	/************************************************************************/
+
+	FPassThrough::FPassThrough(double LifeTime)
+		: IGenerator(LifeTime)
+	{
+	}
+
+	FPassThrough::~FPassThrough()
+	{
+	}
+
+	bool FPassThrough::GetNextBuffer(FCallbackInfo& CallbackInfo)
+	{
+		if (!UpdateTimers())
+		{
+			return true;
+		}
+
+		// Passthrough
+		if (CallbackInfo.NumInputChannels > 0)
+		{
+			uint32 BufferIndex = 0;
+			uint32 InputBufferIndex = 0;
+			for (uint32 Frame = 0; Frame < CallbackInfo.NumFrames; ++Frame, BufferIndex += CallbackInfo.NumChannels, InputBufferIndex += CallbackInfo.NumInputChannels)
+			{
+				float InSample = CallbackInfo.InBuffer[InputBufferIndex];
+				InSample /= (float) CallbackInfo.NumChannels;
+				for (int32 Chan = 0; Chan < CallbackInfo.NumChannels; ++Chan)
+				{
+					CallbackInfo.OutBuffer[BufferIndex + Chan] = InSample;
+				}
+			}
+		}
+		return true;
+	}
+
+	/************************************************************************/
+	/* FInputDelay															*/
+	/************************************************************************/
+
+	FInputDelay::FInputDelay(double LifeTime)
+		: IGenerator(LifeTime)
+		, bInit(false)
+	{
+	}
+
+	FInputDelay::~FInputDelay()
+	{
+	}
+
+	void FInputDelay::Init()
+	{
+		if (bInit)
+		{
+			return;
+		}
+		bInit = true;
+		float MaxLength = 2.0f;
+		int NumTaps = CallbackData.NumChannels;
+		Delay.Init(NumTaps, MaxLength);
+		DelayOutput.Init(0.0f, NumTaps);
+		for (int32 Tap = 0; Tap < NumTaps; ++Tap)
+		{
+			Delay.SetDelayLength(Tap, FMath::FRandRange(0.2f, MaxLength));
+			Delay.SetFeedback(Tap, 0.9f);
+			Delay.SetWet(Tap, FMath::FRandRange(0.7f, 0.99f));
+		}
+	}
+
+	bool FInputDelay::GetNextBuffer(FCallbackInfo& CallbackInfo)
+	{
+		if (!UpdateTimers())
+		{
+			return true;
+		}
+
+		if (!bInit)
+		{
+			Init();
+		}
+
+		if (CallbackInfo.NumInputChannels > 0)
+		{
+			uint32 BufferIndex = 0;
+			uint32 InputBufferIndex = 0;
+			for (uint32 Frame = 0; Frame < CallbackInfo.NumFrames; ++Frame, BufferIndex += CallbackInfo.NumChannels, InputBufferIndex += CallbackInfo.NumInputChannels)
+			{
+				Delay.GetOutput(CallbackInfo.InBuffer[InputBufferIndex], DelayOutput);
+				for (int32 Chan = 0; Chan < CallbackInfo.NumChannels; ++Chan)
+				{
+					CallbackInfo.OutBuffer[BufferIndex + Chan] = DelayOutput[Chan];
+				}
+			}
+		}
+		return true;
+	}
+
+	/************************************************************************/
+	/* FNoisePan															*/
+	/************************************************************************/
+
+	FNoisePan::FNoisePan(double LifeTime)
+		: IGenerator(LifeTime)
+		, Amp(0.5f)
+		, bInit(false)
+	{
+	}
+
+	FNoisePan::~FNoisePan()
+	{
+	}
+
+	bool FNoisePan::GetNextBuffer(FCallbackInfo& CallbackInfo)
+	{
+		if (!UpdateTimers())
+		{
+			return true;
+		}
+
+		if (!bInit)
+		{
+			bInit = true;
+
+			// Initialize our filter and panner
+			PanRamp.SetFrequency(1.0f);
+			LowPass.SetParams(1000.0f, 2.0f);
+			Panner.Init(0.0f);
+			FilterLFO.SetFrequency(0.1f);
+			FilterLFO.SetScaleAdd(1000.0f, 2000.0f);
+		}
+
+		uint32 BufferIndex = 0;
+		for (uint32 Frame = 0; Frame < CallbackInfo.NumFrames; ++Frame, BufferIndex += CallbackInfo.NumChannels)
+		{
+
+			float FilterLFOValue = FilterLFO.Update();
+			LowPass.SetFrequency(FilterLFOValue);
+
+			// Get new random value (noise!)
+			float Value = Amp * FMath::FRand();
+
+			// Filter it to make it a bit less intense
+			Value = LowPass.Update(Value);
+
+			// Now spatialize it
+			Panner.SetPan(PanRamp.Update());
+			Panner.Spatialize(Value, &CallbackInfo.OutBuffer[BufferIndex]);
+		}
+		return true;
+	}
+
+} // namespace Test
+} // namespace UAudio
+
+#endif // #if ENABLE_UNREAL_AUDIO
diff --git a/Engine/Source/Runtime/UnrealAudio/Private/Tests/UnrealAudioTestGenerators.h b/Engine/Source/Runtime/UnrealAudio/Private/Tests/UnrealAudioTestGenerators.h
new file mode 100644
--- /dev/null
+++ b/Engine/Source/Runtime/UnrealAudio/Private/Tests/UnrealAudioTestGenerators.h
@@ -0,0 +1,440 @@
+// Copyright 1998-2015 Epic Games, Inc. All Rights Reserved.
+
+#pragma once 
+
+#include "UnrealAudioModule.h"
+
+#if ENABLE_UNREAL_AUDIO
+
+// We need to define TWO_PI for our FM synthesizer
+#ifndef TWO_PI
+#define TWO_PI (6.28318530718)
+#endif
+
+#ifndef PI_OVER_TWO
+#define PI_OVER_TWO (1.57079632679)
+#endif
+
+namespace UAudio
+{
+namespace Test
+{
+	/**
+	* FTimer
+	* A simple timer that tracks a timer event based on current absolute input time.
+	*/
+	class FTimer
+	{
+	public:
+		FTimer(double TotalTime);
+		~FTimer();
+
+		bool Update();
+		bool IsDone() const;
+		void Start(double InTotalTime);
+		void Reset();
+		float GetTotalTime() const;
+
+	private:
+		double TotalTime;
+		double StartTime;
+		double CurrentTime;
+		double LastTime;
+	};
+
+	/**
+	* FFader
+	* A simple sample-accurate fader
+	*/
+	class FFader
+	{
+	private:
+		struct EFadeState
+		{
+			enum Type
+			{
+				FADE_IN,
+				FADE_OUT,
+				OFF,
+				ON
+			};
+		};
+
+	public:
+		FFader(float FadeInEase, float FadeOutEase);
+		~FFader();
+
+		float Update();
+		void Fade(bool bFadeIn);
+		void Toggle();
+
+	private:
+		float FadeInEase;
+		float FadeOutEase;
+		float FadeValue;
+		EFadeState::Type FadeState;
+	};
+
+	/**
+	* FBiquad
+	* A simple biquad filter implementation.
+	*/
+	class FBiquad
+	{
+	public:
+		FBiquad();
+		~FBiquad();
+
+		float Update(float Value);
+
+	protected:
+		float X1; // X(n) * Z^-1
+		float X2; // X(n) * Z^-2
+		float Y1; // Y(n) * Z^-1
+		float Y2; // Y(n) * Z^-2
+
+		float A0; // input coefficients
+		float A1;
+		float A2;
+		float B1; // output coefficients
+		float B2;
+	};
+
+	/**
+	* FLowPass
+	* Simple 2-pole lowpass filter using a biquad filter as a base class
+	*/
+	class FLowPass : public FBiquad
+	{
+	public:
+		FLowPass();
+		~FLowPass();
+
+		void SetFrequency(float InFrequency);
+		void SetQuality(float InQuality);
+		void SetParams(float InFrequency, float InQuality);
+
+	private:
+		float Quality;
+		float Frequency;
+	};
+
+	/**
+	* FRamp
+	* A linear time-ramp class.
+	*/
+	class FRamp
+	{
+	public:
+		FRamp();
+		~FRamp();
+
+		void SetFrequency(float InFrequency);
+		void SetScaleAdd(float InScale, float InAdd);
+		float Update();
+
+	private:
+		float CurrValue;
+		float DeltaValue;
+		float Frequency;
+		float Scale;
+		float Add;
+	};
+
+	/**
+	* FSineOsc
+	* Simple sinusoidal oscillator class.
+	*/
+	class FSineOsc
+	{
+	public:
+		FSineOsc();
+		~FSineOsc();
+
+		void SetFrequency(float InFrequency);
+		void SetScaleAdd(float InScale, float InAdd);
+		void SetOutputRange(float Min, float Max);
+
+		float Update();
+
+	private:
+		float Frequency;
+		float Phase;
+		float PhaseDelta;
+		float PhaseDeltaDelta;
+		float TargetPhaseDelta;
+		float Scale;
+		float Add;
+		bool bNewValue;
+	};
+
+	/**
+	* FPan
+	* Object which takes a normalized pan position and performs spatialization based on output speaker mappings.
+	*/
+	class FPan
+	{
+	public:
+		FPan();
+		~FPan();
+
+		void Init(float InPan);
+		void SetPan(float InPan);
+		void Spatialize(float Value, float* Frame);
+
+	private:
+		int32 GetOutputSpeaker(int32 MapIndex);
+
+		float Pan;
+		int32 LFEIndex;
+		int32 PrevSpeakerIndex;
+		int32 NumNonLfeSpeakers;
+		TArray<float> SpeakerMap;
+
+		static int32 StereoSpeakerIndexMap[];
+		static int32 QuadSpeakerIndexMap[];
+		static int32 FiveOneSpeakerIndexMap[];
+		static int32 SevenOneSpeakerIndexMap[];
+	};
+
+	/**
+	* FDelay
+	* A single line, multi-tap delay object. Each tap can be set with different params.
+	*/
+	class FDelay
+	{
+	public:
+		FDelay();
+		~FDelay();
+
+		void Init(int32 NumTaps, float InMaxLengthSec);
+		void SetDelayLength(int32 Tap, float LengthSec);
+		void SetWet(int32 Tap, float WetLevel);
+		void SetFeedback(int32 Tap, float FeedbackLevel);
+		void GetOutput(float InSample, TArray<float>& TapOutput);
+
+	private:
+		struct FTap
+		{
+			float DelayFrames;
+			int32 ReadIndex;
+			float Wet;
+			float Feedback;
+
+			FTap()
+				: DelayFrames(0.0f)
+				, ReadIndex(0)
+				, Wet(0.5f)
+				, Feedback(0.5f)
+			{
+			}
+		};
+
+		float MaxLengthSec;
+		int32 LengthFrames;
+		int32 WriteIndex;
+		TArray<FTap> Taps;
+		TArray<float> DelayBuffer;
+	};
+
+
+	/**
+	* IGenerator
+	* An interface for classes which generates audio with give input/output audio buffers.
+	*/
+	class IGenerator
+	{
+	public:
+		IGenerator(double LifeTime);
+		virtual ~IGenerator() {}
+
+		virtual bool GetNextBuffer(FCallbackInfo& CallbackInfo) = 0;
+		bool IsDone() const;
+
+	protected:
+		bool UpdateTimers();
+
+		FTimer LifeTimer;
+		bool bIsDone;
+		bool bIsInit;
+	};
+
+	/**
+	* FSimpleOutput
+	* Class which plays sinusoidal decays in increasing/decreasing harmonics in each speaker at different rates.
+	*/
+	class FSimpleOutput : public IGenerator
+	{
+	public:
+		FSimpleOutput(double LifeTime);
+		~FSimpleOutput();
+
+		bool GetNextBuffer(FCallbackInfo& CallbackInfo) override;
+
+	private:
+		bool IsInit();
+
+		struct FSpeakerInfo
+		{
+			float Phase;
+			FFader Fader;
+			int32 Harmonic;
+			int32 MaxHarmonic;
+			FTimer Timer;
+			float Time;
+			bool bDirection;
+			bool bOn;
+
+			FSpeakerInfo()
+				: Phase(0.0f)
+				, Fader(0.01f, 0.0001f)
+				, Harmonic(1)
+				, MaxHarmonic(10)
+				, Timer(1.0f)
+				, Time(1.0f)
+				, bDirection(true)
+				, bOn(false)
+			{
+			}
+		};
+
+		float BaseFreqHz;
+		float PhaseDelta;
+		float Amplitude;
+		TArray<FSpeakerInfo> SpeakerInfo;
+		int32 NumNonLfeSpeakers;
+		int32 LFEIndex;
+	};
+
+
+	/**
+	* FRandomizedFmGenerator
+	* A bank of PM/FM synthesizers that have randomized and time-modulated parameters.
+	*/
+	class FPhaseModulator : public IGenerator
+	{
+	public:
+		FPhaseModulator(double LifeTime);
+		~FPhaseModulator();
+
+		bool GetNextBuffer(FCallbackInfo& CallbackInfo) override;
+
+	private:
+
+		class OscData
+		{
+		public:
+			OscData(float InMaxAmp);
+			void Reset();
+			void Update(bool bIsCarrier = false);
+			float GetAmp() const;
+			float GetValue() const;
+			float GetPhase() const;
+
+		private:
+			float Phase;
+			float Delta;
+			float TargetDelta;
+			float DeltaEase;
+			float SweepPhase;
+			float SweepDelta;
+			float Amp;
+			float MaxAmp;
+			bool bOscilateSweep;
+		};
+
+		class FSynth
+		{
+		public:
+			FSynth(FPhaseModulator* InParent);
+			~FSynth();
+
+			void GetNextFrame(float* Frame);
+
+		private:
+			void Init();
+
+			FPhaseModulator* Parent;
+			OscData Carrier;
+			OscData ModIndex;
+			int32 PrevSpeakerIndex;
+			TArray<OscData> Mods;
+			TArray<float> SpeakerMap;
+			FLowPass LowPass;
+			FPan Panner;
+			FRamp PanRamp;
+			FSineOsc FilterLFO;
+			bool bInit;
+		};
+
+		friend class FSynth;
+
+		float Amplitude;
+		TArray<FSynth> SynthesisData;
+		int32 TotalNumSynthesizers;
+		int32 CurrNumSynthesisers;
+	};
+
+	/**
+	* FPassThrough
+	* A simple input pass-through to output (mic->speakers). Should hear a bit of delay from input to output due to device latency.
+	*/
+	class FPassThrough : public IGenerator
+	{
+	public:
+		FPassThrough(double LifeTime);
+		~FPassThrough();
+
+		bool GetNextBuffer(FCallbackInfo& CallbackInfo) override;
+	};
+
+	/**
+	* FSimpleRandomizedDelayer
+	* An simple randomized delay line that delays input audio stream by various amount to great effect. 
+	*/
+	class FInputDelay : public IGenerator
+	{
+	public:
+		FInputDelay(double LifeTime);
+		~FInputDelay();
+
+		bool GetNextBuffer(FCallbackInfo& CallbackInfo) override;
+
+	private:
+		void Init();
+
+		bool bInit;
+		FDelay Delay;
+		TArray<float> DelayOutput;
+	};
+
+	/**
+	* FNoisePan
+	* An simple test which generates noise and pans it clockwise
+	*/
+	class FNoisePan : public IGenerator
+	{
+	public:
+		FNoisePan(double LifeTime);
+		~FNoisePan();
+
+		bool GetNextBuffer(FCallbackInfo& CallbackInfo) override;
+
+		float Amp;
+		FLowPass LowPass;
+		FPan Panner;
+		FRamp PanRamp;
+		FSineOsc FilterLFO;
+		FDelay Delay;
+		bool bInit;
+	};
+
+	/**
+	* Updates any callback data to use with generators based on the callback info struct
+	*/
+	void UpdateCallbackData(FCallbackInfo& CallbackInfo);
+
+} // namespace Test
+} // namespace UAudio
+
+#endif // #if ENABLE_UNREAL_AUDIO
\ No newline at end of file
diff --git a/Engine/Source/Runtime/UnrealAudio/Private/UnrealAudio.cpp b/Engine/Source/Runtime/UnrealAudio/Private/UnrealAudio.cpp
new file mode 100644
--- /dev/null
+++ b/Engine/Source/Runtime/UnrealAudio/Private/UnrealAudio.cpp
@@ -0,0 +1,82 @@
+// Copyright 1998-2015 Epic Games, Inc. All Rights Reserved.
+
+#pragma once 
+
+#include "UnrealAudioPrivate.h"
+#include "UnrealAudioModule.h"
+#include "UnrealAudioDeviceModule.h"
+#include "ModuleManager.h"
+
+#if ENABLE_UNREAL_AUDIO
+
+#define AUDIO_DEVICE_DEVICE_MODULE_NAME "UnrealAudioWasapi"
+//#define AUDIO_DEVICE_DEVICE_MODULE_NAME TEXT("UnrealAudioXAudio2")
+
+DEFINE_LOG_CATEGORY(LogUnrealAudio);
+
+IMPLEMENT_MODULE(UAudio::FUnrealAudioModule, UnrealAudio);
+
+namespace UAudio
+{
+
+
+
+	FUnrealAudioModule::FUnrealAudioModule()
+		: UnrealAudioDevice(nullptr)
+	{
+	}
+
+	FUnrealAudioModule::~FUnrealAudioModule()
+	{
+	}
+
+	void FUnrealAudioModule::Initialize()
+	{
+		UnrealAudioDevice = FModuleManager::LoadModulePtr<IUnrealAudioDeviceModule>(GetDeviceModuleName());
+		if (UnrealAudioDevice)
+		{
+			if (!UnrealAudioDevice->Initialize())
+			{
+				UE_LOG(LogUnrealAudio, Error, TEXT("Failed to initialize audio device module."));
+				UnrealAudioDevice = nullptr;
+			}
+		}
+		else
+		{
+			UnrealAudioDevice = CreateDummyDeviceModule();
+		}
+
+		InitializeTests(this);
+	}
+
+	void FUnrealAudioModule::Shutdown()
+	{
+		if (UnrealAudioDevice)
+		{
+			UnrealAudioDevice->Shutdown();
+		}
+
+		EDeviceApi::Type DeviceApi;
+		if (UnrealAudioDevice->GetDevicePlatformApi(DeviceApi) && DeviceApi == EDeviceApi::DUMMY)
+		{
+			delete UnrealAudioDevice;
+			UnrealAudioDevice = nullptr;
+		}
+	}
+
+
+	FName FUnrealAudioModule::GetDeviceModuleName() const
+	{
+		// TODO: Pull the device module name to use from an engine ini or an audio ini file
+		return FName(AUDIO_DEVICE_DEVICE_MODULE_NAME);
+	}
+
+
+	IUnrealAudioDeviceModule* FUnrealAudioModule::GetDeviceModule()
+	{
+		return UnrealAudioDevice;
+	}
+
+}
+
+#endif // #if ENABLE_UNREAL_AUDIO
diff --git a/Engine/Source/Runtime/UnrealAudio/Private/UnrealAudioBuffer.cpp b/Engine/Source/Runtime/UnrealAudio/Private/UnrealAudioBuffer.cpp
new file mode 100644
--- /dev/null
+++ b/Engine/Source/Runtime/UnrealAudio/Private/UnrealAudioBuffer.cpp
@@ -0,0 +1,128 @@
+// Copyright 1998-2015 Epic Games, Inc. All Rights Reserved.
+
+#pragma once 
+
+#include "UnrealAudioPrivate.h"
+#include "UnrealAudioBuffer.h"
+
+#if ENABLE_UNREAL_AUDIO
+
+namespace UAudio
+{
+	template<typename T>
+	class TIntermediateBuffer : public IIntermediateBuffer
+	{
+	public:
+		TIntermediateBuffer()
+			: DataBuffer(nullptr)
+		{
+		}
+
+		~TIntermediateBuffer()
+		{
+			if (DataBuffer)
+			{
+				delete[] DataBuffer;
+			}
+		}
+
+		void Initialize(uint32 InNumSamples, uint32 SampleBytes) override
+		{
+			DataBuffer = new T[InNumSamples];
+			memset((void*)DataBuffer, 0, InNumSamples * sizeof(T));
+			NumSamples = InNumSamples;
+			WriteIndex = 0;
+			ReadIndex = 0;
+		}
+
+		bool Write(uint8* InByteBuffer, uint32 InBufferSampleSize) override
+		{
+			if (!InByteBuffer || InBufferSampleSize == 0 || InBufferSampleSize > NumSamples)
+			{
+				return false;
+			}
+
+			uint32 RelativeReadIndex = ReadIndex;
+			uint32 NextWriteIndex = WriteIndex + InBufferSampleSize;
+
+			if (RelativeReadIndex < WriteIndex && NextWriteIndex >= NumSamples)
+			{
+				RelativeReadIndex += NumSamples;
+			}
+
+			if (WriteIndex <= RelativeReadIndex && NextWriteIndex > RelativeReadIndex)
+			{
+				return false;
+			}
+
+			// How many samples from beginning and from write index
+			uint32 FromBeginning = (uint32)FMath::Max<int32>((int32)WriteIndex + (int32)InBufferSampleSize - (int32)NumSamples, 0);
+			uint32 FromWrite = InBufferSampleSize - FromBeginning;
+
+			memcpy(&DataBuffer[WriteIndex], InByteBuffer, FromWrite * sizeof(T));
+			memcpy(DataBuffer, &((T*)InByteBuffer)[FromWrite], FromBeginning * sizeof(T));
+
+			// Update the write index
+			WriteIndex = (WriteIndex + InBufferSampleSize) % NumSamples;
+			return true;
+
+		}
+
+		bool Read(uint8* OutByteBuffer, uint32 OutBufferSampleSize) override
+		{
+			if (!OutByteBuffer || OutBufferSampleSize == 0 || OutBufferSampleSize > NumSamples)
+			{
+				return false;
+			}
+
+			unsigned int RelativeWriteIndex = WriteIndex;
+			unsigned int NextReadIndex = ReadIndex + OutBufferSampleSize;
+			if (RelativeWriteIndex < ReadIndex && NextReadIndex >= NumSamples)
+			{
+				RelativeWriteIndex += NumSamples;
+			}
+
+			if (ReadIndex < RelativeWriteIndex && NextReadIndex > RelativeWriteIndex) {
+				return false;
+			}
+
+			// How many samples to read from beginning and from read index
+			uint32 FromBeginning = (uint32)FMath::Max<int32>((int32)ReadIndex + (int32)OutBufferSampleSize - (int32)NumSamples, 0);
+			uint32 FromReadIndex = OutBufferSampleSize - FromBeginning;
+
+			memcpy(OutByteBuffer, &DataBuffer[ReadIndex], FromReadIndex * sizeof(T));
+			memcpy(&((T*)OutByteBuffer)[FromReadIndex], DataBuffer, FromBeginning * sizeof(T));
+
+			ReadIndex = (ReadIndex + OutBufferSampleSize) % NumSamples;
+			return true;
+		}
+
+	private:
+		T* DataBuffer;
+	};
+
+	IIntermediateBuffer* IIntermediateBuffer::CreateIntermediateBuffer(EStreamFormat::Type Format)
+	{
+		switch (Format)
+		{
+		default:
+		case EStreamFormat::UNKNOWN:
+			return nullptr;
+		case EStreamFormat::FLT:
+			return new TIntermediateBuffer<float>();
+		case EStreamFormat::DBL:
+			return new TIntermediateBuffer<double>();
+		case EStreamFormat::INT_16:
+			return new TIntermediateBuffer<int16>();
+		case EStreamFormat::INT_24:
+			return new TIntermediateBuffer<FInt24>();
+		case EStreamFormat::INT_32:
+			return new TIntermediateBuffer<int32>();
+		}
+	}
+}
+
+#endif // #if ENABLE_UNREAL_AUDIO
+
+
+
diff --git a/Engine/Source/Runtime/UnrealAudio/Private/UnrealAudioBuffer.h b/Engine/Source/Runtime/UnrealAudio/Private/UnrealAudioBuffer.h
new file mode 100644
--- /dev/null
+++ b/Engine/Source/Runtime/UnrealAudio/Private/UnrealAudioBuffer.h
@@ -0,0 +1,40 @@
+// Copyright 1998-2015 Epic Games, Inc. All Rights Reserved.
+
+#pragma once 
+
+#include "UnrealAudioDeviceModule.h"
+#include "UnrealAudioModule.h"
+
+#if ENABLE_UNREAL_AUDIO
+
+namespace UAudio
+{
+	// Forward declare
+	class UNREALAUDIO_API IIntermediateBuffer
+	{
+	public:
+		IIntermediateBuffer()
+			: NumSamples(0)
+			, WriteIndex(0)
+			, ReadIndex(0)
+		{}
+
+		virtual ~IIntermediateBuffer() {}
+
+		virtual void Initialize(uint32 InNumSamples, uint32 SampleBytes) = 0;
+		virtual bool Write(uint8* Buffer, uint32 Size) = 0;
+		virtual bool Read(uint8* Buffer, uint32 Size) = 0;
+
+		static IIntermediateBuffer* CreateIntermediateBuffer(EStreamFormat::Type Format);
+
+	protected:
+		uint32 NumSamples;
+		uint32 WriteIndex;
+		uint32 ReadIndex;
+	};
+
+}
+
+#endif // #if ENABLE_UNREAL_AUDIO
+
+
diff --git a/Engine/Source/Runtime/UnrealAudio/Private/UnrealAudioDevice.cpp b/Engine/Source/Runtime/UnrealAudio/Private/UnrealAudioDevice.cpp
new file mode 100644
--- /dev/null
+++ b/Engine/Source/Runtime/UnrealAudio/Private/UnrealAudioDevice.cpp
@@ -0,0 +1,210 @@
+// Copyright 1998-2015 Epic Games, Inc. All Rights Reserved.
+
+#pragma once 
+
+#include "UnrealAudioPrivate.h"
+#include "UnrealAudioDeviceModule.h"
+
+#if ENABLE_UNREAL_AUDIO
+
+DEFINE_LOG_CATEGORY(LogUnrealAudioDevice);
+
+namespace UAudio
+{
+	/**
+	Static array specifying which frame rates to check for support on audio devices.
+	http://en.wikipedia.org/wiki/Sampling_%28signal_processing%29
+	*/
+	const uint32 IUnrealAudioDeviceModule::PossibleFrameRates[] =
+	{
+		8000,		/** Used for telephony, walkie talkies, really shitty but ok for human speech. */
+		11025,		/** Quarter sample rate of CD's, low-quality PCM */
+		16000,		/** Rate used for VOIP (which is why VOIP sounds slightly better than normal phones), wide-band extension over normal telephony */
+		22050,		/** Half CD sample rate, low quality PCM */
+		32000,		/** MiniDV, digital FM radio, decent wireless microphones */
+		44100,		/** CD's, MPEG-1 (MP3), covers 20kHz bandwidth of human hearing with room for LP ripple. */
+		48000,		/** Standard rate used by "professional" film and audio guys: mixing console, digital recorders, etc */
+		88200,		/** Used for recording equipment intended for CDs */
+		96000,		/** DVD audio, high-def audio, 2x the 48khz "professional" sample rate. */
+		176400,		/** Rate used by HDCD recorders. */
+		192000,		/** HD-DVD and blue ray audio, 4x the 48khz "professional" sample rate. */
+	};
+	const uint32 IUnrealAudioDeviceModule::MaxPossibleFrameRates = 11;
+
+	IUnrealAudioDeviceModule::IUnrealAudioDeviceModule()
+		: DeviceErrorListener(nullptr)
+	{
+	}
+
+	IUnrealAudioDeviceModule::~IUnrealAudioDeviceModule()
+	{
+	}
+
+	bool IUnrealAudioDeviceModule::CreateStream(const FCreateStreamParams& Params)
+	{
+		checkf(StreamInfo.State == EStreamState::CLOSED, TEXT("Stream state can't be open if creating a new one."));
+		checkf(Params.OutputDeviceIndex != INDEX_NONE, TEXT("Input or output stream params need to be set."));
+
+		bool bSuccess = false;
+
+		Reset();
+
+		if (OpenDevices(Params))
+		{
+			if (Params.InputDeviceIndex != INDEX_NONE)
+			{
+				// This stream also streams audio in (every audio stream streams output)
+				StreamInfo.StreamType = EStreamType::INPUT;
+			}
+			else
+			{
+				// This stream is only streaming audio out
+				StreamInfo.StreamType = EStreamType::OUTPUT;
+			}
+
+			StreamInfo.State = EStreamState::STOPPED;
+			StreamInfo.CallbackFunction = Params.CallbackFunction;
+			StreamInfo.UserData = Params.UserData;
+			StreamInfo.BlockSize = Params.CallbackBlockSize;
+			StreamInfo.StreamDelta = (double)StreamInfo.BlockSize / StreamInfo.FrameRate;
+
+			bSuccess = true;
+		}
+
+		return bSuccess;
+	}
+
+	void IUnrealAudioDeviceModule::UpdateStreamTimeTick()
+	{
+		StreamInfo.StreamTime += StreamInfo.StreamDelta;
+	}
+
+	void IUnrealAudioDeviceModule::OnDeviceError(const EDeviceError::Type ErrorType, const FString& ErrorDetails, const FString& FileName, int32 LineNumber) const
+	{
+		UE_LOG(LogUnrealAudioDevice, Error, TEXT("Audio Device Error: (%s) : %s (%s::%d)"), EDeviceError::ToString(ErrorType), *ErrorDetails, *FileName, LineNumber);
+
+		// And broadcast to listener if one is setup
+		if (DeviceErrorListener)
+		{
+			DeviceErrorListener->OnDeviceError(ErrorType, ErrorDetails, FileName, LineNumber);
+		}
+	}
+
+	void IUnrealAudioDeviceModule::Reset()
+	{
+		StreamInfo.Initialize();
+	}
+
+
+	/** 
+	* FUnrealAudioDeviceDummy
+	* A dummy audio device implementation
+	*/
+	class FUnrealAudioDeviceDummy : public IUnrealAudioDeviceModule
+	{
+	public:
+		FUnrealAudioDeviceDummy()
+		{
+		}
+
+		~FUnrealAudioDeviceDummy()
+		{
+		}
+
+		// IUnrealAudioDeviceModule
+		bool Initialize() override
+		{
+			return true;
+		}
+
+		bool Shutdown() override
+		{
+			return true;
+		}
+
+		bool GetDevicePlatformApi(EDeviceApi::Type & OutType) const override
+		{
+			OutType = EDeviceApi::DUMMY;
+			return true;
+		}
+
+		bool GetNumOutputDevices(uint32& OutNumDevices) const override
+		{
+			OutNumDevices = 0;
+			return true;
+		}
+
+		bool GetOutputDeviceInfo(const uint32 DeviceIndex, FDeviceInfo& OutInfo) const override
+		{
+			memset((void*)&OutInfo, 0, sizeof(FDeviceInfo));
+			return true;
+		}
+
+		bool GetNumInputDevices(uint32& OutNumDevices) const override
+		{
+			OutNumDevices = 0;
+			return true;
+		}
+
+		bool GetInputDeviceInfo(const uint32 DeviceIndex, FDeviceInfo& OutInfo) const override
+		{
+			memset((void*)&OutInfo, 0, sizeof(FDeviceInfo));
+			return true;
+		}
+
+		bool GetDefaultOutputDeviceIndex(uint32& OutDefaultIndex) const override
+		{
+			OutDefaultIndex = 0;
+			return true;
+		}
+
+		bool GetDefaultInputDeviceIndex(uint32& OutDefaultIndex) const override
+		{
+			OutDefaultIndex = 0;
+			return true;
+		}
+
+		bool StartStream() override
+		{
+			return true;
+		}
+
+		bool StopStream() override
+		{
+			return true;
+		}
+
+		bool ShutdownStream() override
+		{
+			return true;
+		}
+
+		bool GetLatency(uint32& OutputDeviceLatency, uint32& InputDeviceLatency) const override
+		{
+			OutputDeviceLatency = 0;
+			InputDeviceLatency = 0;
+			return true;
+		}
+
+		bool GetFrameRate(uint32& OutFrameRate) const override
+		{
+			OutFrameRate = 0;
+			return true;
+		}
+
+		bool OpenDevices(const FCreateStreamParams& Params) override
+		{
+			return true;
+		}
+	};
+
+	// Exported Functions
+	IUnrealAudioDeviceModule* CreateDummyDeviceModule()
+	{
+		return new FUnrealAudioDeviceDummy();
+	}
+
+}
+
+#endif // #if ENABLE_UNREAL_AUDIO
+
diff --git a/Engine/Source/Runtime/UnrealAudio/Private/UnrealAudioDeviceFormat.cpp b/Engine/Source/Runtime/UnrealAudio/Private/UnrealAudioDeviceFormat.cpp
new file mode 100644
--- /dev/null
+++ b/Engine/Source/Runtime/UnrealAudio/Private/UnrealAudioDeviceFormat.cpp
@@ -0,0 +1,229 @@
+// Copyright 1998-2015 Epic Games, Inc. All Rights Reserved.
+
+#pragma once 
+
+#include "UnrealAudioPrivate.h"
+#include "UnrealAudioDeviceModule.h"
+#include "UnrealAudioDeviceFormat.h"
+
+#if ENABLE_UNREAL_AUDIO
+
+namespace UAudio
+{
+	/*
+	 * This code is segmented here since it involves basically converting to and from various audio format types. Even though the
+	 * audio device code assumes floats for the platform-independent mixing code, various audio devices may be in other types of formats
+	 * so we need to support converting byte buffers to and from various formats.
+	 */
+
+	uint32 IUnrealAudioDeviceModule::GetNumBytesForFormat(EStreamFormat::Type Format) const
+	{
+		switch (Format)
+		{
+			case EStreamFormat::FLT:		return 4;
+			case EStreamFormat::DBL:		return 8;
+			case EStreamFormat::INT_16:		return 2;
+			case EStreamFormat::INT_24:		return 3;
+			case EStreamFormat::INT_32:		return 4;
+			default:						return 0;
+		}
+	}
+
+	void IUnrealAudioDeviceModule::SetupBufferFormatConvertInfo(EStreamType::Type StreamType)
+	{
+		FBufferFormatConvertInfo& ConvertInfo = StreamInfo.GetConvertInfo(StreamType);
+		FStreamDeviceInfo& DeviceInfo = StreamInfo.GetDeviceInfo(StreamType);
+
+		if (StreamType == EStreamType::INPUT)
+		{
+			// Setup the convert info to convert from device channels and format to user channels and format
+			ConvertInfo.FromChannels = DeviceInfo.NumChannels;
+			ConvertInfo.ToChannels = DeviceInfo.NumChannels;
+			ConvertInfo.FromFormat = DeviceInfo.DeviceDataFormat;
+			ConvertInfo.ToFormat = EStreamFormat::FLT;
+		}
+		else
+		{
+			// Setup the convert info to convert from user channels and format to device channels and format
+			ConvertInfo.FromChannels = DeviceInfo.NumChannels;
+			ConvertInfo.ToChannels = DeviceInfo.NumChannels;
+			ConvertInfo.FromFormat = EStreamFormat::FLT;
+			ConvertInfo.ToFormat = DeviceInfo.DeviceDataFormat;
+		}
+
+		// For either stream types, set the number of channels to use for the conversion to be the smallest number of channels
+		ConvertInfo.NumChannels = FMath::Min(ConvertInfo.FromChannels, ConvertInfo.ToChannels);
+
+	}
+
+	template <typename FromType, typename ToType>
+	void IUnrealAudioDeviceModule::ConvertToFloatType(TArray<uint8>& ToBuffer, TArray<uint8>& FromBuffer, uint32 NumFrames, FBufferFormatConvertInfo& ConvertInfo)
+	{
+		// Converting to float types is easy: just get the max size of the value (
+		ToType* ToBufferPtr = TDataFormat<ToType>::GetDataBuffer(ToBuffer);
+		FromType* FromBufferPtr = TDataFormat<FromType>::GetDataBuffer(FromBuffer);
+		double MaxSize = TDataFormat<FromType>::GetMaxValue();
+
+		double Scale = (1.0 / MaxSize);
+
+		for (uint32 FrameIndex = 0; FrameIndex < NumFrames; ++FrameIndex)
+		{
+			for (uint32 ChannelIndex = 0; ChannelIndex < ConvertInfo.NumChannels; ++ChannelIndex)
+			{
+				ToBufferPtr[ChannelIndex] = (ToType)((FromBufferPtr[ChannelIndex] * Scale));
+			}
+			FromBufferPtr += ConvertInfo.FromChannels;
+			ToBufferPtr += ConvertInfo.ToChannels;
+		}
+	}
+
+	template <typename FloatType>
+	void IUnrealAudioDeviceModule::ConvertAllToFloatType(TArray<uint8>& ToBuffer, TArray<uint8>& FromBuffer, uint32 NumFrames, FBufferFormatConvertInfo& ConvertInfo)
+	{
+		// Note: even if we're converting float-to-float we still call this because we may be performing a channel-count conversion
+
+		switch (ConvertInfo.FromFormat)
+		{
+			case EStreamFormat::FLT:
+				ConvertToFloatType<float, FloatType>(ToBuffer, FromBuffer, StreamInfo.BlockSize, ConvertInfo);
+				break;
+
+			case EStreamFormat::DBL:
+				ConvertToFloatType<double, FloatType>(ToBuffer, FromBuffer, StreamInfo.BlockSize, ConvertInfo);
+				break;
+
+			case EStreamFormat::INT_16:
+				ConvertToFloatType<int16, FloatType>(ToBuffer, FromBuffer, StreamInfo.BlockSize, ConvertInfo);
+				break;
+
+			case EStreamFormat::INT_24:
+				ConvertToFloatType<int24, FloatType>(ToBuffer, FromBuffer, StreamInfo.BlockSize, ConvertInfo);
+				break;
+
+			case EStreamFormat::INT_32:
+				ConvertToFloatType<int32, FloatType>(ToBuffer, FromBuffer, StreamInfo.BlockSize, ConvertInfo);
+				break;
+		}
+	}
+
+	template <typename FromType, typename ToType>
+	void IUnrealAudioDeviceModule::ConvertIntegerToIntegerType(TArray<uint8>& ToBuffer, TArray<uint8>& FromBuffer, uint32 NumFrames, FBufferFormatConvertInfo& ConvertInfo)
+	{
+		ToType* ToBufferPtr = TDataFormat<ToType>::GetDataBuffer(ToBuffer);
+		FromType* FromBufferPtr = TDataFormat<FromType>::GetDataBuffer(FromBuffer);
+
+		// Figure out the amount of bit-shifting we need to do based on the bytes counts of various formats.
+		int32 ShiftAmount = 8 * ((int32)sizeof(ToType) - (int32)sizeof(FromType));
+
+		// If our shift amount is negative (converting to a integer with fewer bytes, we'll be truncating the data
+		// so we'll flip the shift amount and then right-shift the FromBuffer data, otherwise we'll shift the data to the left
+		if (ShiftAmount < 0)
+		{
+			ShiftAmount *= -1;
+			for (uint32 FrameIndex = 0; FrameIndex < NumFrames; ++FrameIndex)
+			{
+				for (uint32 ChannelIndex = 0; ChannelIndex < ConvertInfo.NumChannels; ++ChannelIndex)
+				{
+					int32 Value = FromBufferPtr[ChannelIndex];
+					ToBufferPtr[ChannelIndex] = (ToType)(Value >> ShiftAmount);
+				}
+				FromBufferPtr += ConvertInfo.FromChannels;
+				ToBufferPtr += ConvertInfo.ToChannels;
+			}
+		}
+		else
+		{
+			for (uint32 FrameIndex = 0; FrameIndex < NumFrames; ++FrameIndex)
+			{
+				for (uint32 ChannelIndex = 0; ChannelIndex < ConvertInfo.NumChannels; ++ChannelIndex)
+				{
+					int32 Value = FromBufferPtr[ChannelIndex];
+					ToBufferPtr[ChannelIndex] = (ToType)(Value << ShiftAmount);
+				}
+				FromBufferPtr += ConvertInfo.FromChannels;
+				ToBufferPtr += ConvertInfo.ToChannels;
+			}
+		}
+	}
+
+	template <typename FromType, typename ToType>
+	void IUnrealAudioDeviceModule::ConvertFloatToIntegerType(TArray<uint8>& ToBuffer, TArray<uint8>& FromBuffer, uint32 NumFrames, FBufferFormatConvertInfo& ConvertInfo)
+	{
+		ToType* ToBufferPtr = TDataFormat<ToType>::GetDataBuffer(ToBuffer);
+		FromType* FromBufferPtr = TDataFormat<FromType>::GetDataBuffer(FromBuffer);
+		for (uint32 FrameIndex = 0; FrameIndex < NumFrames; ++FrameIndex)
+		{
+			double MaxValue = TDataFormat<ToType>::GetMaxValue();
+			for (uint32 ChannelIndex = 0; ChannelIndex < ConvertInfo.NumChannels; ++ChannelIndex)
+			{
+				ToBufferPtr[ChannelIndex] = (ToType)(FromBufferPtr[ChannelIndex] * MaxValue);
+			}
+			FromBufferPtr += ConvertInfo.FromChannels;
+			ToBufferPtr += ConvertInfo.ToChannels;
+		}
+	}
+
+	template <typename FloatType>
+	void IUnrealAudioDeviceModule::ConvertAllToIntegerType(TArray<uint8>& ToBuffer, TArray<uint8>& FromBuffer, uint32 NumFrames, FBufferFormatConvertInfo& ConvertInfo)
+	{
+		// Note: even if we're converting int32-to-int32 we still call this because we may be performing a channel-count conversion
+
+		switch (ConvertInfo.FromFormat)
+		{
+			case EStreamFormat::FLT:
+				ConvertFloatToIntegerType<float, FloatType>(ToBuffer, FromBuffer, StreamInfo.BlockSize, ConvertInfo);
+				break;
+
+			case EStreamFormat::DBL:
+				ConvertFloatToIntegerType<double, FloatType>(ToBuffer, FromBuffer, StreamInfo.BlockSize, ConvertInfo);
+				break;
+
+			case EStreamFormat::INT_16:
+				ConvertIntegerToIntegerType<int16, FloatType>(ToBuffer, FromBuffer, StreamInfo.BlockSize, ConvertInfo);
+				break;
+
+			case EStreamFormat::INT_24:
+				ConvertIntegerToIntegerType<int24, FloatType>(ToBuffer, FromBuffer, StreamInfo.BlockSize, ConvertInfo);
+				break;
+
+			case EStreamFormat::INT_32:
+				ConvertIntegerToIntegerType<int32, FloatType>(ToBuffer, FromBuffer, StreamInfo.BlockSize, ConvertInfo);
+				break;
+		}
+	}
+
+	bool IUnrealAudioDeviceModule::ConvertBufferFormat(TArray<uint8>& ToBuffer, TArray<uint8>& FromBuffer, EStreamType::Type StreamType)
+	{
+		FBufferFormatConvertInfo& ConvertInfo = StreamInfo.GetConvertInfo(StreamType);
+
+		switch (ConvertInfo.ToFormat)
+		{
+			case EStreamFormat::FLT:
+				ConvertAllToFloatType<float>(ToBuffer, FromBuffer, StreamInfo.BlockSize, ConvertInfo);
+				break;
+
+			case EStreamFormat::DBL:
+				ConvertAllToFloatType<double>(ToBuffer, FromBuffer, StreamInfo.BlockSize, ConvertInfo);
+				break;
+
+			case EStreamFormat::INT_16:
+				ConvertAllToIntegerType<double>(ToBuffer, FromBuffer, StreamInfo.BlockSize, ConvertInfo);
+				break;
+
+			case EStreamFormat::INT_24:
+				ConvertAllToIntegerType<double>(ToBuffer, FromBuffer, StreamInfo.BlockSize, ConvertInfo);
+				break;
+
+			case EStreamFormat::INT_32:
+				ConvertAllToIntegerType<double>(ToBuffer, FromBuffer, StreamInfo.BlockSize, ConvertInfo);
+				break;
+		}
+
+		return true;
+	}
+
+}
+
+#endif // #if ENABLE_UNREAL_AUDIO
+
+
diff --git a/Engine/Source/Runtime/UnrealAudio/Private/UnrealAudioDeviceFormat.h b/Engine/Source/Runtime/UnrealAudio/Private/UnrealAudioDeviceFormat.h
new file mode 100644
--- /dev/null
+++ b/Engine/Source/Runtime/UnrealAudio/Private/UnrealAudioDeviceFormat.h
@@ -0,0 +1,183 @@
+// Copyright 1998-2015 Epic Games, Inc. All Rights Reserved.
+
+#pragma once 
+
+#include "UnrealAudioModule.h"
+#include "UnrealAudioPrivate.h"
+#include "UnrealAudioDeviceModule.h"
+
+#if ENABLE_UNREAL_AUDIO
+
+namespace UAudio
+{
+	/**
+	* FInt24
+	* Helper class to deal with 24 bit integers (converts to and from 32 bit integers)
+	*/
+	class FInt24
+	{
+	public:
+
+		FInt24()
+		{
+			memset((void *)Data, 0, sizeof(Data));
+		}
+
+		FInt24(const FInt24& Other)
+		{
+			*this = Other;
+		}
+
+		FInt24(const int32& Value)
+		{
+			*this = Value;
+		}
+
+		FInt24& operator=(const int32& Value)
+		{
+			static const uint32 MASK_BYTE_0 = 0x000000FF;
+			static const uint32 MASK_BYTE_1 = 0x0000FF00;
+			static const uint32 MASK_BYTE_2 = 0x00FF0000;
+
+			Data[0] = (Value & MASK_BYTE_0);
+			Data[1] = (Value & MASK_BYTE_1) >> 8;
+			Data[2] = (Value & MASK_BYTE_2) >> 16;
+			return *this;
+		}
+
+		int32 operator+(const double& Value)
+		{
+			int32 IntValue = GetInt32();
+			return IntValue + (int32)Value;
+		}
+
+		double operator*(const double& Value)
+		{
+			int32 IntValue = GetInt32();
+			return IntValue * Value;
+		}
+
+		operator int()
+		{
+			return GetInt32();
+		}
+
+		int32 GetInt32()
+		{
+			// copy data to front of integer, then shift to right 8 bits for sign extension
+			int32 Value = ((Data[2] << 24) | (Data[1] << 16) | (Data[1] << 8)) >> 8;
+			return Value;
+		}
+
+	private:
+
+		/** 3 bytes to represent the 24 bit integer */
+		uint8 Data[3];
+	};
+
+	// typedef for the special 24 bit integer type
+	typedef FInt24 int24;
+
+	/**
+	* EStreamDataFormat
+	*
+	* An enumeration to specify the state of the audio stream.
+	*
+	*/
+	namespace EStreamFormat
+	{
+		enum Type
+		{
+			UNKNOWN,	// Unknown format type (indicates error or uninitialized values)
+			FLT,
+			DBL,
+			INT_16,
+			INT_24,
+			INT_32,
+			UNSUPPORTED,
+		};
+
+		/** @return the stringified version of the enum passed in */
+		inline const TCHAR* ToString(EStreamFormat::Type StreamFormat)
+		{
+			switch (StreamFormat)
+			{
+				default:
+				case UNKNOWN:	return TEXT("UNKNOWN");
+				case FLT:		return TEXT("FLOAT");
+				case DBL:		return TEXT("DOUBLE");
+				case INT_16:	return TEXT("INT_16");
+				case INT_24:	return TEXT("INT_24");
+				case INT_32:	return TEXT("INT_32");
+			}
+		}
+	}
+
+	/** Template base for various data formats */
+	template<typename T>
+	struct TDataFormat
+	{
+		static FORCEINLINE EStreamFormat::Type GetTypeEnum()			{ return EStreamFormat::UNKNOWN;	}
+		static FORCEINLINE double GetMaxValue()							{ return 0.0;						}
+		static FORCEINLINE T* GetDataBuffer(TArray<uint8>& Buffer)		{ return (T*)Buffer.GetData();		}
+		static FORCEINLINE bool IsInteger()								{ return false;						}
+		static FORCEINLINE uint32 GetByteMask()							{ return 0;							}
+	};
+
+	/* Float template specialization */
+	template<>
+	struct TDataFormat<float>
+	{
+		static FORCEINLINE EStreamFormat::Type GetTypeEnum()			{ return EStreamFormat::FLT;		}
+		static FORCEINLINE double GetMaxValue()							{ return 1.0;						}
+		static FORCEINLINE float* GetDataBuffer(TArray<uint8>& Buffer)	{ return (float*)Buffer.GetData();	}
+		static FORCEINLINE bool IsInteger()								{ return false;						}
+	};
+
+	/* Double template specialization */
+	template<>
+	struct TDataFormat<double>
+	{
+		static FORCEINLINE EStreamFormat::Type GetTypeEnum()			{ return EStreamFormat::DBL;		}
+		static FORCEINLINE double GetMaxValue()							{ return 1.0;						}
+		static FORCEINLINE double* GetDataBuffer(TArray<uint8>& Buffer) { return (double*)Buffer.GetData(); }
+		static FORCEINLINE bool IsInteger()								{ return false;						}
+	};
+
+	/* Int8 template specialization */
+	template<>
+	struct TDataFormat<int16>
+	{
+		static FORCEINLINE EStreamFormat::Type GetTypeEnum()			{ return EStreamFormat::INT_16;		}
+		static FORCEINLINE double GetMaxValue()							{ return 32767;						}
+		static FORCEINLINE int16* GetDataBuffer(TArray<uint8>& Buffer)	{ return (int16*)Buffer.GetData();	}
+		static FORCEINLINE bool IsInteger()								{ return true;						}
+		static FORCEINLINE uint32 GetByteMask()							{ return 0x0000FFFF;				}
+	};
+
+	/* int24 template specialization */
+	template<>
+	struct TDataFormat<int24>
+	{
+		static FORCEINLINE EStreamFormat::Type GetTypeEnum()			{ return EStreamFormat::INT_24;		}
+		static FORCEINLINE double GetMaxValue()							{ return 8388607;					}
+		static FORCEINLINE int24* GetDataBuffer(TArray<uint8>& Buffer)	{ return (int24*)Buffer.GetData();	}
+		static FORCEINLINE bool IsInteger()								{ return true;						}
+		static FORCEINLINE uint32 GetByteMask()							{ return 0x00FFFFFF;				}
+	};
+
+	/* int24 template specialization */
+	template<>
+	struct TDataFormat<int32>
+	{
+		static FORCEINLINE EStreamFormat::Type GetTypeEnum()			{ return EStreamFormat::INT_32;		}
+		static FORCEINLINE double GetMaxValue()							{ return 2147483647;				}
+		static FORCEINLINE int32* GetDataBuffer(TArray<uint8>& Buffer)	{ return (int32*)Buffer.GetData();	}
+		static FORCEINLINE bool IsInteger()								{ return true;						}
+		static FORCEINLINE uint32 GetByteMask()							{ return 0xFFFFFFFF;				}
+	};
+
+}
+
+#endif // #if ENABLE_UNREAL_AUDIO
+
diff --git a/Engine/Source/Runtime/UnrealAudio/Private/UnrealAudioPrivate.h b/Engine/Source/Runtime/UnrealAudio/Private/UnrealAudioPrivate.h
new file mode 100644
--- /dev/null
+++ b/Engine/Source/Runtime/UnrealAudio/Private/UnrealAudioPrivate.h
@@ -0,0 +1,39 @@
+// Copyright 1998-2015 Epic Games, Inc. All Rights Reserved.
+
+#pragma once 
+
+#include "UnrealAudioModule.h"
+#include "UnrealAudioDeviceModule.h"
+#include "UnrealAudioTests.h"
+
+#if ENABLE_UNREAL_AUDIO
+
+DECLARE_LOG_CATEGORY_EXTERN(LogUnrealAudio, Log, All);
+
+namespace UAudio
+{
+	class FUnrealAudioModule :	public IUnrealAudioModule
+	{
+	public:
+
+		FUnrealAudioModule();
+		~FUnrealAudioModule();
+
+		void Initialize() override;
+		void Shutdown() override;
+
+		class IUnrealAudioDeviceModule* GetDeviceModule();
+
+		FName GetDeviceModuleName() const;
+
+	private:
+
+		static void InitializeTests(FUnrealAudioModule* Module);
+
+		class IUnrealAudioDeviceModule* UnrealAudioDevice;
+	};
+}
+
+#endif // #if ENABLE_UNREAL_AUDIO
+
+
diff --git a/Engine/Source/Runtime/UnrealAudio/Public/UnrealAudioDeviceModule.h b/Engine/Source/Runtime/UnrealAudio/Public/UnrealAudioDeviceModule.h
new file mode 100644
--- /dev/null
+++ b/Engine/Source/Runtime/UnrealAudio/Public/UnrealAudioDeviceModule.h
@@ -0,0 +1,690 @@
+// Copyright 1998-2015 Epic Games, Inc. All Rights Reserved.
+
+#pragma once 
+
+#include "ModuleInterface.h"
+#include "UnrealAudioModule.h"
+#include "UnrealAudioDeviceFormat.h"
+
+#if ENABLE_UNREAL_AUDIO
+
+DECLARE_LOG_CATEGORY_EXTERN(LogUnrealAudioDevice, Log, All);
+
+#define UA_DEVICE_PLATFORM_ERROR(INFO)	(OnDeviceError(EDeviceError::PLATFORM, INFO, FString(__FILE__), __LINE__))
+#define AU_DEVICE_PARAM_ERROR(INFO)		(OnDeviceError(EDeviceError::INVALID_PARAMETER, INFO, FString(__FILE__), __LINE__))
+#define AU_DEVICE_WARNING(INFO)			(OnDeviceError(EDeviceError::WARNING, INFO, FString(__FILE__), __LINE__))
+
+
+namespace UAudio
+{
+	/**
+	* EDeviceError
+	* The types of errors that can occur in IUnrealAudioDevice implementations.
+	*/
+	namespace EDeviceError
+	{
+		enum Type
+		{
+			WARNING,
+			UNKNOWN,
+			INVALID_PARAMETER,
+			PLATFORM,
+			SYSTEM,
+			THREAD,
+		};
+
+		/** @return the stringified version of the enum passed in */
+		static inline const TCHAR* ToString(EDeviceError::Type DeviceError)
+		{
+			switch (DeviceError)
+			{
+				case WARNING:			return TEXT("WARNING");
+				case UNKNOWN:			return TEXT("UNKNOWN");
+				case INVALID_PARAMETER:	return TEXT("INVALID_PARAMETER");
+				case PLATFORM:			return TEXT("PLATFORM");
+				case SYSTEM:			return TEXT("SYSTEM");
+				case THREAD:			return TEXT("THREAD");
+				default:				check(false);
+			}
+			return TEXT("");
+		}
+	}
+
+	/**
+	* EDeviceApi
+	* Implemented platform-Specific audio device APIs.
+	*/
+	namespace EDeviceApi
+	{
+		enum Type
+		{
+			WASAPI,				// Windows
+			XAUDIO2,			// Windows/XBox
+			NGS2,				// PS4
+			ALSA,				// Linux
+			CORE_AUDIO,			// Mac/IOS
+			OPENAL,				// Linux
+			HTML5,				// Web
+			DUMMY,				// DUMMY (no device api)
+		};
+
+		/** @return the stringified version of the enum passed in */
+		static inline const TCHAR* ToString(EDeviceApi::Type DeviceApi)
+		{
+			switch (DeviceApi)
+			{
+				case WASAPI:		return TEXT("WASAPI");
+				case XAUDIO2:		return TEXT("XAUDIO2");
+				case NGS2:			return TEXT("NGS2");
+				case ALSA:			return TEXT("ALSA");
+				case CORE_AUDIO:	return TEXT("CORE_AUDIO");
+				case OPENAL:		return TEXT("OPENAL");
+				case HTML5:			return TEXT("HTML5");
+				case DUMMY:			return TEXT("DUMMY");
+				default:			check(false);
+			}
+			return TEXT("");
+		}
+	};
+
+	/**
+	* EStreamStatus
+	* An enumeration to specify the state of the audio stream.
+	*/
+	namespace EStreamState
+	{
+		enum Type
+		{
+			CLOSED,				/** The stream is closed */
+			STOPPED,			/** The stream was running and is now stopped */
+			RUNNING,			/** The stream is running (currently streaming callbacks) */
+			STOPPING,			/** The stream is currently stopping*/
+		};
+	}
+
+	/**
+	* EDeviceDataFlow
+	* An enumeration used to specify the flow of audio on the device
+	*/
+	namespace EStreamType
+	{
+		enum Type
+		{
+			INPUT,			/** The audio device streams audio into the application (e.g. microphone) */
+			OUTPUT,			/** The audio device only streams audio out of the application */
+			COUNT,			/** The count of stream types. */
+			UNKNOWN,		/** Unknown stream type. */
+		};
+	}
+
+	/**
+	* EStreamStatus
+	* An enumeration to specify the flow status of the audio stream.
+	*/
+	namespace EStreamFlowStatus
+	{
+		enum Flag
+		{
+			INPUT_OVERFLOW		= (1 << 0),		/** The input stream is has too much data */
+			OUTPUT_UNDERFLOW	= (1 << 1)		/** The output stream has too little data */
+		};
+	}
+	typedef uint8 StreamStatus;
+
+	/**
+	* ESpeaker
+	* An enumeration to specify speaker types
+	*/
+	struct ESpeaker
+	{
+		enum Type
+		{
+			FRONT_LEFT,
+			FRONT_RIGHT,
+			FRONT_CENTER,
+			LOW_FREQUENCY,
+			BACK_LEFT,
+			BACK_RIGHT,
+			FRONT_LEFT_OF_CENTER,
+			FRONT_RIGHT_OF_CENTER,
+			BACK_CENTER,
+			SIDE_LEFT,
+			SIDE_RIGHT,
+			TOP_CENTER,
+			TOP_FRONT_LEFT,
+			TOP_FRONT_CENTER,
+			TOP_FRONT_RIGHT,
+			TOP_BACK_LEFT,
+			TOP_BACK_CENTER,
+			TOP_BACK_RIGHT,
+			SPEAKER_TYPE_COUNT
+		};
+
+		/** @return the stringified version of the enum passed in */
+		static inline const TCHAR* ToString(ESpeaker::Type Speaker)
+		{
+			switch (Speaker)
+			{
+				case FRONT_LEFT:				return TEXT("FRONT_LEFT");
+				case FRONT_RIGHT:				return TEXT("FRONT_RIGHT");
+				case FRONT_CENTER:				return TEXT("FRONT_CENTER");
+				case LOW_FREQUENCY:				return TEXT("LOW_FREQUENCY");
+				case BACK_LEFT:					return TEXT("BACK_LEFT");
+				case BACK_RIGHT:				return TEXT("BACK_RIGHT");
+				case FRONT_LEFT_OF_CENTER:		return TEXT("FRONT_LEFT_OF_CENTER");
+				case FRONT_RIGHT_OF_CENTER:		return TEXT("FRONT_RIGHT_OF_CENTER");
+				case BACK_CENTER:				return TEXT("BACK_CENTER");
+				case SIDE_LEFT:					return TEXT("SIDE_LEFT");
+				case SIDE_RIGHT:				return TEXT("SIDE_RIGHT");
+				case TOP_CENTER:				return TEXT("TOP_CENTER");
+				case TOP_FRONT_LEFT:			return TEXT("TOP_FRONT_LEFT");
+				case TOP_FRONT_CENTER:			return TEXT("TOP_FRONT_CENTER");
+				case TOP_FRONT_RIGHT:			return TEXT("TOP_FRONT_RIGHT");
+				case TOP_BACK_LEFT:				return TEXT("TOP_BACK_LEFT");
+				case TOP_BACK_CENTER:			return TEXT("TOP_BACK_CENTER");
+				case TOP_BACK_RIGHT:			return TEXT("TOP_BACK_RIGHT");
+				default:						return TEXT("UKNOWN");
+			}
+			return TEXT("");
+		}
+	};
+
+
+	/**
+	* FCallbackInfo
+	* A struct for callback info. Using this rather than params in callback function
+	* because updating new members here will be less painful.
+	*/
+	struct FCallbackInfo
+	{
+		/** A pointer to a buffer of float data which will be written to output device.*/
+		float* OutBuffer;
+
+		/** The number of buffer output frames */
+		uint32 NumFrames;
+
+		/** A pointer to a buffer of float data which is passed in from device input. */
+		float* InBuffer;
+
+		/** Array of outputs speakers */
+		TArray<ESpeaker::Type> OutputSpeakers;
+
+		/** The number of channels in output. */
+		int32 NumChannels;
+
+		/** The number of input channels in input stream. */
+		int32 NumInputChannels;
+
+		/** The current status flags of the input and output buffers */
+		StreamStatus StatusFlags;
+
+		/** The current frame-accurate lifetime of the audio stream. */
+		double StreamTime;
+
+		/** The output device framerate */
+		int32 FrameRate;
+
+		/** A pointer to user data */
+		void* UserData;
+	};
+
+	/**
+	* StreamCallback
+	*
+	* Callback which calls from platform audio device code into platform-independent mixing code.
+	* @param CallbackInfo A struct of callback info needed to process audio.
+	*/
+	typedef bool(*StreamCallback)(FCallbackInfo& CallbackInfo);
+
+	/**
+	* FDeviceInfo
+	* Struct used to hold information about audio devices, queried by user.
+	*/
+	struct FDeviceInfo
+	{
+		/** Friendly name of device. */
+		FString FriendlyName;
+
+		/** Number of channels (e.g. 2 for stereo) natively supported by the device. */
+		uint32 NumChannels;
+
+		/** The frame rate of the device. */
+		uint32 FrameRate;
+
+		/** The data format of the device (e.g. float). */
+		EStreamFormat::Type StreamFormat;
+
+		/** What type of device it is (input, output, or both). */
+		EStreamType::Type StreamType;
+
+		/** What speakers this device supports (if output device) */
+		TArray<ESpeaker::Type> Speakers;
+
+		/** Whether or not it is the OS default device for the type. */
+		bool bIsSystemDefault;
+
+		FDeviceInfo()
+			: FriendlyName(TEXT("Uknown"))
+			, NumChannels(0)
+			, FrameRate(0)
+			, StreamFormat(EStreamFormat::UNKNOWN)
+			, StreamType(EStreamType::UNKNOWN)
+			, bIsSystemDefault(false)
+		{}
+	};
+
+	/**
+	* FCreateStreamParams
+	* Struct used to define stream creation.
+	*/
+	struct FCreateStreamParams
+	{
+		/** The index of the device to use for audio input. Initialized to INDEX_NONE, which means no device input. */
+		uint32 InputDeviceIndex;
+
+		/** The index of the device to use for audio output. Must be defined. */
+		uint32 OutputDeviceIndex;
+
+		/** The size of the callback block (in frames) that the user would like. (e.g. 512) */
+		uint32 CallbackBlockSize;
+
+		/** The function pointer of a user callback function to generate audio samples to output device. */
+		StreamCallback CallbackFunction;
+
+		/** Pointer to store user data. Passed in callback function. */
+		void* UserData;
+
+		/** Constructor */
+		FCreateStreamParams()
+			: InputDeviceIndex(INDEX_NONE)
+			, OutputDeviceIndex(INDEX_NONE)
+			, CallbackBlockSize(512)
+			, UserData(nullptr)
+		{}
+	};
+
+	/**
+	* IDeviceErrorListener
+	* An optional class that defines an error listener interface. Used for handling or logging any audio device errors which may have occurred.
+	*/
+	class UNREALAUDIO_API IDeviceErrorListener
+	{
+	public:
+		virtual void OnDeviceError(const EDeviceError::Type Error, const FString& ErrorDetails, const FString& FileName, int32 LineNumber) const = 0;
+	};
+
+	/**
+	* IUnrealAudioDeviceModule
+	* Main audio device module, needs to be implemented for every specific platform API.
+	*/
+	class UNREALAUDIO_API IUnrealAudioDeviceModule : public IModuleInterface,
+													 public IDeviceErrorListener
+	{
+	public:
+		/** Constructor */
+		IUnrealAudioDeviceModule();
+
+		/** Virtual destructor */
+		virtual ~IUnrealAudioDeviceModule();
+
+		/** 
+		* Initializes the audio device module. 
+		* @return true if succeeded.
+		*/
+		virtual bool Initialize() = 0;
+
+		/** 
+		* Shuts down the audio device module. 
+		* @return true if succeeded.
+		*/
+		virtual bool Shutdown() = 0;
+
+		/** 
+		* Returns the API enumeration for the internal implementation of the device module. 
+		* @param OutType Enumeration of the output tyep.
+		* @return true if succeeded.
+		*/
+		virtual bool GetDevicePlatformApi(EDeviceApi::Type & OutType) const = 0;
+
+		/** 
+		* Returns the number of connected output devices to the system. 
+		* @param OutNumDevices The number of output devices
+		* @return true if succeeded.
+		*/
+		virtual bool GetNumOutputDevices(uint32& OutNumDevices) const = 0;
+
+		/** 
+		* Returns information about the output device at that index. 
+		* @param DeviceIndex The index of the output device. 
+		* @param OutInfo Information struct about the device.
+		* @return true if succeeded.
+		* @note DeviceIndex should be less than total number of devices given by GetNumOutputDevices.
+		*/
+		virtual bool GetOutputDeviceInfo(const uint32 DeviceIndex, FDeviceInfo& OutInfo) const = 0;
+
+		/** 
+		* Returns the number of input devices. 
+		* @param OutNumDevices The number of input devices.
+		* @return true if succeeded.
+		*/
+		virtual bool GetNumInputDevices(uint32& OutNumDevices) const = 0;
+
+		/** 
+		* Returns information about the input device at the given index. 
+		* @param DeviceIndex The index of the input device.
+		* @param OutInfo Information struct about the device.
+		* @return true if succeeded.
+		*/
+		virtual bool GetInputDeviceInfo(const uint32 DeviceIndex, FDeviceInfo& OutInfo) const = 0;
+
+		/**
+		* Returns the default output device index.
+		* @param OutDefaultIndex The default device index.
+		* @return true if succeeded.
+		* @note The default device is specified by the OS default device output.
+		*/
+		virtual bool GetDefaultOutputDeviceIndex(uint32& OutDefaultIndex) const = 0;
+
+		/**
+		* Returns the default input device index.
+		* @param OutDefaultIndex The default device index.
+		* @return true if succeeded.
+		* @note The default device is specified by the OS default device input.
+		*/
+		virtual bool GetDefaultInputDeviceIndex(uint32& OutDefaultIndex) const = 0;
+
+		/**
+		* Starts the device audio stream.
+		* @return true if succeeded.
+		* @note The device stream needs to have been created before starting.
+		*/
+		virtual bool StartStream() = 0;
+
+		/**
+		* Stops the device audio stream.
+		* @return true if succeeded.
+		* @note The device stream needs to have been created and started before stopping.
+		*/
+		virtual bool StopStream() = 0;
+
+		/**
+		* Frees resources of the device audio stream.
+		* @return true if succeeded.
+		* @note The device stream needs to have been created before shutting down.
+		*/
+		virtual bool ShutdownStream() = 0;
+
+		/**
+		* Returns the latency of the input and output devices.
+		* @param OutputDeviceLatency The latency of the output device as reported by platform API
+		* @param InputDeviceLatency The latency of the input device as reported by platform API
+		* @return true if succeeded.
+		* @note The device stream needs to have been started for this return anything. A return of 0 means
+		* the device hasn't started or that device hasn't been initialized.
+		*/
+		virtual bool GetLatency(uint32& OutputDeviceLatency, uint32& InputDeviceLatency) const = 0;
+
+		/**
+		* Returns the frame rate of the audio devices.
+		* @param OutFrameRate The frame rate of the output audio device.
+		* @return true if succeeded.
+		* @note FrameRate is also known as SampleRate. A "Frame" is the minimal time delta of 
+		* audio and is composed of interleaved samples. e.g. 1 stereo frame is 2 samples: left and right.
+		*/
+		virtual bool GetFrameRate(uint32& OutFrameRate) const = 0;
+
+		/**
+		* Creates an audio stream given the input parameter struct.
+		* @param Params The parameters needed to create an audio stream.
+		* @return true if succeeded.
+		*/
+		bool CreateStream(const FCreateStreamParams& Params);
+
+		/**
+		* Function called when an error occurs in the device code.
+		* @param Error The error type that occurred.
+		* @param ErrorDetails A string of specific details about the error.
+		* @param FileName The file that the error occurred in.
+		* @param LineNumber The line number that the error occurred in.
+		*/
+		void OnDeviceError(const EDeviceError::Type Error, const FString& ErrorDetails, const FString& FileName, int32 LineNumber) const override;
+
+	protected:
+
+		/**
+		* FBufferFormatConvertInfo
+		* Struct used to convert data formats (e.g. float, int32, etc) to and from the device to user callback stream.
+		* The user audio callback will always be in floats, but audio devices may have native formats that are different, so 
+		* we need to support converting data formats.
+		*/
+		struct FBufferFormatConvertInfo
+		{
+			/** The base number of channels for the format conversion. */
+			uint32 NumChannels;
+
+			/** The number of channels we're converting from. */
+			uint32 FromChannels;
+
+			/** The number of channels we're to. */
+			uint32 ToChannels;
+
+			/** The data format we're converting from. */
+			EStreamFormat::Type FromFormat;
+
+			/** The data format we're converting to. */
+			EStreamFormat::Type	ToFormat;
+		};
+
+		/**
+		* FStreamDeviceInfo
+		* Struct used to represent information about a particular device (input or output device)
+		*/
+		struct FStreamDeviceInfo
+		{
+			/** The index this device is in (from list of devices of this type). */
+			uint32 DeviceIndex;
+
+			/** The speaker types this device uses. */
+			TArray<ESpeaker::Type> Speakers;
+
+			/** The number of channels this device actually supports. */
+			uint32 NumChannels;
+
+			/** The reported latency of this device. */
+			uint64 Latency;
+
+			/** The native framerate of this device. */
+			uint32 FrameRate;
+
+			/** The native data format of this device. */
+			EStreamFormat::Type DeviceDataFormat;
+
+			/** Conversion information to convert audio streams to/from this device. */
+			FBufferFormatConvertInfo BufferFormatConvertInfo;
+
+			/** A buffer used to store data to/from this device. */
+			TArray<uint8> UserBuffer;
+
+			/** True if we need to perform a format conversion. */
+			bool bPerformFormatConversion;
+
+			/** True if we need to perform a byte swap for this device. */
+			bool bPerformByteSwap;
+
+			/** Constructor */
+			FStreamDeviceInfo()
+			{
+				Initialize();
+			}
+
+			/** Function which initializes the struct variables to default values. Maye be called multiple times. */
+			void Initialize()
+			{
+				DeviceIndex = 0;
+				NumChannels = 0;
+				Latency = 0;
+				FrameRate = 0;
+				DeviceDataFormat = EStreamFormat::UNKNOWN;
+				bPerformFormatConversion = false;
+				bPerformByteSwap = false;
+
+				BufferFormatConvertInfo.NumChannels = 0;
+				BufferFormatConvertInfo.FromChannels = 0;
+				BufferFormatConvertInfo.ToChannels = 0;
+				BufferFormatConvertInfo.FromFormat = EStreamFormat::UNKNOWN;
+				BufferFormatConvertInfo.ToFormat = EStreamFormat::UNKNOWN;
+			}
+		};
+
+		/**
+		* FStreamInfo
+		* Struct used to represent general information about the audio stream.
+		*/
+		class FStreamInfo
+		{
+		public:
+			/** The overall framerate of the stream. This may be different from device frame rate. */
+			uint32 FrameRate;
+
+			/** The type of stream this is (input or output) */
+			EStreamType::Type StreamType;
+
+			/** The current state of the stream. */
+			EStreamState::Type State;
+
+			/** Information on the user callback */
+			FCallbackInfo CallbackInfo;
+
+			/** The sample-accurate running time of the stream in seconds (i.e. not necessarily real-world time but stream time). */
+			double StreamTime;
+
+			/** The amount of time that passes per update block. */
+			double StreamDelta;
+
+			/** Running audio thread */
+			FRunnableThread* Thread;
+
+			/** User callback function */
+			StreamCallback CallbackFunction;
+
+			/** User data */
+			void* UserData;
+
+			/** The size of the callback frame count*/
+			uint32 BlockSize;
+
+			/** A byte array used to store data to and from audio devices. */
+			TArray<uint8> DeviceBuffer;
+
+			/** An array of device-specific information for input and output devices. */
+			FStreamDeviceInfo DeviceInfo[EStreamType::COUNT];
+
+			/** Constructor */
+			FStreamInfo()
+			{
+				// Just call intialize since it will get called when starting up a new stream.
+				Initialize();
+			}
+
+			/** Gets called every time a stream is opened, is reset in cae the device data exists from a previous open. */
+			void Initialize()
+			{
+				FrameRate = 0;
+				StreamType = EStreamType::COUNT;
+				State = EStreamState::CLOSED;
+				StreamTime = 0.0;
+				StreamDelta = 0.0;
+
+				CallbackFunction = nullptr;
+				Thread = nullptr;
+				UserData = nullptr;
+				BlockSize = 0;
+
+				for (uint32 Id = 0; Id < EStreamType::COUNT; ++Id)
+				{
+					DeviceInfo[Id].Initialize();
+				}
+			}
+
+			/** Helper function to get specific convert information based on stream type */
+			FBufferFormatConvertInfo& GetConvertInfo(EStreamType::Type StreamType)
+			{
+				check(StreamType < EStreamType::COUNT);
+				return DeviceInfo[StreamType].BufferFormatConvertInfo;
+			}
+
+			/** Helper function to get device info based on stream type */
+			FStreamDeviceInfo& GetDeviceInfo(EStreamType::Type StreamType)
+			{
+				check(StreamType < EStreamType::COUNT);
+				return DeviceInfo[StreamType];
+			}
+
+			/** Helper function to get device info based on stream type const overload */
+			const FStreamDeviceInfo& GetDeviceInfo(EStreamType::Type StreamType) const
+			{
+				check(StreamType < EStreamType::COUNT);
+				return DeviceInfo[StreamType];
+			}
+		};
+
+	protected: // Protected Methods
+
+		/** Opens audio devices given the input params. Implemented per platform */
+		virtual bool OpenDevices(const FCreateStreamParams& Params) = 0;
+
+		/** Called before opening up new streams. */
+		void Reset();
+
+		/** Helper function to get number of bytes for a given stream format. */
+		uint32 GetNumBytesForFormat(EStreamFormat::Type Format) const;
+
+		/** Sets up any convert information for given stream type (figures out to/from convert format and channel formats) */
+		void SetupBufferFormatConvertInfo(EStreamType::Type StreamType);
+
+		/** Performs actual buffer format and channel conversion */
+		bool ConvertBufferFormat(TArray<uint8>& OutputBuffer, TArray<uint8>& InputBuffer, EStreamType::Type StreamType);
+
+		/** Updates the sample-accurate stream time value. */
+		void UpdateStreamTimeTick();
+
+		/** Templated function to convert to float from any other format type. */
+		template <typename FromType, typename ToType>
+		void ConvertToFloatType(TArray<uint8>& ToBuffer, TArray<uint8>& FromBuffer, uint32 NumFrames, FBufferFormatConvertInfo& ConvertInfo);
+
+		/** Templated function to any format type to float. */
+		template <typename FloatType>
+		void ConvertAllToFloatType(TArray<uint8>& ToBuffer, TArray<uint8>& FromBuffer, uint32 NumFrames, FBufferFormatConvertInfo& ConvertInfo);
+
+		/** Templated function to integer type to any other integer type. */
+		template <typename FromType, typename ToType>
+		void ConvertIntegerToIntegerType(TArray<uint8>& ToBuffer, TArray<uint8>& FromBuffer, uint32 NumFrames, FBufferFormatConvertInfo& ConvertInfo);
+
+		/** Templated function to convert a float type to integer type. */
+		template <typename FromType, typename ToType>
+		void ConvertFloatToIntegerType(TArray<uint8>& ToBuffer, TArray<uint8>& FromBuffer, uint32 NumFrames, FBufferFormatConvertInfo& ConvertInfo);
+
+		/** Templated function to any type to integer type. */
+		template <typename FloatType>
+		void ConvertAllToIntegerType(TArray<uint8>& ToBuffer, TArray<uint8>& FromBuffer, uint32 NumFrames, FBufferFormatConvertInfo& ConvertInfo);
+
+	protected: // Protected Data
+
+		/** An array of possible frame rates that can exist with any audio device. */
+		static const uint32 PossibleFrameRates[];
+		static const uint32 MaxPossibleFrameRates;
+
+		/** An object which listens to device errors. */
+		IDeviceErrorListener* DeviceErrorListener;
+
+		/** Device stream info struct */
+		FStreamInfo StreamInfo;
+	};
+
+	// Exported Functions
+	UNREALAUDIO_API IUnrealAudioDeviceModule* CreateDummyDeviceModule();
+}
+
+#endif // #if ENABLE_UNREAL_AUDIO
+
+
diff --git a/Engine/Source/Runtime/UnrealAudio/Public/UnrealAudioModule.h b/Engine/Source/Runtime/UnrealAudio/Public/UnrealAudioModule.h
new file mode 100644
--- /dev/null
+++ b/Engine/Source/Runtime/UnrealAudio/Public/UnrealAudioModule.h
@@ -0,0 +1,32 @@
+// Copyright 1998-2015 Epic Games, Inc. All Rights Reserved.
+
+#pragma once 
+
+#include "ModuleInterface.h"
+
+// Only enable unreal audio on windows
+#if PLATFORM_WINDOWS
+#define ENABLE_UNREAL_AUDIO 1
+#else
+#define ENABLE_UNREAL_AUDIO 0
+#endif
+
+namespace UAudio
+{
+	class UNREALAUDIO_API IUnrealAudioModule :  public IModuleInterface
+	{
+	public:
+		virtual ~IUnrealAudioModule() {}
+
+		virtual void Initialize()
+		{
+		}
+
+		virtual void Shutdown()
+		{
+		}
+	};
+}
+
+
+
diff --git a/Engine/Source/Runtime/UnrealAudio/Public/UnrealAudioTests.h b/Engine/Source/Runtime/UnrealAudio/Public/UnrealAudioTests.h
new file mode 100644
--- /dev/null
+++ b/Engine/Source/Runtime/UnrealAudio/Public/UnrealAudioTests.h
@@ -0,0 +1,66 @@
+// Copyright 1998-2015 Epic Games, Inc. All Rights Reserved.
+
+#pragma once 
+
+#include "UnrealAudioModule.h"
+
+namespace UAudio
+{
+	/**
+	* TestDeviceQuery
+	*
+	* Test which queries input and output devices and displays information to log about device capabilities, format, etc.
+	* @return True if test succeeds.
+	*/
+	bool UNREALAUDIO_API TestDeviceQuery();
+
+	/**
+	* TestDeviceOutputSimple
+	*
+	* Test which runs a simple sinusoid output on all channels (each channel is one octave higher than left channel)
+	* to test that audio output is working and function with no discontinuities.
+	* @param TestTime The amount of time to run the test in seconds (negative value means to run indefinitely)
+	* @return True if test succeeds.
+	*/
+	bool UNREALAUDIO_API TestDeviceOutputSimple(double TestTime);
+	
+	/**
+	* TestDeviceInputSimple
+	*
+	* Test which test input device functionality by feeding input device data to output device data (pass through).
+	* If you have a connected microphone on your system (set to default) you should hear the mic input feed through output speakers.
+	* @param TestTime The amount of time to run the test in seconds (negative value means to run indefinitely)
+	* @return True if test succeeds.
+	*/
+	bool UNREALAUDIO_API TestDeviceInputSimple(double TestTime);
+	
+	/**
+	* TestDeviceOutputRandomizedFm
+	*
+	* Test which tests output device functionality with a bit more fun than the simple test. You should hear randomized FM synthesis in panning through
+	* all connected speaker outputs.
+	* @param TestTime The amount of time to run the test in seconds (negative value means to run indefinitely)
+	* @return True if test succeeds.
+	*/
+	bool UNREALAUDIO_API TestDeviceOutputRandomizedFm(double TestTime);
+
+	/**
+	* TestDeviceOutputNoisePan
+	*
+	* Test which tests output device functionality with white noise that pans clockwise in surround sound setup.
+	* @param TestTime The amount of time to run the test in seconds (negative value means to run indefinitely)
+	* @return True if test succeeds.
+	*/
+	bool UNREALAUDIO_API TestDeviceOutputNoisePan(double TestTime);
+
+	/**
+	* TestDeviceInputRandomizedDelay
+	*
+	* Test input device functionality with a bit more fun that the simple test. You should hear the input data randomly delayed by various amounts
+	* with randomized feedback and wetness. Should sound robotic and "interesting".
+	* @param TestTime The amount of time to run the test in seconds (negative value means to run indefinitely)
+	* @return True if test succeeds.
+	*/
+	bool UNREALAUDIO_API TestDeviceInputRandomizedDelay(double TestTime);
+}
+
diff --git a/Engine/Source/Runtime/UnrealAudio/UnrealAudio.Build.cs b/Engine/Source/Runtime/UnrealAudio/UnrealAudio.Build.cs
new file mode 100644
--- /dev/null
+++ b/Engine/Source/Runtime/UnrealAudio/UnrealAudio.Build.cs
@@ -0,0 +1,35 @@
+// Copyright 1998-2015 Epic Games, Inc. All Rights Reserved.
+
+using UnrealBuildTool;
+
+public class UnrealAudio : ModuleRules
+{
+	public UnrealAudio(TargetInfo Target)
+	{
+		PrivateIncludePaths.AddRange(
+			new string[] {
+				"Runtime/UnrealAudio/Private",
+				"Runtime/UnrealAudio/Private/Tests",
+			}
+		);
+
+		PublicIncludePaths.AddRange(
+			new string[] {
+				"Runtime/UnrealAudio/Public",
+			}
+		);
+
+		PublicDependencyModuleNames.AddRange(
+			new string[] {
+				"Core",
+				"CoreUObject",
+			}
+		);
+
+		PrivateDependencyModuleNames.AddRange(
+			new string[] {
+				"CoreUObject", 
+			}
+		);
+	}
+}
\ No newline at end of file
diff --git a/Engine/Source/Runtime/Windows/UnrealAudioWasapi/Private/UnrealAudioDeviceWasapi.cpp b/Engine/Source/Runtime/Windows/UnrealAudioWasapi/Private/UnrealAudioDeviceWasapi.cpp
new file mode 100644
--- /dev/null
+++ b/Engine/Source/Runtime/Windows/UnrealAudioWasapi/Private/UnrealAudioDeviceWasapi.cpp
@@ -0,0 +1,1137 @@
+// Copyright 1998-2015 Epic Games, Inc. All Rights Reserved.
+
+#pragma once 
+
+/**
+	Concrete implementation of IAudioDevice for WASAPI (Windows Audio Session API)
+
+	See https://msdn.microsoft.com/en-us/library/windows/desktop/dd371455%28v=vs.85%29.aspx
+
+*/
+
+#include "UnrealAudioDeviceModule.h"
+#include "UnrealAudioBuffer.h"
+
+#include "AllowWindowsPlatformTypes.h"
+#include <Audioclient.h>					// WASAPI api
+#include <avrt.h>							// This module contains the multimedia class scheduler APIs and any public data structures needed to call these APIs.
+#include <mmdeviceapi.h>					// For getting endpoint notifications about audio devices
+#include <functiondiscoverykeys_devpkey.h>	// Defines property keys for the Plug and Play Device Property API.
+#include "HideWindowsPlatformTypes.h"
+
+#if ENABLE_UNREAL_AUDIO
+
+DEFINE_LOG_CATEGORY(LogUnrealAudioDevice);
+
+// See MSDN documentation for what these error codes mean in the context of the API call
+static const TCHAR* GetWasapiError(HRESULT Result)
+{
+	switch (Result)
+	{
+		case AUDCLNT_E_NOT_INITIALIZED:					return TEXT("AUDCLNT_E_NOT_INITIALIZED");
+		case AUDCLNT_E_ALREADY_INITIALIZED:				return TEXT("AUDCLNT_E_ALREADY_INITIALIZED");
+		case AUDCLNT_E_WRONG_ENDPOINT_TYPE:				return TEXT("AUDCLNT_E_WRONG_ENDPOINT_TYPE");
+		case AUDCLNT_E_DEVICE_INVALIDATED:				return TEXT("AUDCLNT_E_DEVICE_INVALIDATED");
+		case AUDCLNT_E_NOT_STOPPED:						return TEXT("AUDCLNT_E_NOT_STOPPED");
+		case AUDCLNT_E_BUFFER_TOO_LARGE:				return TEXT("AUDCLNT_E_BUFFER_TOO_LARGE");
+		case AUDCLNT_E_OUT_OF_ORDER:					return TEXT("AUDCLNT_E_OUT_OF_ORDER");
+		case AUDCLNT_E_UNSUPPORTED_FORMAT:				return TEXT("AUDCLNT_E_UNSUPPORTED_FORMAT");
+		case AUDCLNT_E_INVALID_SIZE:					return TEXT("AUDCLNT_E_INVALID_SIZE");
+		case AUDCLNT_E_DEVICE_IN_USE:					return TEXT("AUDCLNT_E_DEVICE_IN_USE");
+		case AUDCLNT_E_BUFFER_OPERATION_PENDING:		return TEXT("AUDCLNT_E_BUFFER_OPERATION_PENDING");
+		case AUDCLNT_E_THREAD_NOT_REGISTERED:			return TEXT("AUDCLNT_E_THREAD_NOT_REGISTERED");
+		case AUDCLNT_E_EXCLUSIVE_MODE_NOT_ALLOWED:		return TEXT("AUDCLNT_E_EXCLUSIVE_MODE_NOT_ALLOWED");
+		case AUDCLNT_E_ENDPOINT_CREATE_FAILED:			return TEXT("AUDCLNT_E_ENDPOINT_CREATE_FAILED");
+		case AUDCLNT_E_SERVICE_NOT_RUNNING:				return TEXT("AUDCLNT_E_SERVICE_NOT_RUNNING");
+		case AUDCLNT_E_EVENTHANDLE_NOT_EXPECTED:		return TEXT("AUDCLNT_E_EVENTHANDLE_NOT_EXPECTED");
+		case AUDCLNT_E_EXCLUSIVE_MODE_ONLY:				return TEXT("AUDCLNT_E_EXCLUSIVE_MODE_ONLY");
+		case AUDCLNT_E_BUFDURATION_PERIOD_NOT_EQUAL:	return TEXT("AUDCLNT_E_BUFDURATION_PERIOD_NOT_EQUAL");
+		case AUDCLNT_E_EVENTHANDLE_NOT_SET:				return TEXT("AUDCLNT_E_EVENTHANDLE_NOT_SET");
+		case AUDCLNT_E_INCORRECT_BUFFER_SIZE:			return TEXT("AUDCLNT_E_INCORRECT_BUFFER_SIZE");
+		case AUDCLNT_E_BUFFER_SIZE_ERROR:				return TEXT("AUDCLNT_E_BUFFER_SIZE_ERROR");
+		case AUDCLNT_E_CPUUSAGE_EXCEEDED:				return TEXT("AUDCLNT_E_CPUUSAGE_EXCEEDED");
+		case AUDCLNT_E_BUFFER_ERROR:					return TEXT("AUDCLNT_E_BUFFER_ERROR");
+		case AUDCLNT_E_BUFFER_SIZE_NOT_ALIGNED:			return TEXT("AUDCLNT_E_BUFFER_SIZE_NOT_ALIGNED");
+		case AUDCLNT_E_INVALID_DEVICE_PERIOD:			return TEXT("AUDCLNT_E_INVALID_DEVICE_PERIOD");
+		case AUDCLNT_E_INVALID_STREAM_FLAG:				return TEXT("AUDCLNT_E_INVALID_STREAM_FLAG");
+		case AUDCLNT_E_ENDPOINT_OFFLOAD_NOT_CAPABLE:	return TEXT("AUDCLNT_E_ENDPOINT_OFFLOAD_NOT_CAPABLE");
+		case AUDCLNT_E_OUT_OF_OFFLOAD_RESOURCES:		return TEXT("AUDCLNT_E_OUT_OF_OFFLOAD_RESOURCES");
+		case AUDCLNT_E_OFFLOAD_MODE_ONLY:				return TEXT("AUDCLNT_E_OFFLOAD_MODE_ONLY");
+		case AUDCLNT_E_NONOFFLOAD_MODE_ONLY:			return TEXT("AUDCLNT_E_NONOFFLOAD_MODE_ONLY");
+		case AUDCLNT_E_RESOURCES_INVALIDATED:			return TEXT("AUDCLNT_E_RESOURCES_INVALIDATED");
+		case AUDCLNT_E_RAW_MODE_UNSUPPORTED:			return TEXT("AUDCLNT_E_RAW_MODE_UNSUPPORTED");
+		case REGDB_E_CLASSNOTREG:						return TEXT("REGDB_E_CLASSNOTREG");
+		case CLASS_E_NOAGGREGATION:						return TEXT("CLASS_E_NOAGGREGATION");
+		case E_NOINTERFACE:								return TEXT("E_NOINTERFACE");
+		case E_POINTER:									return TEXT("E_POINTER");
+		case E_INVALIDARG:								return TEXT("E_INVALIDARG");
+		case E_OUTOFMEMORY:								return TEXT("E_OUTOFMEMORY");
+		default: return TEXT("UKNOWN");
+	}
+}
+
+/** Some helper defines to reduce boilerplate code for error handling windows API calls. */
+
+#define CLEANUP_ON_FAIL(Result)						\
+	if (FAILED(Result))								\
+	{												\
+		const TCHAR* Err = GetWasapiError(Result);	\
+		UA_DEVICE_PLATFORM_ERROR(Err);				\
+		goto Cleanup;								\
+	}
+
+#define RETURN_FALSE_ON_FAIL(Result)				\
+	if (FAILED(Result))								\
+	{												\
+		const TCHAR* Err = GetWasapiError(Result);	\
+		UA_DEVICE_PLATFORM_ERROR(Err);				\
+		return false;								\
+	}
+
+
+namespace UAudio
+{
+	/** Struct used to map Unreal Audio enumerations for speaker types to device speaker types */
+	struct FSpeakerMaskMapping
+	{
+		ESpeaker::Type UnrealSpeaker;
+		long XAudio2SpeakerTypeFlag;
+	};
+
+	/** Mapping of XAudio2/Windows speaker type enumerations and unreal audio enumerations */
+	FSpeakerMaskMapping SpeakerMaskMapping[18] = 
+	{
+		{ ESpeaker::FRONT_LEFT,				SPEAKER_FRONT_LEFT				},
+		{ ESpeaker::FRONT_RIGHT,			SPEAKER_FRONT_RIGHT				},
+		{ ESpeaker::FRONT_CENTER,			SPEAKER_FRONT_CENTER			},
+		{ ESpeaker::LOW_FREQUENCY,			SPEAKER_LOW_FREQUENCY			},
+		{ ESpeaker::BACK_LEFT,				SPEAKER_BACK_LEFT				},
+		{ ESpeaker::BACK_RIGHT,				SPEAKER_BACK_RIGHT				},
+		{ ESpeaker::FRONT_LEFT_OF_CENTER,	SPEAKER_FRONT_LEFT_OF_CENTER	},
+		{ ESpeaker::FRONT_RIGHT_OF_CENTER,	SPEAKER_FRONT_RIGHT_OF_CENTER	},
+		{ ESpeaker::BACK_CENTER,			SPEAKER_BACK_CENTER				},
+		{ ESpeaker::SIDE_LEFT,				SPEAKER_SIDE_LEFT				},
+		{ ESpeaker::SIDE_RIGHT,				SPEAKER_SIDE_RIGHT				},
+		{ ESpeaker::TOP_CENTER,				SPEAKER_TOP_CENTER				},
+		{ ESpeaker::TOP_FRONT_LEFT,			SPEAKER_TOP_FRONT_LEFT			},
+		{ ESpeaker::TOP_FRONT_CENTER,		SPEAKER_TOP_FRONT_CENTER		},
+		{ ESpeaker::TOP_FRONT_RIGHT,		SPEAKER_TOP_FRONT_RIGHT			},
+		{ ESpeaker::TOP_BACK_LEFT,			SPEAKER_TOP_BACK_LEFT			},
+		{ ESpeaker::TOP_BACK_CENTER,		SPEAKER_TOP_BACK_CENTER			},
+		{ ESpeaker::TOP_BACK_RIGHT,			SPEAKER_TOP_BACK_RIGHT			},
+	};
+
+	/** Helper function which turns a windows speaker flag mask into an array of speaker types. */
+	static void GetArrayOfSpeakers(const uint32 NumChannels, TArray<ESpeaker::Type>& Speakers, const ::DWORD ChannelMask)
+	{
+		Speakers.Reset();
+		uint32 ChanCount = 0;
+		// Build a flag field of the speaker outputs of this device
+		for (uint32 SpeakerTypeIndex = 0; SpeakerTypeIndex < ESpeaker::SPEAKER_TYPE_COUNT, ChanCount < NumChannels; ++SpeakerTypeIndex)
+		{
+			if (ChannelMask & SpeakerMaskMapping[SpeakerTypeIndex].XAudio2SpeakerTypeFlag)
+			{
+				++ChanCount;
+				Speakers.Add(SpeakerMaskMapping[SpeakerTypeIndex].UnrealSpeaker);
+			}
+		}
+
+		check(ChanCount == NumChannels);
+	}
+
+	/**
+	* FUnrealAudioWasapi
+	* Wasapi implementation of IUnrealAudioDeviceModule
+	* Inherits from FRunnable so that audio device I/O can be run on separate thread.
+	*/
+	class FUnrealAudioWasapi :	public IUnrealAudioDeviceModule, 
+								public FRunnable
+	{
+	public:
+		FUnrealAudioWasapi();
+		~FUnrealAudioWasapi();
+
+		// IUnrealAudioDeviceModule
+		bool Initialize() override;
+		bool Shutdown() override;
+		bool GetDevicePlatformApi(EDeviceApi::Type & OutType) const override;
+		bool GetNumOutputDevices(uint32& OutNumDevices) const override;
+		bool GetOutputDeviceInfo(const uint32 DeviceIndex, FDeviceInfo& OutInfo) const override;
+		bool GetNumInputDevices(uint32& OutNumDevices) const override;
+		bool GetInputDeviceInfo(const uint32 DeviceIndex, FDeviceInfo& OutInfo) const override;
+		bool GetDefaultOutputDeviceIndex(uint32& OutDefaultIndex) const override;
+		bool GetDefaultInputDeviceIndex(uint32& OutDefaultIndex) const override;
+		bool StartStream() override;
+		bool StopStream() override;
+		bool ShutdownStream() override;
+		bool GetLatency(uint32& OutputDeviceLatency, uint32& InputDeviceLatency) const override;
+		bool GetFrameRate(uint32& OutFrameRate) const override;
+
+		// FRunnable
+		uint32 Run() override;
+		void Stop() override;
+		void Exit() override;
+
+	private:
+
+		/**
+		* FWasapiInfo
+		*
+		* Structure for holding WASAPI specific resources
+		*
+		*/
+		struct FWasapiInfo
+		{
+			/** The windows device enumerator. Used to query connected audio devices. */
+			IMMDeviceEnumerator* DeviceEnumerator;
+
+			/** Allows creation and initialization of audio input stream from audio engine and hardware buffer of audio endpoint device*/
+			IAudioClient* CaptureClient;
+
+			/** Allows creation and initialization of audio output stream from audio engine and hardware buffer of audio endpoint device*/
+			IAudioClient* RenderClient;
+
+			/** Allows reading data input from capture endpoint buffer */
+			IAudioCaptureClient* CaptureService;
+
+			/** Allows writing data output to capture endpoint buffer */
+			IAudioRenderClient* RenderService;
+
+			/** Handle used to notify when new input data is ready to be read in from hardware */
+			HANDLE CaptureEvent;
+
+			/** Handle used to notify when hardware is ready for new audio data to be written to it*/
+			HANDLE RenderEvent;
+
+			/** Intermediate buffer used to store audio data from input device, converted to user data format (float)*/
+			IIntermediateBuffer* CaptureIntermediateBuffer;
+
+			/** Intermediate buffer used to store audio data from user callback, converted to hardware native format*/
+			IIntermediateBuffer* RenderIntermediateBuffer;
+
+			/** Whether or not devices are open */
+			bool bDevicesOpen;
+
+			FWasapiInfo()
+				: DeviceEnumerator(nullptr)
+				, CaptureClient(nullptr)
+				, RenderClient(nullptr)
+				, CaptureService(nullptr)
+				, RenderService(nullptr)
+				, CaptureEvent(0)
+				, RenderEvent(0)
+				, CaptureIntermediateBuffer(nullptr)
+				, RenderIntermediateBuffer(nullptr)
+				, bDevicesOpen(false)
+			{
+			}
+		};
+
+		/** The number of active input devices detected by WASAPI */
+		uint32 InputDeviceCount;
+
+		/** The number of active output devices detected by WASAPI */
+		uint32 OutputDeviceCount;
+
+		/** WASAPI-specific data */
+		FWasapiInfo WasapiInfo;
+
+		/** Whether or not the device api has been initialized */
+		bool bInitialized;
+
+		/** Whether or not com was initialized */
+		bool bComInitialized;
+
+	private:
+
+		// IUnrealAudioDeviceModule
+		bool OpenDevices(const FCreateStreamParams& CreateStreamParams) override;
+
+		// Helper functions
+		bool OpenDevice(uint32 DeviceIndex, EStreamType::Type StreamType);
+		bool GetDeviceCount(EDataFlow DataFlow, uint32& NumDevices) const;
+		bool GetDeviceInfo(EDataFlow DataFlow, const uint32 DeviceINdex, FDeviceInfo& DeviceInfo) const;
+		bool GetDefaultDeviceIndex(EDataFlow DataFlow, uint32& OutDeviceIndex) const;
+		void GetDeviceBufferSize(uint32& DeviceBufferSize, EStreamType::Type StreamType);
+	};
+
+	/** 
+	* FUnrealAudioWasapi Implementation
+	*/
+
+	FUnrealAudioWasapi::FUnrealAudioWasapi()
+		: InputDeviceCount(0)
+		, OutputDeviceCount(0)
+		, bInitialized(false)
+		, bComInitialized(false)
+	{
+	}
+
+	FUnrealAudioWasapi::~FUnrealAudioWasapi()
+	{
+	}
+
+	bool FUnrealAudioWasapi::Initialize()
+	{
+ 		bComInitialized = FWindowsPlatformMisc::CoInitialize();
+
+		HRESULT Result = CoCreateInstance(__uuidof(MMDeviceEnumerator), NULL, CLSCTX_ALL, IID_PPV_ARGS(&WasapiInfo.DeviceEnumerator));
+		RETURN_FALSE_ON_FAIL(Result);
+
+		bInitialized = true;
+		return true;
+	}
+
+	bool FUnrealAudioWasapi::Shutdown()
+	{
+		if (!bInitialized)
+		{
+			return false;
+		}
+
+		SAFE_RELEASE(WasapiInfo.DeviceEnumerator);
+
+		if (bComInitialized)
+		{
+			FWindowsPlatformMisc::CoUninitialize();
+		}
+		return true;
+	}
+
+	bool FUnrealAudioWasapi::GetDevicePlatformApi(EDeviceApi::Type& OutType) const
+	{
+		OutType = EDeviceApi::WASAPI;
+		return true;
+	}
+
+	bool FUnrealAudioWasapi::GetNumOutputDevices(uint32& OutNumDevices) const
+	{
+		return GetDeviceCount(eRender, OutNumDevices);
+	}
+
+	bool FUnrealAudioWasapi::GetNumInputDevices(uint32& OutNumDevices) const
+	{
+		return GetDeviceCount(eCapture, OutNumDevices);
+	}
+
+	bool FUnrealAudioWasapi::GetDeviceCount(EDataFlow DataFlow, uint32& OutNumDevices) const
+	{
+		if (!bInitialized || !WasapiInfo.DeviceEnumerator)
+		{
+			return false;
+		}
+
+		IMMDeviceCollection* Devices = nullptr;
+		HRESULT Result = S_OK;
+
+		Result = WasapiInfo.DeviceEnumerator->EnumAudioEndpoints(DataFlow, DEVICE_STATE_ACTIVE, &Devices);
+		CLEANUP_ON_FAIL(Result);
+
+		Result = Devices->GetCount(&OutNumDevices);
+		CLEANUP_ON_FAIL(Result);
+
+	Cleanup:
+		SAFE_RELEASE(Devices);
+		return SUCCEEDED(Result);
+	}
+
+	bool FUnrealAudioWasapi::GetOutputDeviceInfo(const uint32 DeviceIndex, FDeviceInfo& OutInfo) const
+	{
+		return GetDeviceInfo(eRender, DeviceIndex, OutInfo);
+	}
+
+	bool FUnrealAudioWasapi::GetInputDeviceInfo(const uint32 DeviceIndex, FDeviceInfo& OutInfo) const
+	{
+		return GetDeviceInfo(eCapture, DeviceIndex, OutInfo);
+	}
+
+	bool FUnrealAudioWasapi::GetDeviceInfo(EDataFlow DataFlow, const uint32 DeviceIndex, FDeviceInfo& DeviceInfo) const
+	{
+		if (!bInitialized || !WasapiInfo.DeviceEnumerator)
+		{
+			return false;
+		}
+
+		IMMDeviceCollection* Devices = nullptr;
+		IMMDevice* DefaultDevice = nullptr;
+		IMMDevice* Device = nullptr;
+		IPropertyStore* DevicePropertyStore = nullptr;
+		IPropertyStore* DefaultDevicePropertyStore = nullptr;
+		IAudioClient* WasapiAudioClient = nullptr;
+		PROPVARIANT DeviceNameProperty;
+		PROPVARIANT DefaultDeviceNameProperty;
+		WAVEFORMATEX* WaveFormatEx = nullptr;
+		HRESULT Result = S_OK;
+
+		Result = WasapiInfo.DeviceEnumerator->EnumAudioEndpoints(DataFlow, DEVICE_STATE_ACTIVE, &Devices);
+		CLEANUP_ON_FAIL(Result);
+
+		uint32 NumDevices = 0;
+		Result = Devices->GetCount(&NumDevices);
+		CLEANUP_ON_FAIL(Result);
+
+		Result = WasapiInfo.DeviceEnumerator->GetDefaultAudioEndpoint(DataFlow, eConsole, &DefaultDevice);
+		CLEANUP_ON_FAIL(Result);
+
+		Result = Devices->Item(DeviceIndex, &Device);
+		CLEANUP_ON_FAIL(Result);
+
+		Result = Device->OpenPropertyStore(STGM_READ, &DevicePropertyStore);
+		CLEANUP_ON_FAIL(Result);
+
+		Result = DefaultDevice->OpenPropertyStore(STGM_READ, &DefaultDevicePropertyStore);
+		CLEANUP_ON_FAIL(Result);
+
+		Result = DevicePropertyStore->GetValue(PKEY_Device_FriendlyName, &DeviceNameProperty);
+		CLEANUP_ON_FAIL(Result);
+
+		Result = DefaultDevicePropertyStore->GetValue(PKEY_Device_FriendlyName, &DefaultDeviceNameProperty);
+		CLEANUP_ON_FAIL(Result);
+
+		Result = Device->Activate(__uuidof(IAudioClient), CLSCTX_ALL, NULL, (void **)&WasapiAudioClient);
+		CLEANUP_ON_FAIL(Result);
+
+		Result = WasapiAudioClient->GetMixFormat(&WaveFormatEx);
+		CLEANUP_ON_FAIL(Result);
+
+		// At this point we've succeeded in getting all the device objects from WindowsMM that we need
+		DeviceInfo.FriendlyName = FString(DeviceNameProperty.pwszVal);
+		DeviceInfo.bIsSystemDefault = (DeviceInfo.FriendlyName == FString(DefaultDeviceNameProperty.pwszVal));
+		DeviceInfo.NumChannels = WaveFormatEx->nChannels;
+		DeviceInfo.StreamType = (DataFlow == eCapture) ? EStreamType::INPUT : EStreamType::OUTPUT;
+		DeviceInfo.FrameRate = WaveFormatEx->nSamplesPerSec;
+
+		// Figure out native supported data formats of the device
+		DeviceInfo.StreamFormat = EStreamFormat::UNSUPPORTED;
+
+		bool bIsFloat = (WaveFormatEx->wFormatTag == WAVE_FORMAT_IEEE_FLOAT);
+		bool bIsExtensible = (WaveFormatEx->wFormatTag == WAVE_FORMAT_EXTENSIBLE);
+		bool bIsSubtypeFloat = (bIsExtensible && (((WAVEFORMATEXTENSIBLE*)WaveFormatEx)->SubFormat == KSDATAFORMAT_SUBTYPE_IEEE_FLOAT));
+
+		if (bIsFloat || bIsSubtypeFloat)
+		{
+			if (WaveFormatEx->wBitsPerSample == 32)
+			{
+				DeviceInfo.StreamFormat = EStreamFormat::FLT;
+			}
+			else if (WaveFormatEx->wBitsPerSample == 64)
+			{
+				DeviceInfo.StreamFormat = EStreamFormat::DBL;
+			}
+		}
+		else
+		{
+			bool bIsPcm = (WaveFormatEx->wFormatTag == WAVE_FORMAT_PCM);
+			bool bIsSubtypePcm = (bIsExtensible && (((WAVEFORMATEXTENSIBLE*)WaveFormatEx)->SubFormat == KSDATAFORMAT_SUBTYPE_PCM));
+			if (bIsPcm || bIsSubtypePcm)
+			{
+				if (WaveFormatEx->wBitsPerSample == 16)
+				{
+					DeviceInfo.StreamFormat = EStreamFormat::INT_16;
+				}
+				else if (WaveFormatEx->wBitsPerSample == 24)
+				{
+					DeviceInfo.StreamFormat = EStreamFormat::INT_24;
+				}
+				else if (WaveFormatEx->wBitsPerSample == 32)
+				{
+					DeviceInfo.StreamFormat = EStreamFormat::INT_32;
+				}
+				else
+				{
+					goto Cleanup;
+				}
+			}
+		}
+
+		DeviceInfo.Speakers.Reset();
+		uint32 ChanCount = 0;
+		if (bIsExtensible)
+		{
+			// Build a flag field of the speaker outputs of this device
+			WAVEFORMATEXTENSIBLE* WaveFormatExtensible = (WAVEFORMATEXTENSIBLE*)WaveFormatEx;
+			GetArrayOfSpeakers(DeviceInfo.NumChannels, DeviceInfo.Speakers, WaveFormatExtensible->dwChannelMask);
+		}
+		else
+		{
+			// Non-extensible formats only support 1 or 2 channels
+			DeviceInfo.Speakers.Add(ESpeaker::FRONT_LEFT);
+			if (DeviceInfo.NumChannels == 2)
+			{
+				DeviceInfo.Speakers.Add(ESpeaker::FRONT_RIGHT);
+			}
+		}
+
+Cleanup:
+		// Release all our resources that we used above
+		SAFE_RELEASE(Devices);
+		SAFE_RELEASE(DefaultDevice);
+		SAFE_RELEASE(Device);
+		SAFE_RELEASE(DevicePropertyStore);
+		SAFE_RELEASE(DefaultDevicePropertyStore);
+
+		PropVariantClear(&DeviceNameProperty);
+		PropVariantClear(&DefaultDeviceNameProperty);
+
+		CoTaskMemFree(WaveFormatEx);
+
+		return SUCCEEDED(Result);
+	}
+
+	bool FUnrealAudioWasapi::GetDefaultOutputDeviceIndex(uint32& OutDefaultIndex) const
+	{
+		return GetDefaultDeviceIndex(eRender, OutDefaultIndex);
+	}
+
+	bool FUnrealAudioWasapi::GetDefaultInputDeviceIndex(uint32& OutDefaultIndex) const
+	{
+		return GetDefaultDeviceIndex(eCapture, OutDefaultIndex);
+	}
+
+	bool FUnrealAudioWasapi::GetDefaultDeviceIndex(EDataFlow DataFlow, uint32& OutDeviceIndex) const
+	{
+		uint32 NumDevices = 0;
+		if (!GetDeviceCount(DataFlow, NumDevices))
+		{
+			return false;
+		}
+
+		bool bFoundDeviceIndex = false;
+
+		for (uint32 DeviceIndex = 0; DeviceIndex < NumDevices; ++DeviceIndex)
+		{
+			FDeviceInfo DeviceInfo;
+			if (!GetDeviceInfo(DataFlow, DeviceIndex, DeviceInfo))
+			{
+				return false;
+			}
+
+			if (DeviceInfo.bIsSystemDefault)
+			{
+				bFoundDeviceIndex = true;
+				OutDeviceIndex = DeviceIndex;
+				break;
+			}
+		}
+
+		return bFoundDeviceIndex;
+	}
+
+	bool FUnrealAudioWasapi::OpenDevices(const FCreateStreamParams& CreateStreamParams)
+	{
+		if (!bInitialized || WasapiInfo.bDevicesOpen)
+		{
+			return false;
+		}
+
+		StreamInfo.BlockSize = CreateStreamParams.CallbackBlockSize;
+
+		StreamInfo.FrameRate = INDEX_NONE;
+
+		// If we have input stream params, we try to open an input device
+		if (CreateStreamParams.InputDeviceIndex != INDEX_NONE)
+		{
+			if (!OpenDevice(CreateStreamParams.InputDeviceIndex, EStreamType::INPUT))
+			{
+				return false;
+			}
+			FStreamDeviceInfo& StreamDeviceInfo = StreamInfo.GetDeviceInfo(EStreamType::INPUT);
+			StreamInfo.FrameRate = FMath::Min(StreamInfo.FrameRate, StreamDeviceInfo.FrameRate);
+		}
+		
+		check(CreateStreamParams.OutputDeviceIndex != INDEX_NONE);
+		if (!OpenDevice(CreateStreamParams.OutputDeviceIndex, EStreamType::OUTPUT))
+		{
+			return false;
+		}
+		FStreamDeviceInfo& StreamDeviceInfo = StreamInfo.GetDeviceInfo(EStreamType::OUTPUT);
+		StreamInfo.FrameRate = FMath::Min(StreamInfo.FrameRate, StreamDeviceInfo.FrameRate);
+
+		return true;
+	}
+
+	bool FUnrealAudioWasapi::OpenDevice(uint32 DeviceIndex, EStreamType::Type StreamType)
+	{
+		check(WasapiInfo.DeviceEnumerator);
+
+		IMMDevice* Device = nullptr;
+		IMMDeviceCollection* DeviceList = nullptr;
+		WAVEFORMATEX* DeviceFormat = nullptr;
+		FDeviceInfo DeviceInfo;
+		HRESULT Result = S_OK;
+
+		EDataFlow DataFlow = (StreamType == EStreamType::INPUT) ? eCapture : eRender;
+
+		Result = WasapiInfo.DeviceEnumerator->EnumAudioEndpoints(DataFlow, DEVICE_STATE_ACTIVE, &DeviceList);
+		CLEANUP_ON_FAIL(Result);
+
+		if (!GetDeviceInfo(DataFlow, DeviceIndex, DeviceInfo))
+		{
+			Result = -1;
+			goto Cleanup;
+		}
+
+		Result = DeviceList->Item(DeviceIndex, &Device);
+		CLEANUP_ON_FAIL(Result);
+
+		FStreamDeviceInfo& StreamDeviceInfo = StreamInfo.GetDeviceInfo(StreamType);
+
+		StreamDeviceInfo.DeviceIndex = DeviceIndex;
+
+		IAudioClient*& AudioClient = (StreamType == EStreamType::INPUT) ? WasapiInfo.CaptureClient : WasapiInfo.RenderClient;
+		check(AudioClient == nullptr);
+
+		Result = Device->Activate(__uuidof(IAudioClient), CLSCTX_ALL, nullptr, (void**)&AudioClient);
+		CLEANUP_ON_FAIL(Result);
+
+		Result = AudioClient->GetMixFormat(&DeviceFormat);
+		CLEANUP_ON_FAIL(Result);
+
+		StreamDeviceInfo.bPerformByteSwap = false;
+		StreamDeviceInfo.NumChannels = DeviceFormat->nChannels;
+		StreamDeviceInfo.DeviceDataFormat = DeviceInfo.StreamFormat;
+		StreamDeviceInfo.FrameRate = DeviceInfo.FrameRate;
+		StreamDeviceInfo.Speakers = DeviceInfo.Speakers;
+
+		if (StreamDeviceInfo.DeviceDataFormat != EStreamFormat::FLT)
+		{
+			StreamDeviceInfo.bPerformFormatConversion = true;
+		}
+		SetupBufferFormatConvertInfo(StreamType);
+
+		// Initialize the raw buffer used to write into
+		uint32 BufferSize = StreamDeviceInfo.NumChannels * StreamInfo.BlockSize * GetNumBytesForFormat(EStreamFormat::FLT);
+		StreamDeviceInfo.UserBuffer.Init(0, BufferSize);
+
+	Cleanup:
+
+		SAFE_RELEASE(DeviceList);
+		SAFE_RELEASE(Device);
+		CoTaskMemFree(DeviceFormat);
+
+		if (FAILED(Result))
+		{
+			ShutdownStream();
+		}
+
+		return SUCCEEDED(Result);
+	}
+
+	bool FUnrealAudioWasapi::StartStream()
+	{
+		if (!bInitialized || StreamInfo.State == EStreamState::CLOSED || StreamInfo.State == EStreamState::RUNNING)
+		{
+			return false;
+		}
+		check(StreamInfo.Thread == nullptr);
+		StreamInfo.Thread = FRunnableThread::Create(this, TEXT("WasapiDeviceThread"), 0, TPri_AboveNormal);
+		return true;
+	}
+
+	bool FUnrealAudioWasapi::StopStream()
+	{
+		if (!bInitialized || StreamInfo.State == EStreamState::CLOSED || StreamInfo.State == EStreamState::STOPPED || StreamInfo.State == EStreamState::STOPPING)
+		{
+			return false;
+		}
+
+		check(StreamInfo.Thread != nullptr);
+		check(WasapiInfo.CaptureClient != nullptr || WasapiInfo.RenderClient != nullptr);
+
+		// Tell the stream update that we are stopping
+		Stop();
+
+			// Wait for the thread to finish
+		StreamInfo.Thread->WaitForCompletion();
+
+		check(StreamInfo.State == EStreamState::STOPPED);
+
+		// delete the thread
+		delete StreamInfo.Thread;
+		StreamInfo.Thread = nullptr;
+
+		// clean up any audio resources
+		HRESULT Result = S_OK;
+		bool bSuccess = true;
+
+		if (WasapiInfo.CaptureClient)
+		{
+			Result = WasapiInfo.CaptureClient->Stop();
+			if (FAILED(Result))
+			{
+				bSuccess = false;
+				UA_DEVICE_PLATFORM_ERROR("Failed to stop capture client.");
+			}
+		}
+
+		if (WasapiInfo.RenderClient)
+		{
+			Result = WasapiInfo.RenderClient->Stop();
+			if (FAILED(Result))
+			{
+				bSuccess = false;
+				UA_DEVICE_PLATFORM_ERROR("Failed to stop render client.");
+			}
+		}
+
+		return true;
+	}
+
+	bool FUnrealAudioWasapi::ShutdownStream()
+	{
+		if (!bInitialized || StreamInfo.State == EStreamState::CLOSED)
+		{
+			return false;
+		}
+
+		if (StreamInfo.State != EStreamState::STOPPED)
+		{
+			StopStream();
+		}
+
+		SAFE_RELEASE(WasapiInfo.CaptureClient);
+		SAFE_RELEASE(WasapiInfo.CaptureService);
+		SAFE_RELEASE(WasapiInfo.RenderClient);
+		SAFE_RELEASE(WasapiInfo.RenderService);
+
+		if (WasapiInfo.CaptureEvent)
+		{
+			CloseHandle(WasapiInfo.CaptureEvent);
+			WasapiInfo.CaptureEvent = nullptr;
+		}
+
+		if (WasapiInfo.RenderEvent)
+		{
+			CloseHandle(WasapiInfo.RenderEvent);
+			WasapiInfo.RenderEvent = nullptr;
+		}
+
+		StreamInfo.State = EStreamState::CLOSED;
+		return true;
+	}
+
+	bool FUnrealAudioWasapi::GetLatency(uint32& OutputDeviceLatency, uint32& InputDeviceLatency) const
+	{
+		OutputDeviceLatency = StreamInfo.GetDeviceInfo(EStreamType::OUTPUT).Latency;
+		InputDeviceLatency = StreamInfo.GetDeviceInfo(EStreamType::INPUT).Latency;
+		return true;
+	}
+
+	bool FUnrealAudioWasapi::GetFrameRate(uint32& OutFrameRate) const
+	{
+		OutFrameRate = StreamInfo.FrameRate;
+		return true;
+	}
+
+	// FRunnable
+
+	uint32 FUnrealAudioWasapi::Run()
+	{
+		// This is on a new thread, so lets initialize com again in case it wasn't already
+		bool bThreadComInitialized = FWindowsPlatformMisc::CoInitialize();
+
+		// Get local stack versions of various wasapi info we need for the device thread
+		WAVEFORMATEX* InputFormat = nullptr;
+		WAVEFORMATEX* OutputFormat = nullptr;
+		bool bWasInputBufferRead = false;
+		bool bWasOutputBufferWritten = false;
+		bool bWasInputBufferWritten = false;
+		bool bOutputCallbackCalled = false;
+		BYTE* DeviceByteBuffer = nullptr;		// byte buffer to hold input data from input device
+		uint32 CaptureBufferFrameCount = 0;		// how many frames read in from input device
+		::DWORD CaptureFlags = 0;				// flags from wasapi on audio input
+		FStreamDeviceInfo& InputDeviceInfo = StreamInfo.GetDeviceInfo(EStreamType::INPUT);
+		FStreamDeviceInfo& OutputDeviceInfo = StreamInfo.GetDeviceInfo(EStreamType::OUTPUT);
+		uint32 NumBytesRenderFormat = GetNumBytesForFormat(OutputDeviceInfo.DeviceDataFormat);
+		uint32 NumBytesCaptureFormat = GetNumBytesForFormat(InputDeviceInfo.DeviceDataFormat);
+		uint32 InputDeviceSamples = StreamInfo.BlockSize * InputDeviceInfo.NumChannels;
+		uint32 OutputDeviceSamples = StreamInfo.BlockSize * OutputDeviceInfo.NumChannels;
+		HRESULT Result = S_OK;
+		FCallbackInfo CallbackInfo;
+
+		StreamInfo.State = EStreamState::RUNNING;
+
+		// Number used to compute the REFERENCE_TIME for callback periods (REF_TIME is 100 nanoseconds)
+		static const uint64 REF_TIMES_PER_SECOND = 10000000;
+
+		// If we are trying to get audio from an input device, set up a capture stream
+		if (WasapiInfo.CaptureClient)
+		{
+
+			// We haven't set up an actual capture service yet
+			check(WasapiInfo.CaptureService == nullptr);
+			
+			Result = WasapiInfo.CaptureClient->GetMixFormat(&InputFormat);
+			CLEANUP_ON_FAIL(Result);
+
+			// The actual buffer size we're going to use, independent of the actual buffer size we want. 
+			// Compute the callback buffer period using the buffer size
+			REFERENCE_TIME CallbackBufferPeriod = (REFERENCE_TIME)((StreamInfo.BlockSize * REF_TIMES_PER_SECOND) / InputFormat->nSamplesPerSec);
+
+			Result = WasapiInfo.CaptureClient->Initialize(
+				AUDCLNT_SHAREMODE_SHARED,				// Other clients can use this input device
+				AUDCLNT_STREAMFLAGS_EVENTCALLBACK,		// Processing of the buffer by client is event driven
+				CallbackBufferPeriod,					// the size of the buffer in 100-nanosecond units (REF_TIME)
+				0,										// always 0 in shared mode 
+				InputFormat,							// the input format to use
+				nullptr									// audio session guid, we're ignoring
+			);
+			CLEANUP_ON_FAIL(Result);
+
+			REFERENCE_TIME InputLatency = 0;
+			WasapiInfo.CaptureClient->GetStreamLatency(&InputLatency);
+
+			// Get the reported device latency -- note for WASAPI, this can only be called after Initialize on capture client
+			InputDeviceInfo.Latency = (1000 * InputLatency) / REF_TIMES_PER_SECOND;
+
+			Result = WasapiInfo.CaptureClient->GetService(__uuidof(IAudioCaptureClient), (void**)&WasapiInfo.CaptureService);
+			CLEANUP_ON_FAIL(Result);
+
+			WasapiInfo.CaptureEvent = CreateEvent(nullptr, false, false, nullptr);
+			if (!WasapiInfo.CaptureClient)
+			{
+				Result = HRESULT_FROM_WIN32(GetLastError());
+				CLEANUP_ON_FAIL(Result);
+			}
+			Result = WasapiInfo.CaptureClient->SetEventHandle(WasapiInfo.CaptureEvent);
+			CLEANUP_ON_FAIL(Result);
+
+			uint32 ReadBufferSize = 0;
+			Result = WasapiInfo.CaptureClient->GetBufferSize(&ReadBufferSize);
+			CLEANUP_ON_FAIL(Result);
+
+			uint32 WriteBufferSize = StreamInfo.BlockSize * InputDeviceInfo.NumChannels;
+			ReadBufferSize *= InputDeviceInfo.NumChannels;
+
+			WasapiInfo.CaptureIntermediateBuffer = IIntermediateBuffer::CreateIntermediateBuffer(InputDeviceInfo.DeviceDataFormat);
+			check(WasapiInfo.CaptureIntermediateBuffer);
+			WasapiInfo.CaptureIntermediateBuffer->Initialize(ReadBufferSize + WriteBufferSize, NumBytesCaptureFormat);
+
+			Result = WasapiInfo.CaptureClient->Reset();
+			CLEANUP_ON_FAIL(Result);
+
+			Result = WasapiInfo.CaptureClient->Start();
+			CLEANUP_ON_FAIL(Result);
+		}
+
+		if (WasapiInfo.RenderClient)
+		{
+			// We haven't set up an actual render service yet
+			check(WasapiInfo.RenderService == nullptr);
+
+			Result = WasapiInfo.RenderClient->GetMixFormat(&OutputFormat);
+			CLEANUP_ON_FAIL(Result);
+
+			// Compute the callback buffer period using the buffer size
+			REFERENCE_TIME CallbackBufferPeriod = (REFERENCE_TIME)((StreamInfo.BlockSize * REF_TIMES_PER_SECOND) / OutputFormat->nSamplesPerSec);
+
+			Result = WasapiInfo.RenderClient->Initialize(
+				AUDCLNT_SHAREMODE_SHARED,				// Other clients can use this input device
+				AUDCLNT_STREAMFLAGS_EVENTCALLBACK,		// Processing of the buffer by client is event driven
+				CallbackBufferPeriod,					// the size of the buffer in 100-nanosecond units (REF_TIME)
+				0,										// always 0 in shared mode 
+				OutputFormat,							// the output format to use
+				nullptr									// audio session guid, we're ignoring
+			);
+			CLEANUP_ON_FAIL(Result);
+
+			// Get the output latency
+			REFERENCE_TIME OutputLatency = 0;
+			WasapiInfo.RenderClient->GetStreamLatency(&OutputLatency);
+			OutputDeviceInfo.Latency = (1000 * OutputLatency) / REF_TIMES_PER_SECOND;
+
+			Result = WasapiInfo.RenderClient->GetService(__uuidof(IAudioRenderClient), (void**)&WasapiInfo.RenderService);
+			CLEANUP_ON_FAIL(Result);
+
+			WasapiInfo.RenderEvent = CreateEvent(nullptr, false, false, nullptr);
+			if (!WasapiInfo.RenderEvent)
+			{
+				Result = HRESULT_FROM_WIN32(GetLastError());
+				CLEANUP_ON_FAIL(Result);
+			}
+
+			Result = WasapiInfo.RenderClient->SetEventHandle(WasapiInfo.RenderEvent);
+			CLEANUP_ON_FAIL(Result);
+
+			uint32 ReadBufferSize = 0;
+			Result = WasapiInfo.RenderClient->GetBufferSize(&ReadBufferSize);
+			CLEANUP_ON_FAIL(Result);
+
+			uint32 WriteBufferSize = StreamInfo.BlockSize* OutputDeviceInfo.NumChannels;
+
+			ReadBufferSize *= OutputDeviceInfo.NumChannels;
+
+			WasapiInfo.RenderIntermediateBuffer = IIntermediateBuffer::CreateIntermediateBuffer(OutputDeviceInfo.DeviceDataFormat);
+			check(WasapiInfo.RenderIntermediateBuffer);
+			WasapiInfo.RenderIntermediateBuffer->Initialize(ReadBufferSize + WriteBufferSize, NumBytesRenderFormat);
+
+			Result = WasapiInfo.RenderClient->Reset();
+			CLEANUP_ON_FAIL(Result);
+
+			Result = WasapiInfo.RenderClient->Start();
+			CLEANUP_ON_FAIL(Result);
+		}
+
+		// Set up the device buffer (buffer read from and written directly to the device in the native device format)
+
+		// If we have an input stream, set up our device buffer size to accommodate
+		uint32 DeviceBufferSize = 0;
+		if (StreamInfo.StreamType == EStreamType::INPUT)
+		{
+			uint32 DeviceBufferInput = 0;
+			GetDeviceBufferSize(DeviceBufferInput, EStreamType::INPUT);
+
+			uint32 DeviceBufferOutput = 0;
+			GetDeviceBufferSize(DeviceBufferOutput, EStreamType::OUTPUT);
+
+			DeviceBufferSize = FMath::Max(DeviceBufferInput, DeviceBufferOutput);
+		}
+		else
+		{
+			GetDeviceBufferSize(DeviceBufferSize, StreamInfo.StreamType);
+		}
+
+		StreamInfo.DeviceBuffer.Init(0, DeviceBufferSize);
+		uint32 Count = 0;
+
+		// Prepare the struct which will be used to make audio callbacks
+		CallbackInfo.OutBuffer = (float *)OutputDeviceInfo.UserBuffer.GetData();;
+		CallbackInfo.InBuffer = (float *)InputDeviceInfo.UserBuffer.GetData();;
+		CallbackInfo.NumFrames = StreamInfo.BlockSize;
+		CallbackInfo.NumChannels = OutputDeviceInfo.NumChannels;
+		CallbackInfo.NumInputChannels = InputDeviceInfo.NumChannels;
+		CallbackInfo.StreamTime = 0.0;
+		CallbackInfo.UserData = StreamInfo.UserData;
+		CallbackInfo.StatusFlags = 0;
+		CallbackInfo.OutputSpeakers = OutputDeviceInfo.Speakers;
+		CallbackInfo.FrameRate = OutputDeviceInfo.FrameRate;
+
+		/**
+		
+		The update loop performs the following stages:
+
+		1. Prepare input audio to feed into user callback
+			- input was pushed into intermediate buffer from previous update tick
+			- input is retrieved from intermediate buffer
+			- input is converted from device format from device buffer to user format in user buffer
+
+		2. Perform the user callback
+			- user callback is called with possible input data from stage 1
+			- return value of user callback is evaluated (success or failure) to determine the state of the thread loop
+
+		3. Convert user callback data buffer format to device data buffer format
+			- write results to intermediate output buffer
+
+		4. Get new audio input from device
+			- get input data from audio device
+			- push the input data to the intermediate buffer so that in next update it can be given to user callback
+
+		5. Write audio to output buffer
+			- get output data from intermediate buffer
+			- write output data to render device
+
+		*/
+
+		while (StreamInfo.State != EStreamState::STOPPING)
+		{
+			//UE_LOG(LogUnrealAudioDevice, Display, TEXT("Device Callback (%d)"), ++Count);
+
+			// Stage 1
+
+			if (!bWasInputBufferRead)
+			{
+				if (WasapiInfo.CaptureClient)
+				{
+					// Read data from the intermediate buffer so that it can be fed to the user callback
+					if (InputDeviceInfo.bPerformFormatConversion)
+					{
+						bWasInputBufferRead = WasapiInfo.CaptureIntermediateBuffer->Read(StreamInfo.DeviceBuffer.GetData(), InputDeviceSamples);
+						if (bWasInputBufferRead)
+						{
+							ConvertBufferFormat(InputDeviceInfo.UserBuffer, StreamInfo.DeviceBuffer, EStreamType::INPUT);
+						}
+					}
+					else
+					{
+						bWasInputBufferRead = WasapiInfo.CaptureIntermediateBuffer->Read(InputDeviceInfo.UserBuffer.GetData(), InputDeviceSamples);
+					}
+				}
+				else
+				{
+					// If there was no audio to read in (i.e. we are in device output only) pretend like we read audio input
+					bWasInputBufferRead = true;
+				}
+
+				// Stage 2
+
+				if (bWasInputBufferRead)
+				{
+					// Report to callback if the input device reported any discontinuities in the input stream
+					CallbackInfo.StatusFlags = 0;
+					if (CaptureFlags & AUDCLNT_BUFFERFLAGS_DATA_DISCONTINUITY)
+					{
+						CallbackInfo.StatusFlags |= EStreamFlowStatus::INPUT_OVERFLOW;
+					}
+
+					CallbackInfo.StreamTime = StreamInfo.StreamTime;
+
+					// Zero out the output buffer before calling the user callback
+					memset((void*)CallbackInfo.OutBuffer, 0, StreamInfo.BlockSize*OutputDeviceInfo.NumChannels*sizeof(float));
+
+					// Make the user callback
+					bool bSuccess = StreamInfo.CallbackFunction(CallbackInfo);
+
+					UpdateStreamTimeTick();
+
+					// Check if we need to stop the device thread
+					if (!bSuccess)
+					{
+						// Now stop the thread...
+						StreamInfo.State = EStreamState::STOPPING;
+					}
+				}
+			}
+
+			// Stage 3
+
+			if (WasapiInfo.RenderClient && bWasInputBufferRead)
+			{
+				if (OutputDeviceInfo.bPerformFormatConversion)
+				{
+					ConvertBufferFormat(StreamInfo.DeviceBuffer, OutputDeviceInfo.UserBuffer, EStreamType::OUTPUT);
+					bWasOutputBufferWritten = WasapiInfo.RenderIntermediateBuffer->Write(StreamInfo.DeviceBuffer.GetData(), OutputDeviceSamples);
+				}
+				else
+				{
+					bWasOutputBufferWritten = WasapiInfo.RenderIntermediateBuffer->Write(OutputDeviceInfo.UserBuffer.GetData(), OutputDeviceSamples);
+				}
+			}
+			else
+			{
+				bWasOutputBufferWritten = true;
+			}
+
+			// Stage 4
+
+			if (WasapiInfo.CaptureClient)
+			{
+				// if we didn't read from the input buffer then, pause this thread until we get an event
+				// signal from the input audio client (this will get sent as soon as the device has read an input buffer)
+				if (!bWasInputBufferRead)
+				{
+					WaitForSingleObject(WasapiInfo.CaptureEvent, INFINITE);
+				}
+
+				Result = WasapiInfo.CaptureService->GetBuffer(&DeviceByteBuffer, &CaptureBufferFrameCount, &CaptureFlags, nullptr, nullptr);
+				CLEANUP_ON_FAIL(Result);
+
+				bool bCaptureSuccess = false;
+				if (CaptureBufferFrameCount != 0)
+				{
+					uint32 WriteSamples = CaptureBufferFrameCount * InputDeviceInfo.NumChannels;
+					bWasInputBufferWritten = WasapiInfo.CaptureIntermediateBuffer->Write((uint8*)DeviceByteBuffer, WriteSamples);
+					if (bWasInputBufferWritten)
+					{
+						bCaptureSuccess = true;
+						Result = WasapiInfo.CaptureService->ReleaseBuffer(CaptureBufferFrameCount);
+						CLEANUP_ON_FAIL(Result);
+					}
+				}
+
+				if (!bCaptureSuccess)
+				{
+					// We failed to read data from input device, releasing 0 bytes will release the capture buffer
+					Result = WasapiInfo.CaptureService->ReleaseBuffer(0);
+					CLEANUP_ON_FAIL(Result);
+				}
+			}
+
+			// Stage 5
+
+			if (WasapiInfo.RenderClient)
+			{
+				// If the user output buffer was not pushed to the intermediate buffer, wait for the 
+				// next render event from the device before going on
+				if (bWasInputBufferRead && !bWasOutputBufferWritten)
+				{
+					WaitForSingleObject(WasapiInfo.RenderEvent, INFINITE);
+				}
+
+				uint32 OutputBufferFrameCount = 0;
+				uint32 OutputBufferFramePadding = 0;
+
+				// Get render buffer from stream
+				Result = WasapiInfo.RenderClient->GetBufferSize(&OutputBufferFrameCount);
+				CLEANUP_ON_FAIL(Result);
+
+				Result = WasapiInfo.RenderClient->GetCurrentPadding(&OutputBufferFramePadding);
+				CLEANUP_ON_FAIL(Result);
+
+				OutputBufferFrameCount -= OutputBufferFramePadding;
+
+				if (OutputBufferFrameCount != 0)
+				{
+					Result = WasapiInfo.RenderService->GetBuffer(OutputBufferFrameCount, &DeviceByteBuffer);
+					CLEANUP_ON_FAIL(Result);
+
+					// Read the next buffer from the intermediate output buffer
+					uint32 ReadSamples = OutputBufferFrameCount * OutputDeviceInfo.NumChannels;
+					bool bSuccess = WasapiInfo.RenderIntermediateBuffer->Read((uint8*)DeviceByteBuffer, ReadSamples);
+					if (bSuccess)
+					{
+						Result = WasapiInfo.RenderService->ReleaseBuffer(OutputBufferFrameCount, 0);
+						CLEANUP_ON_FAIL(Result);
+					}
+					else
+					{
+						Result = WasapiInfo.RenderService->ReleaseBuffer(0, 0);
+						CLEANUP_ON_FAIL(Result);
+					}
+				}
+			}
+
+			if (bWasOutputBufferWritten)
+			{
+				bWasInputBufferRead = false;
+			}
+		}
+
+		Cleanup:
+
+		CoTaskMemFree(InputFormat);
+		CoTaskMemFree(OutputFormat);
+
+		if (bThreadComInitialized)
+		{
+			FWindowsPlatformMisc::CoUninitialize();
+		}
+
+		StreamInfo.State = EStreamState::STOPPED;
+
+		return 0;
+	}
+
+	void FUnrealAudioWasapi::Stop()
+	{
+		StreamInfo.State = EStreamState::STOPPING;
+	}
+
+	void FUnrealAudioWasapi::Exit()
+	{
+
+	}
+
+	void FUnrealAudioWasapi::GetDeviceBufferSize(uint32& DeviceBufferSize, EStreamType::Type StreamType)
+	{
+		FStreamDeviceInfo& DeviceInfo = StreamInfo.GetDeviceInfo(StreamType);
+		uint32 NumBytes = GetNumBytesForFormat(DeviceInfo.DeviceDataFormat);
+		uint32 NumChannels = DeviceInfo.NumChannels;
+		DeviceBufferSize = StreamInfo.BlockSize * NumChannels * NumBytes;
+	}
+}
+
+IMPLEMENT_MODULE(UAudio::FUnrealAudioWasapi, UnrealAudioWasapi);
+
+#endif // #if ENABLE_UNREAL_AUDIO
+
+
diff --git a/Engine/Source/Runtime/Windows/UnrealAudioWasapi/UnrealAudioWasapi.Build.cs b/Engine/Source/Runtime/Windows/UnrealAudioWasapi/UnrealAudioWasapi.Build.cs
new file mode 100644
--- /dev/null
+++ b/Engine/Source/Runtime/Windows/UnrealAudioWasapi/UnrealAudioWasapi.Build.cs
@@ -0,0 +1,30 @@
+// Copyright 1998-2015 Epic Games, Inc. All Rights Reserved.
+
+using UnrealBuildTool;
+
+public class UnrealAudioWasapi : ModuleRules
+{
+	public UnrealAudioWasapi(TargetInfo Target)
+	{
+		PrivateIncludePaths.AddRange(
+			new string[] {
+				"Runtime/UnrealAudio/Private",
+			}
+)		;
+
+		PrivateIncludePathModuleNames.Add("TargetPlatform");
+
+		PrivateDependencyModuleNames.AddRange(
+		new string[] {
+				"Core",
+				"CoreUObject",
+			}
+		);
+
+		PublicDependencyModuleNames.AddRange(
+			new string[] {
+				"UnrealAudio"
+			}
+		);
+	}
+}
diff --git a/Engine/Source/Runtime/Windows/UnrealAudioXAudio2/Private/UnrealAudioDeviceXAudio2.cpp b/Engine/Source/Runtime/Windows/UnrealAudioXAudio2/Private/UnrealAudioDeviceXAudio2.cpp
new file mode 100644
--- /dev/null
+++ b/Engine/Source/Runtime/Windows/UnrealAudioXAudio2/Private/UnrealAudioDeviceXAudio2.cpp
@@ -0,0 +1,352 @@
+// Copyright 1998-2015 Epic Games, Inc. All Rights Reserved.
+
+#pragma once 
+
+/**
+	Concrete implementation of IAudioDevice for XAudio2
+
+	See https://msdn.microsoft.com/en-us/library/windows/desktop/hh405049%28v=vs.85%29.aspx
+
+*/
+
+#include "UnrealAudioDeviceModule.h"
+#include "UnrealAudioBuffer.h"
+
+#include "AllowWindowsPlatformTypes.h"
+#include <xaudio2.h>
+#if PLATFORM_WINDOWS
+#include <mmdeviceapi.h>					// For getting endpoint notifications about audio devices
+#include <functiondiscoverykeys_devpkey.h>	// Defines property keys for the Plug and Play Device Property API.
+#endif
+#include "HideWindowsPlatformTypes.h"
+
+#if ENABLE_UNREAL_AUDIO
+
+DEFINE_LOG_CATEGORY(LogUnrealAudioDevice);
+
+// See MSDN documentation for what these error codes mean in the context of the API call
+static const TCHAR* GetXAudio2Error(HRESULT Result)
+{
+	switch (Result)
+	{
+		case XAUDIO2_E_INVALID_CALL:					return TEXT("XAUDIO2_E_INVALID_CALL");
+		case XAUDIO2_E_XMA_DECODER_ERROR:				return TEXT("XAUDIO2_E_XMA_DECODER_ERROR");
+		case XAUDIO2_E_XAPO_CREATION_FAILED:			return TEXT("XAUDIO2_E_XAPO_CREATION_FAILED");
+		case XAUDIO2_E_DEVICE_INVALIDATED:				return TEXT("XAUDIO2_E_DEVICE_INVALIDATED");
+		case REGDB_E_CLASSNOTREG:						return TEXT("REGDB_E_CLASSNOTREG");
+		case CLASS_E_NOAGGREGATION:						return TEXT("CLASS_E_NOAGGREGATION");
+		case E_NOINTERFACE:								return TEXT("E_NOINTERFACE");
+		case E_POINTER:									return TEXT("E_POINTER");
+		case E_INVALIDARG:								return TEXT("E_INVALIDARG");
+		case E_OUTOFMEMORY:								return TEXT("E_OUTOFMEMORY");
+		default: return TEXT("UKNOWN");
+	}
+}
+
+/** Some helper defines to reduce boilerplate code for error handling windows API calls. */
+
+#define CLEANUP_ON_FAIL(Result)						\
+	if (FAILED(Result))								\
+	{												\
+		const TCHAR* Err = GetXAudio2Error(Result);	\
+		UA_DEVICE_PLATFORM_ERROR(Err);				\
+		goto Cleanup;								\
+	}
+
+#define RETURN_FALSE_ON_FAIL(Result)				\
+	if (FAILED(Result))								\
+	{												\
+		const TCHAR* Err = GetXAudio2Error(Result);	\
+		UA_DEVICE_PLATFORM_ERROR(Err);				\
+		return false;								\
+	}
+
+namespace UAudio
+{
+	class FUnrealAudioXAudio2 : public IUnrealAudioDeviceModule,
+		public FRunnable
+	{
+	public:
+		FUnrealAudioXAudio2();
+		~FUnrealAudioXAudio2();
+
+		// IUnrealAudioDeviceModule
+		bool Initialize() override;
+		bool Shutdown() override;
+		bool GetDevicePlatformApi(EDeviceApi::Type & OutType) const override;
+		bool GetNumOutputDevices(uint32& OutNumDevices) const override;
+		bool GetOutputDeviceInfo(const uint32 DeviceIndex, FDeviceInfo& OutInfo) const override;
+		bool GetNumInputDevices(uint32& OutNumDevices) const override;
+		bool GetInputDeviceInfo(const uint32 DeviceIndex, FDeviceInfo& OutInfo) const override;
+		bool GetDefaultOutputDeviceIndex(uint32& OutDefaultIndex) const override;
+		bool GetDefaultInputDeviceIndex(uint32& OutDefaultIndex) const override;
+		bool StartStream() override;
+		bool StopStream() override;
+		bool ShutdownStream() override;
+		bool GetLatency(uint32& OutputDeviceLatency, uint32& InputDeviceLatency) const override;
+		bool GetFrameRate(uint32& OutFrameRate) const override;
+
+		// FRunnable
+		uint32 Run() override;
+		void Stop() override;
+		void Exit() override;
+
+	private:
+		struct FXAudio2Info
+		{
+			/** XAudio2 system object */
+			IXAudio2* XAudio2System;
+
+#if PLATFORM_WINDOWS
+			IMMDeviceEnumerator* DeviceEnumerator;
+#endif
+
+			FXAudio2Info()
+				: XAudio2System(nullptr)
+#if PLATFORM_WINDOWS
+				, DeviceEnumerator(nullptr)
+#endif
+			{
+			}
+		};
+
+		/** XAudio2-specific data */
+		FXAudio2Info XAudio2Info;
+
+
+		/** Whether or not the device api has been initialized */
+		bool bInitialized;
+
+		/** Whether or not com was initialized */
+		bool bComInitialized;
+
+	private:
+		bool OpenDevices(const FCreateStreamParams& Params) override;
+
+		// Helper functions
+		EStreamFormat::Type GetStreamFormat(WAVEFORMATEX* WaveFormatEx) const;
+
+	};
+
+	/**
+	* FUnrealAudioXAudio2 Implementation
+	*/
+
+	FUnrealAudioXAudio2::FUnrealAudioXAudio2()
+		: bInitialized(false)
+		, bComInitialized(false)
+	{
+	}
+
+	FUnrealAudioXAudio2::~FUnrealAudioXAudio2()
+	{
+	}
+
+	bool FUnrealAudioXAudio2::Initialize()
+	{
+		bComInitialized = FWindowsPlatformMisc::CoInitialize();
+
+		// Create the xaudio2 system
+		HRESULT Result = XAudio2Create(&XAudio2Info.XAudio2System);
+		RETURN_FALSE_ON_FAIL(Result);
+
+#if PLATFORM_WINDOWS
+		Result = CoCreateInstance(__uuidof(MMDeviceEnumerator), NULL, CLSCTX_ALL, IID_PPV_ARGS(&XAudio2Info.DeviceEnumerator));
+		RETURN_FALSE_ON_FAIL(Result);
+#endif
+
+		bInitialized = true;
+		return true;
+	}
+
+	bool FUnrealAudioXAudio2::Shutdown()
+	{
+		if (!bInitialized)
+		{
+			return false;
+		}
+
+		SAFE_RELEASE(XAudio2Info.XAudio2System);
+
+		if (bComInitialized)
+		{
+			FWindowsPlatformMisc::CoUninitialize();
+		}
+
+		return true;
+	}
+
+	bool FUnrealAudioXAudio2::GetDevicePlatformApi(EDeviceApi::Type& OutType) const
+	{
+		OutType = EDeviceApi::XAUDIO2;
+		return true;
+	}
+
+	bool FUnrealAudioXAudio2::GetNumOutputDevices(uint32& OutNumDevices) const
+	{
+		if (!bInitialized || !XAudio2Info.XAudio2System)
+		{
+			return false;
+		}
+
+		// Note: XAudio2 only supports output devices, so GetDeviceCount() is only returning output devices
+		HRESULT Result = XAudio2Info.XAudio2System->GetDeviceCount(&OutNumDevices);
+		RETURN_FALSE_ON_FAIL(Result);
+
+		return true;
+	}
+
+	bool FUnrealAudioXAudio2::GetNumInputDevices(uint32& OutNumDevices) const
+	{
+#if PLATFORM_WINDOWS
+		// On windows we're going to use WASAPI/WINMM to do audio input since XAudio2 doesn't support capture devices
+		if (!bInitialized || !XAudio2Info.DeviceEnumerator)
+		{
+			return false;
+		}
+
+		IMMDeviceCollection* Devices = nullptr;
+		HRESULT Result = S_OK;
+
+		Result = XAudio2Info.DeviceEnumerator->EnumAudioEndpoints(eCapture, DEVICE_STATE_ACTIVE, &Devices);
+		CLEANUP_ON_FAIL(Result);
+
+		Result = Devices->GetCount(&OutNumDevices);
+		CLEANUP_ON_FAIL(Result);
+
+	Cleanup:
+		SAFE_RELEASE(Devices);
+		return SUCCEEDED(Result);
+#else
+		// TODO: Figure out how to get input devices on XBOX api
+		// Look at using XInput and XAudio2 http://blogs.msdn.com/b/chuckw/archive/2012/05/03/xinput-and-xaudio2.aspx
+		OutNumDevices = 0;
+		return true;
+#endif
+	}
+
+	EStreamFormat::Type FUnrealAudioXAudio2::GetStreamFormat(WAVEFORMATEX* WaveFormatEx) const
+	{
+		bool bIsFloat = (WaveFormatEx->wFormatTag == WAVE_FORMAT_IEEE_FLOAT);
+		bool bIsExtensible = (WaveFormatEx->wFormatTag == WAVE_FORMAT_EXTENSIBLE);
+		bool bIsSubtypeFloat = (bIsExtensible && (((WAVEFORMATEXTENSIBLE*)WaveFormatEx)->SubFormat == KSDATAFORMAT_SUBTYPE_IEEE_FLOAT));
+
+		if (bIsFloat || bIsSubtypeFloat)
+		{
+			if (WaveFormatEx->wBitsPerSample == 32)
+			{
+				return EStreamFormat::FLT;
+			}
+			else if (WaveFormatEx->wBitsPerSample == 64)
+			{
+				return EStreamFormat::DBL;
+			}
+		}
+		else
+		{
+			bool bIsPcm = (WaveFormatEx->wFormatTag == WAVE_FORMAT_PCM);
+			bool bIsSubtypePcm = (bIsExtensible && (((WAVEFORMATEXTENSIBLE*)WaveFormatEx)->SubFormat == KSDATAFORMAT_SUBTYPE_PCM));
+			if (bIsPcm || bIsSubtypePcm)
+			{
+				if (WaveFormatEx->wBitsPerSample == 16)
+				{
+					return EStreamFormat::INT_16;
+				}
+				else if (WaveFormatEx->wBitsPerSample == 24)
+				{
+					return EStreamFormat::INT_24;
+				}
+				else if (WaveFormatEx->wBitsPerSample == 32)
+				{
+					return EStreamFormat::INT_32;
+				}
+			}
+		}
+		return EStreamFormat::UNSUPPORTED;
+	}
+
+	bool FUnrealAudioXAudio2::GetOutputDeviceInfo(const uint32 DeviceIndex, FDeviceInfo& OutInfo) const
+	{
+		if (!bInitialized)
+		{
+			return false;
+		}
+
+		check(XAudio2Info.XAudio2System);
+
+		XAUDIO2_DEVICE_DETAILS DeviceDetails;
+		HRESULT Result = XAudio2Info.XAudio2System->GetDeviceDetails(DeviceIndex, &DeviceDetails);
+		RETURN_FALSE_ON_FAIL(Result);
+
+		OutInfo.FrameRate = DeviceDetails.OutputFormat.Format.nSamplesPerSec;
+		OutInfo.NumChannels = DeviceDetails.OutputFormat.Format.nChannels;
+		OutInfo.StreamFormat = GetStreamFormat(&DeviceDetails.OutputFormat.Format);
+
+		return true;
+	}
+
+	bool FUnrealAudioXAudio2::GetInputDeviceInfo(const uint32 DeviceIndex, FDeviceInfo& OutInfo) const
+	{
+		return true;
+	}
+
+	bool FUnrealAudioXAudio2::GetDefaultOutputDeviceIndex(uint32& OutDefaultIndex) const
+	{
+		return true;
+	}
+
+	bool FUnrealAudioXAudio2::GetDefaultInputDeviceIndex(uint32& OutDefaultIndex) const
+	{
+		return true;
+	}
+
+	bool FUnrealAudioXAudio2::OpenDevices(const FCreateStreamParams& Params)
+	{
+		return true;
+	}
+
+	bool FUnrealAudioXAudio2::StartStream()
+	{
+		return true;
+	}
+
+	bool FUnrealAudioXAudio2::StopStream()
+	{
+		return true;
+	}
+
+	bool FUnrealAudioXAudio2::ShutdownStream()
+	{
+		return true;
+	}
+
+	bool FUnrealAudioXAudio2::GetLatency(uint32& OutputDeviceLatency, uint32& InputDeviceLatency) const
+	{
+		return true;
+	}
+
+	bool FUnrealAudioXAudio2::GetFrameRate(uint32& OutFrameRate) const
+	{
+		return true;
+	}
+
+	// FRunnable
+
+	uint32 FUnrealAudioXAudio2::Run()
+	{
+		return 0;
+	}
+
+	void FUnrealAudioXAudio2::Stop()
+	{
+	}
+
+	void FUnrealAudioXAudio2::Exit()
+	{
+
+	}
+}
+
+IMPLEMENT_MODULE(UAudio::FUnrealAudioXAudio2, UnrealAudioXAudio2);
+
+#endif // #if ENABLE_UNREAL_AUDIO
diff --git a/Engine/Source/Runtime/Windows/UnrealAudioXAudio2/UnrealAudioXAudio2.Build.cs b/Engine/Source/Runtime/Windows/UnrealAudioXAudio2/UnrealAudioXAudio2.Build.cs
new file mode 100644
--- /dev/null
+++ b/Engine/Source/Runtime/Windows/UnrealAudioXAudio2/UnrealAudioXAudio2.Build.cs
@@ -0,0 +1,34 @@
+// Copyright 1998-2015 Epic Games, Inc. All Rights Reserved.
+
+using UnrealBuildTool;
+
+public class UnrealAudioXAudio2 : ModuleRules
+{
+	public UnrealAudioXAudio2(TargetInfo Target)
+	{
+		PrivateIncludePaths.AddRange(
+			new string[] {
+				"Runtime/UnrealAudio/Private",
+			}
+)		;
+
+		PrivateIncludePathModuleNames.Add("TargetPlatform");
+
+		PrivateDependencyModuleNames.AddRange(
+		new string[] {
+				"Core",
+				"CoreUObject",
+			}
+		);
+
+		PublicDependencyModuleNames.AddRange(
+			new string[] {
+				"UnrealAudio"
+			}
+		);
+
+		AddThirdPartyPrivateStaticDependencies(Target, 
+			"DX11Audio"
+		);
+	}
+}
