Removing accidental #pragma once from cpp files

diff --git a/Engine/Source/Runtime/UnrealAudio/Private/Tests/UnrealAudioTestGenerators.cpp b/Engine/Source/Runtime/UnrealAudio/Private/Tests/UnrealAudioTestGenerators.cpp
--- a/Engine/Source/Runtime/UnrealAudio/Private/Tests/UnrealAudioTestGenerators.cpp
+++ b/Engine/Source/Runtime/UnrealAudio/Private/Tests/UnrealAudioTestGenerators.cpp
@@ -1,7 +1,5 @@
 // Copyright 1998-2015 Epic Games, Inc. All Rights Reserved.
 
-#pragma once 
-
 #include "UnrealAudioPrivate.h"
 #include "UnrealAudioModule.h"
 #include "UnrealAudioTestGenerators.h"
diff --git a/Engine/Source/Runtime/UnrealAudio/Private/UnrealAudio.cpp b/Engine/Source/Runtime/UnrealAudio/Private/UnrealAudio.cpp
--- a/Engine/Source/Runtime/UnrealAudio/Private/UnrealAudio.cpp
+++ b/Engine/Source/Runtime/UnrealAudio/Private/UnrealAudio.cpp
@@ -1,7 +1,5 @@
 // Copyright 1998-2015 Epic Games, Inc. All Rights Reserved.
 
-#pragma once 
-
 #include "UnrealAudioPrivate.h"
 #include "UnrealAudioModule.h"
 #include "UnrealAudioDeviceModule.h"
diff --git a/Engine/Source/Runtime/UnrealAudio/Private/UnrealAudioBuffer.cpp b/Engine/Source/Runtime/UnrealAudio/Private/UnrealAudioBuffer.cpp
--- a/Engine/Source/Runtime/UnrealAudio/Private/UnrealAudioBuffer.cpp
+++ b/Engine/Source/Runtime/UnrealAudio/Private/UnrealAudioBuffer.cpp
@@ -1,7 +1,5 @@
 // Copyright 1998-2015 Epic Games, Inc. All Rights Reserved.
 
-#pragma once 
-
 #include "UnrealAudioPrivate.h"
 #include "UnrealAudioBuffer.h"
 
diff --git a/Engine/Source/Runtime/UnrealAudio/Private/UnrealAudioDevice.cpp b/Engine/Source/Runtime/UnrealAudio/Private/UnrealAudioDevice.cpp
--- a/Engine/Source/Runtime/UnrealAudio/Private/UnrealAudioDevice.cpp
+++ b/Engine/Source/Runtime/UnrealAudio/Private/UnrealAudioDevice.cpp
@@ -1,7 +1,5 @@
 // Copyright 1998-2015 Epic Games, Inc. All Rights Reserved.
 
-#pragma once 
-
 #include "UnrealAudioPrivate.h"
 #include "UnrealAudioDeviceModule.h"
 
diff --git a/Engine/Source/Runtime/UnrealAudio/Private/UnrealAudioDeviceFormat.cpp b/Engine/Source/Runtime/UnrealAudio/Private/UnrealAudioDeviceFormat.cpp
--- a/Engine/Source/Runtime/UnrealAudio/Private/UnrealAudioDeviceFormat.cpp
+++ b/Engine/Source/Runtime/UnrealAudio/Private/UnrealAudioDeviceFormat.cpp
@@ -1,7 +1,5 @@
 // Copyright 1998-2015 Epic Games, Inc. All Rights Reserved.
 
-#pragma once 
-
 #include "UnrealAudioPrivate.h"
 #include "UnrealAudioDeviceModule.h"
 #include "UnrealAudioDeviceFormat.h"
diff --git a/Engine/Source/Runtime/Windows/UnrealAudioWasapi/Private/UnrealAudioDeviceWasapi.cpp b/Engine/Source/Runtime/Windows/UnrealAudioWasapi/Private/UnrealAudioDeviceWasapi.cpp
--- a/Engine/Source/Runtime/Windows/UnrealAudioWasapi/Private/UnrealAudioDeviceWasapi.cpp
+++ b/Engine/Source/Runtime/Windows/UnrealAudioWasapi/Private/UnrealAudioDeviceWasapi.cpp
@@ -1,7 +1,5 @@
 // Copyright 1998-2015 Epic Games, Inc. All Rights Reserved.
 
-#pragma once 
-
 /**
 	Concrete implementation of IAudioDevice for WASAPI (Windows Audio Session API)
 
diff --git a/Engine/Source/Runtime/Windows/UnrealAudioXAudio2/Private/UnrealAudioDeviceXAudio2.cpp b/Engine/Source/Runtime/Windows/UnrealAudioXAudio2/Private/UnrealAudioDeviceXAudio2.cpp
--- a/Engine/Source/Runtime/Windows/UnrealAudioXAudio2/Private/UnrealAudioDeviceXAudio2.cpp
+++ b/Engine/Source/Runtime/Windows/UnrealAudioXAudio2/Private/UnrealAudioDeviceXAudio2.cpp
@@ -1,7 +1,5 @@
 // Copyright 1998-2015 Epic Games, Inc. All Rights Reserved.
 
-#pragma once 
-
 /**
 	Concrete implementation of IAudioDevice for XAudio2
 
