UnrealAudio

- Changing design to not include audio device input in main user callback since there were issues with patching XAudio2 (which doesn't support device input) and Wasapi capture devices when the device sample rates were not exactly the same. I could fix this but its probably not worth the time at the moment. We'll add mic input back into the audio engine in the future, but likely in a different mechanism, perhaps as an audio generator object.

- Fully implemented unreal audio device module for XAudio2

- Added ability to switch which device module to test from the commandlet, now to optionally test a specific module (i.e. flip between Wasapi and XAudio2 on pc), type:

-command UnrealEd.AudioTestCommandlet XAudio2 device all

or:

-command UnrealEd.AudioTestCommandlet Wasapi device all

The default as of this check in is Wasapi, so you can also do the following to test Wasapi:

-command UnrealEd.AudioTestCommandlet device all

Of course, the "in" flavors of tests have been removed since there is no longer support for device input

diff --git a/Engine/Source/Editor/UnrealEd/Private/Commandlets/AudioTestCommandlet.cpp b/Engine/Source/Editor/UnrealEd/Private/Commandlets/AudioTestCommandlet.cpp
--- a/Engine/Source/Editor/UnrealEd/Private/Commandlets/AudioTestCommandlet.cpp
+++ b/Engine/Source/Editor/UnrealEd/Private/Commandlets/AudioTestCommandlet.cpp
@@ -11,10 +11,14 @@ DEFINE_LOG_CATEGORY_STATIC(AudioTestCommandlet, Log, All);
 /** Unreal audio module ptr */
 static UAudio::IUnrealAudioModule* UnrealAudioModule = nullptr;
 
-static bool UnrealAudioLoad()
+static bool UnrealAudioLoad(const FString* DeviceApi = nullptr)
 {
 	UnrealAudioModule = FModuleManager::LoadModulePtr<UAudio::IUnrealAudioModule>(FName("UnrealAudio"));
-	if (UnrealAudioModule)
+	if (DeviceApi)
+	{
+		UnrealAudioModule->Initialize(*DeviceApi);
+	}
+	else
 	{
 		UnrealAudioModule->Initialize();
 	}
@@ -57,15 +61,6 @@ static bool TestAudioDeviceAll()
 		return false;
 	}
 
-	if (!UAudio::TestDeviceInputSimple(10))
-	{
-		return false;
-	}
-
-	if (!UAudio::TestDeviceInputRandomizedDelay(10))
-	{
-		return false;
-	}
 	return true;
 }
 
@@ -79,11 +74,6 @@ static bool TestAudioDeviceOutputSimple()
 	return UAudio::TestDeviceOutputSimple(-1);
 }
 
-static bool TestAudioDeviceInputSimple()
-{
-	return UAudio::TestDeviceInputSimple(-1);
-}
-
 static bool TestAudioDeviceOutputFm()
 {
 	return UAudio::TestDeviceOutputRandomizedFm(-1);
@@ -94,11 +84,6 @@ static bool TestAudioDeviceOutputPan()
 	return UAudio::TestDeviceOutputNoisePan(-1);
 }
 
-static bool TestAudioDeviceInputDelay()
-{
-	return UAudio::TestDeviceInputRandomizedDelay(-1);
-}
-
 /**
 * Setting up commandlet code
 */
@@ -119,8 +104,6 @@ enum EAudioTests
 	AUDIO_TEST_DEVICE_OUTPUT_SIMPLE,
 	AUDIO_TEST_DEVICE_OUTPUT_FM,
 	AUDIO_TEST_DEVICE_OUTPUT_PAN,
-	AUDIO_TEST_DEVICE_INPUT_SIMPLE,
-	AUDIO_TEST_DEVICE_INPUT_DELAY,
 	AUDIO_TESTS
 };
 
@@ -131,8 +114,6 @@ static AudioTestInfo AudioTestInfoList[] =
 	{ "device", "out",			TestAudioDeviceOutputSimple }, // AUDIO_TEST_DEVICE_OUTPUT_SIMPLE
 	{ "device", "out_fm",		TestAudioDeviceOutputFm		}, // AUDIO_TEST_DEVICE_OUTPUT_FM
 	{ "device", "out_pan",		TestAudioDeviceOutputPan	}, // AUDIO_TEST_DEVICE_OUTPUT_PAN
-	{ "device", "in",			TestAudioDeviceInputSimple	}, // AUDIO_TEST_DEVICE_INPUT_SIMPLE
-	{ "device", "in_delay",		TestAudioDeviceInputDelay	}, // AUDIO_TEST_DEVICE_INPUT_DELAY
 };
 static_assert(ARRAY_COUNT(AudioTestInfoList) == AUDIO_TESTS, "Mismatched info list and test enumeration");
 
@@ -163,16 +144,32 @@ int32 UAudioTestCommandlet::Main(const FString& InParams)
 	TArray<FString> Switches;
 	UCommandlet::ParseCommandLine(*InParams, Tokens, Switches);
 
-	if (Tokens.Num() != 3)
+	if (Tokens.Num() < 3 || Tokens.Num() > 4)
 	{
 		PrintUsage();
 		return 0;
 	}
+	int32 CategoryNameIndex = 1;
+	int32 TestNameIndex = 2;
 
-	if (!UnrealAudioLoad())
+	if (Tokens.Num() == 4)
 	{
-		UE_LOG(AudioTestCommandlet, Display, TEXT("Failed to load unreal audio module. Exiting."));
-		return 0;
+		FString DeviceApiName(Tokens[1]);
+		if (!UnrealAudioLoad(&DeviceApiName))
+		{
+			UE_LOG(AudioTestCommandlet, Display, TEXT("Failed to load unreal audio module. Exiting."));
+			return 0;
+		}
+		CategoryNameIndex++;
+		TestNameIndex++;
+	}
+	else
+	{
+		if (!UnrealAudioLoad())
+		{
+			UE_LOG(AudioTestCommandlet, Display, TEXT("Failed to load unreal audio module. Exiting."));
+			return 0;
+		}
 	}
 
 	check(UnrealAudioModule != nullptr);
@@ -180,9 +177,9 @@ int32 UAudioTestCommandlet::Main(const FString& InParams)
 	bool FoundTest = false;
 	for (uint32 Index = 0; Index < AUDIO_TESTS; ++Index)
 	{
-		if (AudioTestInfoList[Index].CategoryName == Tokens[1])
+		if (AudioTestInfoList[Index].CategoryName == Tokens[CategoryNameIndex])
 		{
-			if (AudioTestInfoList[Index].TestName == Tokens[2])
+			if (AudioTestInfoList[Index].TestName == Tokens[TestNameIndex])
 			{
 				FoundTest = true;
 				if (AudioTestInfoList[Index].TestFunction())
diff --git a/Engine/Source/Runtime/UnrealAudio/Private/Tests/UnrealAudioDeviceTests.cpp b/Engine/Source/Runtime/UnrealAudio/Private/Tests/UnrealAudioDeviceTests.cpp
--- a/Engine/Source/Runtime/UnrealAudio/Private/Tests/UnrealAudioDeviceTests.cpp
+++ b/Engine/Source/Runtime/UnrealAudio/Private/Tests/UnrealAudioDeviceTests.cpp
@@ -86,35 +86,6 @@ namespace UAudio
 			{
 				UE_LOG(LogUnrealAudio, Display, TEXT("    %s"), ESpeaker::ToString(DeviceInfo.Speakers[Channel]));
 			}
-
-
-		}
-
-		UE_LOG(LogUnrealAudio, Display, TEXT("Querying the number of output devices for current system..."), );
-		uint32 NumInputDevices = 0;
-		if (!UnrealAudioDevice->GetNumInputDevices(NumInputDevices))
-		{
-			UE_LOG(LogUnrealAudio, Display, TEXT("Failed."));
-			return false;
-		}
-		UE_LOG(LogUnrealAudio, Display, TEXT("Success: %d Number of Input Devices"), NumInputDevices);
-
-		UE_LOG(LogUnrealAudio, Display, TEXT("Retrieving input device info for each input device..."), );
-		for (uint32 DeviceIndex = 0; DeviceIndex < NumInputDevices; ++DeviceIndex)
-		{
-			FDeviceInfo DeviceInfo;
-			if (!UnrealAudioDevice->GetInputDeviceInfo(DeviceIndex, DeviceInfo))
-			{
-				UE_LOG(LogUnrealAudio, Display, TEXT("Failed for index %d."), DeviceIndex);
-				return false;
-			}
-
-			UE_LOG(LogUnrealAudio, Display, TEXT("Device Index: %d"), DeviceIndex);
-			UE_LOG(LogUnrealAudio, Display, TEXT("Device Name: %s"), *DeviceInfo.FriendlyName);
-			UE_LOG(LogUnrealAudio, Display, TEXT("Device FrameRate: %d"), DeviceInfo.FrameRate);
-			UE_LOG(LogUnrealAudio, Display, TEXT("Device NumChannels: %d"), DeviceInfo.NumChannels);
-			UE_LOG(LogUnrealAudio, Display, TEXT("Device Default?: %s"), DeviceInfo.bIsSystemDefault ? TEXT("yes") : TEXT("no"));
-			UE_LOG(LogUnrealAudio, Display, TEXT("Device Native Format: %s"), EStreamFormat::ToString(DeviceInfo.StreamFormat));
 		}
 
 		return true;
@@ -239,122 +210,12 @@ namespace UAudio
 		return true;
 	}
 
-	static bool DoInputTest(const FString& TestName, double LifeTime, Test::IGenerator* Generator)
-	{
-		check(Generator);
-
-		if (!bTestInitialized)
-		{
-			return false;
-		}
-
-		UE_LOG(LogUnrealAudio, Display, TEXT("Running audio test: \"%s\"..."), *TestName);
-
-		if (LifeTime > 0.0)
-		{
-			UE_LOG(LogUnrealAudio, Display, TEXT("Time of test: %d (seconds)"), LifeTime);
-		}
-		else
-		{
-			UE_LOG(LogUnrealAudio, Display, TEXT("Time of test: (indefinite)"));
-		}
-
-		IUnrealAudioDeviceModule* UnrealAudioDevice = UnrealAudioModule->GetDeviceModule();
-
-		EDeviceApi::Type DeviceApi;
-		if (!UnrealAudioDevice->GetDevicePlatformApi(DeviceApi))
-		{
-			UE_LOG(LogUnrealAudio, Display, TEXT("Failed."));
-			return false;
-		}
-		if (DeviceApi == EDeviceApi::DUMMY)
-		{
-			UE_LOG(LogUnrealAudio, Display, TEXT("This is a dummy API. Platform not implemented yet."));
-			return true;
-		}
-
-		uint32 DefaultOutputDeviceIndex = 0;
-		if (!UnrealAudioDevice->GetDefaultOutputDeviceIndex(DefaultOutputDeviceIndex))
-		{
-			UE_LOG(LogUnrealAudio, Display, TEXT("Failed to get default output device index."));
-			return false;
-		}
-
-		FDeviceInfo OutputDeviceInfo;
-		if (!UnrealAudioDevice->GetOutputDeviceInfo(DefaultOutputDeviceIndex, OutputDeviceInfo))
-		{
-			UE_LOG(LogUnrealAudio, Display, TEXT("Failed to get default output device info."));
-			return false;
-		}
-
-		uint32 DefaultInputDeviceIndex = 0;
-		if (!UnrealAudioDevice->GetDefaultInputDeviceIndex(DefaultInputDeviceIndex))
-		{
-			UE_LOG(LogUnrealAudio, Display, TEXT("Failed to get default input device index."));
-			return false;
-		}
-
-		FDeviceInfo InputDeviceInfo;
-		if (!UnrealAudioDevice->GetInputDeviceInfo(DefaultInputDeviceIndex, InputDeviceInfo))
-		{
-			UE_LOG(LogUnrealAudio, Display, TEXT("Failed to get default input device info."));
-			return false;
-		}
-
-		FCreateStreamParams CreateStreamParams;
-		CreateStreamParams.InputDeviceIndex = DefaultInputDeviceIndex;
-		CreateStreamParams.OutputDeviceIndex = DefaultOutputDeviceIndex;
-		CreateStreamParams.CallbackFunction = &TestCallback;
-		CreateStreamParams.UserData = (void*)Generator;
-
-		UE_LOG(LogUnrealAudio, Display, TEXT("Creating the stream..."));
-		if (!UnrealAudioDevice->CreateStream(CreateStreamParams))
-		{
-			UE_LOG(LogUnrealAudio, Display, TEXT("Failed to create a stream."));
-			return false;
-		}
-
-		UE_LOG(LogUnrealAudio, Display, TEXT("Starting the stream."));
-		UnrealAudioDevice->StartStream();
-
-		// Block until the synthesizer is done
-		while (!Generator->IsDone())
-		{
-			Sleep(1);
-		}
-
-		// Stop the output stream (which blocks until its actually freed)
-		UE_LOG(LogUnrealAudio, Display, TEXT("Shutting down the stream."));
-		if (!UnrealAudioDevice->StopStream())
-		{
-			UE_LOG(LogUnrealAudio, Display, TEXT("Failed to stop the device stream."));
-			return false;
-		}
-
-		if (!UnrealAudioDevice->ShutdownStream())
-		{
-			UE_LOG(LogUnrealAudio, Display, TEXT("Failed to shut down the device stream."));
-			return false;
-		}
-
-		// At this point audio device I/O is done
-		UE_LOG(LogUnrealAudio, Display, TEXT("Success!"));
-
-		return true;
-	}
-
 	bool TestDeviceOutputSimple(double LifeTime)
 	{
 		Test::FSimpleOutput SimpleOutput(LifeTime);
 		return DoOutputTest("output randomized FM synthesis", LifeTime, &SimpleOutput);
 	}
 
-	bool TestDeviceInputSimple(double LifeTime)
-	{
-		Test::FPassThrough SimplePassThrough(LifeTime);
-		return DoInputTest("input pass through", LifeTime, &SimplePassThrough);
-	}
-
 	bool TestDeviceOutputRandomizedFm(double LifeTime)
 	{
 		Test::FPhaseModulator RandomizedFmGenerator(LifeTime);
@@ -367,12 +228,6 @@ namespace UAudio
 		return DoOutputTest("output white noise panner", LifeTime, &SimpleWhiteNoisePan);
 	}
 
-	bool UNREALAUDIO_API TestDeviceInputRandomizedDelay(double LifeTime)
-	{
-		Test::FInputDelay SimpleDelayer(LifeTime);
-		return DoInputTest("input randomized delay", LifeTime, &SimpleDelayer);
-	}
-
 }
 
 #endif // #if ENABLE_UNREAL_AUDIO
diff --git a/Engine/Source/Runtime/UnrealAudio/Private/Tests/UnrealAudioTestGenerators.cpp b/Engine/Source/Runtime/UnrealAudio/Private/Tests/UnrealAudioTestGenerators.cpp
--- a/Engine/Source/Runtime/UnrealAudio/Private/Tests/UnrealAudioTestGenerators.cpp
+++ b/Engine/Source/Runtime/UnrealAudio/Private/Tests/UnrealAudioTestGenerators.cpp
@@ -954,105 +954,6 @@ namespace Test
 		Panner.Spatialize(Value, Frame);
 	}
 
-	/************************************************************************/
-	/* FPassThrough															*/
-	/************************************************************************/
-
-	FPassThrough::FPassThrough(double LifeTime)
-		: IGenerator(LifeTime)
-	{
-	}
-
-	FPassThrough::~FPassThrough()
-	{
-	}
-
-	bool FPassThrough::GetNextBuffer(FCallbackInfo& CallbackInfo)
-	{
-		if (!UpdateTimers())
-		{
-			return true;
-		}
-
-		// Passthrough
-		if (CallbackInfo.NumInputChannels > 0)
-		{
-			uint32 BufferIndex = 0;
-			uint32 InputBufferIndex = 0;
-			for (uint32 Frame = 0; Frame < CallbackInfo.NumFrames; ++Frame, BufferIndex += CallbackInfo.NumChannels, InputBufferIndex += CallbackInfo.NumInputChannels)
-			{
-				float InSample = CallbackInfo.InBuffer[InputBufferIndex];
-				InSample /= (float) CallbackInfo.NumChannels;
-				for (int32 Chan = 0; Chan < CallbackInfo.NumChannels; ++Chan)
-				{
-					CallbackInfo.OutBuffer[BufferIndex + Chan] = InSample;
-				}
-			}
-		}
-		return true;
-	}
-
-	/************************************************************************/
-	/* FInputDelay															*/
-	/************************************************************************/
-
-	FInputDelay::FInputDelay(double LifeTime)
-		: IGenerator(LifeTime)
-		, bInit(false)
-	{
-	}
-
-	FInputDelay::~FInputDelay()
-	{
-	}
-
-	void FInputDelay::Init()
-	{
-		if (bInit)
-		{
-			return;
-		}
-		bInit = true;
-		float MaxLength = 2.0f;
-		int NumTaps = CallbackData.NumChannels;
-		Delay.Init(NumTaps, MaxLength);
-		DelayOutput.Init(0.0f, NumTaps);
-		for (int32 Tap = 0; Tap < NumTaps; ++Tap)
-		{
-			Delay.SetDelayLength(Tap, FMath::FRandRange(0.2f, MaxLength));
-			Delay.SetFeedback(Tap, 0.9f);
-			Delay.SetWet(Tap, FMath::FRandRange(0.7f, 0.99f));
-		}
-	}
-
-	bool FInputDelay::GetNextBuffer(FCallbackInfo& CallbackInfo)
-	{
-		if (!UpdateTimers())
-		{
-			return true;
-		}
-
-		if (!bInit)
-		{
-			Init();
-		}
-
-		if (CallbackInfo.NumInputChannels > 0)
-		{
-			uint32 BufferIndex = 0;
-			uint32 InputBufferIndex = 0;
-			for (uint32 Frame = 0; Frame < CallbackInfo.NumFrames; ++Frame, BufferIndex += CallbackInfo.NumChannels, InputBufferIndex += CallbackInfo.NumInputChannels)
-			{
-				Delay.GetOutput(CallbackInfo.InBuffer[InputBufferIndex], DelayOutput);
-				for (int32 Chan = 0; Chan < CallbackInfo.NumChannels; ++Chan)
-				{
-					CallbackInfo.OutBuffer[BufferIndex + Chan] = DelayOutput[Chan];
-				}
-			}
-		}
-		return true;
-	}
-
 	/************************************************************************/
 	/* FNoisePan															*/
 	/************************************************************************/
diff --git a/Engine/Source/Runtime/UnrealAudio/Private/Tests/UnrealAudioTestGenerators.h b/Engine/Source/Runtime/UnrealAudio/Private/Tests/UnrealAudioTestGenerators.h
--- a/Engine/Source/Runtime/UnrealAudio/Private/Tests/UnrealAudioTestGenerators.h
+++ b/Engine/Source/Runtime/UnrealAudio/Private/Tests/UnrealAudioTestGenerators.h
@@ -375,39 +375,6 @@ namespace Test
 		int32 CurrNumSynthesisers;
 	};
 
-	/**
-	* FPassThrough
-	* A simple input pass-through to output (mic->speakers). Should hear a bit of delay from input to output due to device latency.
-	*/
-	class FPassThrough : public IGenerator
-	{
-	public:
-		FPassThrough(double LifeTime);
-		~FPassThrough();
-
-		bool GetNextBuffer(FCallbackInfo& CallbackInfo) override;
-	};
-
-	/**
-	* FSimpleRandomizedDelayer
-	* An simple randomized delay line that delays input audio stream by various amount to great effect. 
-	*/
-	class FInputDelay : public IGenerator
-	{
-	public:
-		FInputDelay(double LifeTime);
-		~FInputDelay();
-
-		bool GetNextBuffer(FCallbackInfo& CallbackInfo) override;
-
-	private:
-		void Init();
-
-		bool bInit;
-		FDelay Delay;
-		TArray<float> DelayOutput;
-	};
-
 	/**
 	* FNoisePan
 	* An simple test which generates noise and pans it clockwise
diff --git a/Engine/Source/Runtime/UnrealAudio/Private/UnrealAudio.cpp b/Engine/Source/Runtime/UnrealAudio/Private/UnrealAudio.cpp
--- a/Engine/Source/Runtime/UnrealAudio/Private/UnrealAudio.cpp
+++ b/Engine/Source/Runtime/UnrealAudio/Private/UnrealAudio.cpp
@@ -16,9 +16,6 @@ IMPLEMENT_MODULE(UAudio::FUnrealAudioModule, UnrealAudio);
 
 namespace UAudio
 {
-
-
-
 	FUnrealAudioModule::FUnrealAudioModule()
 		: UnrealAudioDevice(nullptr)
 	{
@@ -28,12 +25,32 @@ namespace UAudio
 	{
 	}
 
-	void FUnrealAudioModule::Initialize()
+	bool FUnrealAudioModule::Initialize()
+	{
+		ModuleName = GetDefaultDeviceModuleName();
+		return InitializeInternal();
+	}
+
+	bool FUnrealAudioModule::Initialize(const FString& DeviceModuleName)
 	{
-		UnrealAudioDevice = FModuleManager::LoadModulePtr<IUnrealAudioDeviceModule>(GetDeviceModuleName());
+		FString Name("UnrealAudio");
+		Name += DeviceModuleName;
+		ModuleName = *Name;
+		return InitializeInternal();
+	}
+
+	bool FUnrealAudioModule::InitializeInternal()
+	{
+		bool bSuccess = false;
+		UnrealAudioDevice = FModuleManager::LoadModulePtr<IUnrealAudioDeviceModule>(ModuleName);
 		if (UnrealAudioDevice)
 		{
-			if (!UnrealAudioDevice->Initialize())
+			if (UnrealAudioDevice->Initialize())
+			{
+				bSuccess = true;
+				InitializeTests(this);
+			}
+			else
 			{
 				UE_LOG(LogUnrealAudio, Error, TEXT("Failed to initialize audio device module."));
 				UnrealAudioDevice = nullptr;
@@ -43,8 +60,7 @@ namespace UAudio
 		{
 			UnrealAudioDevice = CreateDummyDeviceModule();
 		}
-
-		InitializeTests(this);
+		return bSuccess;
 	}
 
 	void FUnrealAudioModule::Shutdown()
@@ -63,7 +79,7 @@ namespace UAudio
 	}
 
 
-	FName FUnrealAudioModule::GetDeviceModuleName() const
+	FName FUnrealAudioModule::GetDefaultDeviceModuleName() const
 	{
 		// TODO: Pull the device module name to use from an engine ini or an audio ini file
 		return FName(AUDIO_DEVICE_DEVICE_MODULE_NAME);
diff --git a/Engine/Source/Runtime/UnrealAudio/Private/UnrealAudioBuffer.cpp b/Engine/Source/Runtime/UnrealAudio/Private/UnrealAudioBuffer.cpp
--- a/Engine/Source/Runtime/UnrealAudio/Private/UnrealAudioBuffer.cpp
+++ b/Engine/Source/Runtime/UnrealAudio/Private/UnrealAudioBuffer.cpp
@@ -24,7 +24,7 @@ namespace UAudio
 			}
 		}
 
-		void Initialize(uint32 InNumSamples, uint32 SampleBytes) override
+		void Initialize(uint32 InNumSamples) override
 		{
 			DataBuffer = new T[InNumSamples];
 			memset((void*)DataBuffer, 0, InNumSamples * sizeof(T));
diff --git a/Engine/Source/Runtime/UnrealAudio/Private/UnrealAudioBuffer.h b/Engine/Source/Runtime/UnrealAudio/Private/UnrealAudioBuffer.h
--- a/Engine/Source/Runtime/UnrealAudio/Private/UnrealAudioBuffer.h
+++ b/Engine/Source/Runtime/UnrealAudio/Private/UnrealAudioBuffer.h
@@ -21,7 +21,7 @@ namespace UAudio
 
 		virtual ~IIntermediateBuffer() {}
 
-		virtual void Initialize(uint32 InNumSamples, uint32 SampleBytes) = 0;
+		virtual void Initialize(uint32 InNumSamples) = 0;
 		virtual bool Write(uint8* Buffer, uint32 Size) = 0;
 		virtual bool Read(uint8* Buffer, uint32 Size) = 0;
 
diff --git a/Engine/Source/Runtime/UnrealAudio/Private/UnrealAudioDevice.cpp b/Engine/Source/Runtime/UnrealAudio/Private/UnrealAudioDevice.cpp
--- a/Engine/Source/Runtime/UnrealAudio/Private/UnrealAudioDevice.cpp
+++ b/Engine/Source/Runtime/UnrealAudio/Private/UnrealAudioDevice.cpp
@@ -30,7 +30,6 @@ namespace UAudio
 	const uint32 IUnrealAudioDeviceModule::MaxPossibleFrameRates = 11;
 
 	IUnrealAudioDeviceModule::IUnrealAudioDeviceModule()
-		: DeviceErrorListener(nullptr)
 	{
 	}
 
@@ -47,25 +46,13 @@ namespace UAudio
 
 		Reset();
 
-		if (OpenDevices(Params))
+		if (OpenDevice(Params))
 		{
-			if (Params.InputDeviceIndex != INDEX_NONE)
-			{
-				// This stream also streams audio in (every audio stream streams output)
-				StreamInfo.StreamType = EStreamType::INPUT;
-			}
-			else
-			{
-				// This stream is only streaming audio out
-				StreamInfo.StreamType = EStreamType::OUTPUT;
-			}
-
 			StreamInfo.State = EStreamState::STOPPED;
 			StreamInfo.CallbackFunction = Params.CallbackFunction;
 			StreamInfo.UserData = Params.UserData;
 			StreamInfo.BlockSize = Params.CallbackBlockSize;
 			StreamInfo.StreamDelta = (double)StreamInfo.BlockSize / StreamInfo.FrameRate;
-
 			bSuccess = true;
 		}
 
@@ -80,12 +67,6 @@ namespace UAudio
 	void IUnrealAudioDeviceModule::OnDeviceError(const EDeviceError::Type ErrorType, const FString& ErrorDetails, const FString& FileName, int32 LineNumber) const
 	{
 		UE_LOG(LogUnrealAudioDevice, Error, TEXT("Audio Device Error: (%s) : %s (%s::%d)"), EDeviceError::ToString(ErrorType), *ErrorDetails, *FileName, LineNumber);
-
-		// And broadcast to listener if one is setup
-		if (DeviceErrorListener)
-		{
-			DeviceErrorListener->OnDeviceError(ErrorType, ErrorDetails, FileName, LineNumber);
-		}
 	}
 
 	void IUnrealAudioDeviceModule::Reset()
@@ -138,30 +119,12 @@ namespace UAudio
 			return true;
 		}
 
-		bool GetNumInputDevices(uint32& OutNumDevices) const override
-		{
-			OutNumDevices = 0;
-			return true;
-		}
-
-		bool GetInputDeviceInfo(const uint32 DeviceIndex, FDeviceInfo& OutInfo) const override
-		{
-			memset((void*)&OutInfo, 0, sizeof(FDeviceInfo));
-			return true;
-		}
-
 		bool GetDefaultOutputDeviceIndex(uint32& OutDefaultIndex) const override
 		{
 			OutDefaultIndex = 0;
 			return true;
 		}
 
-		bool GetDefaultInputDeviceIndex(uint32& OutDefaultIndex) const override
-		{
-			OutDefaultIndex = 0;
-			return true;
-		}
-
 		bool StartStream() override
 		{
 			return true;
@@ -177,10 +140,9 @@ namespace UAudio
 			return true;
 		}
 
-		bool GetLatency(uint32& OutputDeviceLatency, uint32& InputDeviceLatency) const override
+		bool GetLatency(uint32& OutputDeviceLatency) const override
 		{
 			OutputDeviceLatency = 0;
-			InputDeviceLatency = 0;
 			return true;
 		}
 
@@ -190,7 +152,7 @@ namespace UAudio
 			return true;
 		}
 
-		bool OpenDevices(const FCreateStreamParams& Params) override
+		bool OpenDevice(const FCreateStreamParams& Params) override
 		{
 			return true;
 		}
diff --git a/Engine/Source/Runtime/UnrealAudio/Private/UnrealAudioDeviceFormat.cpp b/Engine/Source/Runtime/UnrealAudio/Private/UnrealAudioDeviceFormat.cpp
--- a/Engine/Source/Runtime/UnrealAudio/Private/UnrealAudioDeviceFormat.cpp
+++ b/Engine/Source/Runtime/UnrealAudio/Private/UnrealAudioDeviceFormat.cpp
@@ -27,31 +27,18 @@ namespace UAudio
 		}
 	}
 
-	void IUnrealAudioDeviceModule::SetupBufferFormatConvertInfo(EStreamType::Type StreamType)
+	void IUnrealAudioDeviceModule::SetupBufferFormatConvertInfo()
 	{
-		FBufferFormatConvertInfo& ConvertInfo = StreamInfo.GetConvertInfo(StreamType);
-		FStreamDeviceInfo& DeviceInfo = StreamInfo.GetDeviceInfo(StreamType);
+		FBufferFormatConvertInfo& ConvertInfo = StreamInfo.DeviceInfo.BufferFormatConvertInfo;
 
-		if (StreamType == EStreamType::INPUT)
-		{
-			// Setup the convert info to convert from device channels and format to user channels and format
-			ConvertInfo.FromChannels = DeviceInfo.NumChannels;
-			ConvertInfo.ToChannels = DeviceInfo.NumChannels;
-			ConvertInfo.FromFormat = DeviceInfo.DeviceDataFormat;
-			ConvertInfo.ToFormat = EStreamFormat::FLT;
-		}
-		else
-		{
-			// Setup the convert info to convert from user channels and format to device channels and format
-			ConvertInfo.FromChannels = DeviceInfo.NumChannels;
-			ConvertInfo.ToChannels = DeviceInfo.NumChannels;
-			ConvertInfo.FromFormat = EStreamFormat::FLT;
-			ConvertInfo.ToFormat = DeviceInfo.DeviceDataFormat;
-		}
+		// Setup the convert info to convert from user channels and format to device channels and format
+		ConvertInfo.FromChannels = StreamInfo.DeviceInfo.NumChannels;
+		ConvertInfo.ToChannels = StreamInfo.DeviceInfo.NumChannels;
+		ConvertInfo.FromFormat = EStreamFormat::FLT;
+		ConvertInfo.ToFormat = StreamInfo.DeviceInfo.DeviceDataFormat;
 
 		// For either stream types, set the number of channels to use for the conversion to be the smallest number of channels
 		ConvertInfo.NumChannels = FMath::Min(ConvertInfo.FromChannels, ConvertInfo.ToChannels);
-
 	}
 
 	template <typename FromType, typename ToType>
@@ -196,9 +183,9 @@ namespace UAudio
 		}
 	}
 
-	bool IUnrealAudioDeviceModule::ConvertBufferFormat(TArray<uint8>& ToBuffer, TArray<uint8>& FromBuffer, EStreamType::Type StreamType)
+	bool IUnrealAudioDeviceModule::ConvertBufferFormat(TArray<uint8>& ToBuffer, TArray<uint8>& FromBuffer)
 	{
-		FBufferFormatConvertInfo& ConvertInfo = StreamInfo.GetConvertInfo(StreamType);
+		FBufferFormatConvertInfo& ConvertInfo = StreamInfo.DeviceInfo.BufferFormatConvertInfo;
 
 		switch (ConvertInfo.ToFormat)
 		{
diff --git a/Engine/Source/Runtime/UnrealAudio/Private/UnrealAudioPrivate.h b/Engine/Source/Runtime/UnrealAudio/Private/UnrealAudioPrivate.h
--- a/Engine/Source/Runtime/UnrealAudio/Private/UnrealAudioPrivate.h
+++ b/Engine/Source/Runtime/UnrealAudio/Private/UnrealAudioPrivate.h
@@ -19,18 +19,21 @@ namespace UAudio
 		FUnrealAudioModule();
 		~FUnrealAudioModule();
 
-		void Initialize() override;
+		bool Initialize() override;
+		bool Initialize(const FString& DeviceModuleName) override;
 		void Shutdown() override;
 
 		class IUnrealAudioDeviceModule* GetDeviceModule();
 
-		FName GetDeviceModuleName() const;
+		FName GetDefaultDeviceModuleName() const;
 
 	private:
 
+		bool InitializeInternal();
 		static void InitializeTests(FUnrealAudioModule* Module);
 
 		class IUnrealAudioDeviceModule* UnrealAudioDevice;
+		FName ModuleName;
 	};
 }
 
diff --git a/Engine/Source/Runtime/UnrealAudio/Public/UnrealAudioDeviceModule.h b/Engine/Source/Runtime/UnrealAudio/Public/UnrealAudioDeviceModule.h
--- a/Engine/Source/Runtime/UnrealAudio/Public/UnrealAudioDeviceModule.h
+++ b/Engine/Source/Runtime/UnrealAudio/Public/UnrealAudioDeviceModule.h
@@ -102,21 +102,6 @@ namespace UAudio
 		};
 	}
 
-	/**
-	* EDeviceDataFlow
-	* An enumeration used to specify the flow of audio on the device
-	*/
-	namespace EStreamType
-	{
-		enum Type
-		{
-			INPUT,			/** The audio device streams audio into the application (e.g. microphone) */
-			OUTPUT,			/** The audio device only streams audio out of the application */
-			COUNT,			/** The count of stream types. */
-			UNKNOWN,		/** Unknown stream type. */
-		};
-	}
-
 	/**
 	* EStreamStatus
 	* An enumeration to specify the flow status of the audio stream.
@@ -125,7 +110,6 @@ namespace UAudio
 	{
 		enum Flag
 		{
-			INPUT_OVERFLOW		= (1 << 0),		/** The input stream is has too much data */
 			OUTPUT_UNDERFLOW	= (1 << 1)		/** The output stream has too little data */
 		};
 	}
@@ -203,18 +187,12 @@ namespace UAudio
 		/** The number of buffer output frames */
 		uint32 NumFrames;
 
-		/** A pointer to a buffer of float data which is passed in from device input. */
-		float* InBuffer;
-
 		/** Array of outputs speakers */
 		TArray<ESpeaker::Type> OutputSpeakers;
 
 		/** The number of channels in output. */
 		int32 NumChannels;
 
-		/** The number of input channels in input stream. */
-		int32 NumInputChannels;
-
 		/** The current status flags of the input and output buffers */
 		StreamStatus StatusFlags;
 
@@ -254,9 +232,6 @@ namespace UAudio
 		/** The data format of the device (e.g. float). */
 		EStreamFormat::Type StreamFormat;
 
-		/** What type of device it is (input, output, or both). */
-		EStreamType::Type StreamType;
-
 		/** What speakers this device supports (if output device) */
 		TArray<ESpeaker::Type> Speakers;
 
@@ -268,7 +243,6 @@ namespace UAudio
 			, NumChannels(0)
 			, FrameRate(0)
 			, StreamFormat(EStreamFormat::UNKNOWN)
-			, StreamType(EStreamType::UNKNOWN)
 			, bIsSystemDefault(false)
 		{}
 	};
@@ -279,9 +253,6 @@ namespace UAudio
 	*/
 	struct FCreateStreamParams
 	{
-		/** The index of the device to use for audio input. Initialized to INDEX_NONE, which means no device input. */
-		uint32 InputDeviceIndex;
-
 		/** The index of the device to use for audio output. Must be defined. */
 		uint32 OutputDeviceIndex;
 
@@ -291,14 +262,18 @@ namespace UAudio
 		/** The function pointer of a user callback function to generate audio samples to output device. */
 		StreamCallback CallbackFunction;
 
+		/** The preferred frame rate of the audio stream (this may be platform dependent) */
+		uint32 FrameRate;
+
 		/** Pointer to store user data. Passed in callback function. */
 		void* UserData;
 
 		/** Constructor */
 		FCreateStreamParams()
-			: InputDeviceIndex(INDEX_NONE)
-			, OutputDeviceIndex(INDEX_NONE)
+			: OutputDeviceIndex(INDEX_NONE)
 			, CallbackBlockSize(512)
+			, CallbackFunction(nullptr)
+			, FrameRate(48000)
 			, UserData(nullptr)
 		{}
 	};
@@ -362,21 +337,6 @@ namespace UAudio
 		*/
 		virtual bool GetOutputDeviceInfo(const uint32 DeviceIndex, FDeviceInfo& OutInfo) const = 0;
 
-		/** 
-		* Returns the number of input devices. 
-		* @param OutNumDevices The number of input devices.
-		* @return true if succeeded.
-		*/
-		virtual bool GetNumInputDevices(uint32& OutNumDevices) const = 0;
-
-		/** 
-		* Returns information about the input device at the given index. 
-		* @param DeviceIndex The index of the input device.
-		* @param OutInfo Information struct about the device.
-		* @return true if succeeded.
-		*/
-		virtual bool GetInputDeviceInfo(const uint32 DeviceIndex, FDeviceInfo& OutInfo) const = 0;
-
 		/**
 		* Returns the default output device index.
 		* @param OutDefaultIndex The default device index.
@@ -385,14 +345,6 @@ namespace UAudio
 		*/
 		virtual bool GetDefaultOutputDeviceIndex(uint32& OutDefaultIndex) const = 0;
 
-		/**
-		* Returns the default input device index.
-		* @param OutDefaultIndex The default device index.
-		* @return true if succeeded.
-		* @note The default device is specified by the OS default device input.
-		*/
-		virtual bool GetDefaultInputDeviceIndex(uint32& OutDefaultIndex) const = 0;
-
 		/**
 		* Starts the device audio stream.
 		* @return true if succeeded.
@@ -422,7 +374,7 @@ namespace UAudio
 		* @note The device stream needs to have been started for this return anything. A return of 0 means
 		* the device hasn't started or that device hasn't been initialized.
 		*/
-		virtual bool GetLatency(uint32& OutputDeviceLatency, uint32& InputDeviceLatency) const = 0;
+		virtual bool GetLatency(uint32& OutputDeviceLatency) const = 0;
 
 		/**
 		* Returns the frame rate of the audio devices.
@@ -546,9 +498,6 @@ namespace UAudio
 			/** The overall framerate of the stream. This may be different from device frame rate. */
 			uint32 FrameRate;
 
-			/** The type of stream this is (input or output) */
-			EStreamType::Type StreamType;
-
 			/** The current state of the stream. */
 			EStreamState::Type State;
 
@@ -576,13 +525,13 @@ namespace UAudio
 			/** A byte array used to store data to and from audio devices. */
 			TArray<uint8> DeviceBuffer;
 
-			/** An array of device-specific information for input and output devices. */
-			FStreamDeviceInfo DeviceInfo[EStreamType::COUNT];
+			/** device-specific information for output device. */
+			FStreamDeviceInfo DeviceInfo;
 
 			/** Constructor */
 			FStreamInfo()
 			{
-				// Just call intialize since it will get called when starting up a new stream.
+				// Just call initialize since it will get called when starting up a new stream.
 				Initialize();
 			}
 
@@ -590,7 +539,6 @@ namespace UAudio
 			void Initialize()
 			{
 				FrameRate = 0;
-				StreamType = EStreamType::COUNT;
 				State = EStreamState::CLOSED;
 				StreamTime = 0.0;
 				StreamDelta = 0.0;
@@ -599,39 +547,14 @@ namespace UAudio
 				Thread = nullptr;
 				UserData = nullptr;
 				BlockSize = 0;
-
-				for (uint32 Id = 0; Id < EStreamType::COUNT; ++Id)
-				{
-					DeviceInfo[Id].Initialize();
-				}
-			}
-
-			/** Helper function to get specific convert information based on stream type */
-			FBufferFormatConvertInfo& GetConvertInfo(EStreamType::Type StreamType)
-			{
-				check(StreamType < EStreamType::COUNT);
-				return DeviceInfo[StreamType].BufferFormatConvertInfo;
-			}
-
-			/** Helper function to get device info based on stream type */
-			FStreamDeviceInfo& GetDeviceInfo(EStreamType::Type StreamType)
-			{
-				check(StreamType < EStreamType::COUNT);
-				return DeviceInfo[StreamType];
-			}
-
-			/** Helper function to get device info based on stream type const overload */
-			const FStreamDeviceInfo& GetDeviceInfo(EStreamType::Type StreamType) const
-			{
-				check(StreamType < EStreamType::COUNT);
-				return DeviceInfo[StreamType];
+				DeviceInfo.Initialize();
 			}
 		};
 
 	protected: // Protected Methods
 
 		/** Opens audio devices given the input params. Implemented per platform */
-		virtual bool OpenDevices(const FCreateStreamParams& Params) = 0;
+		virtual bool OpenDevice(const FCreateStreamParams& Params) = 0;
 
 		/** Called before opening up new streams. */
 		void Reset();
@@ -640,10 +563,10 @@ namespace UAudio
 		uint32 GetNumBytesForFormat(EStreamFormat::Type Format) const;
 
 		/** Sets up any convert information for given stream type (figures out to/from convert format and channel formats) */
-		void SetupBufferFormatConvertInfo(EStreamType::Type StreamType);
+		void SetupBufferFormatConvertInfo();
 
 		/** Performs actual buffer format and channel conversion */
-		bool ConvertBufferFormat(TArray<uint8>& OutputBuffer, TArray<uint8>& InputBuffer, EStreamType::Type StreamType);
+		bool ConvertBufferFormat(TArray<uint8>& OutputBuffer, TArray<uint8>& InputBuffer);
 
 		/** Updates the sample-accurate stream time value. */
 		void UpdateStreamTimeTick();
@@ -674,9 +597,6 @@ namespace UAudio
 		static const uint32 PossibleFrameRates[];
 		static const uint32 MaxPossibleFrameRates;
 
-		/** An object which listens to device errors. */
-		IDeviceErrorListener* DeviceErrorListener;
-
 		/** Device stream info struct */
 		FStreamInfo StreamInfo;
 	};
diff --git a/Engine/Source/Runtime/UnrealAudio/Public/UnrealAudioModule.h b/Engine/Source/Runtime/UnrealAudio/Public/UnrealAudioModule.h
--- a/Engine/Source/Runtime/UnrealAudio/Public/UnrealAudioModule.h
+++ b/Engine/Source/Runtime/UnrealAudio/Public/UnrealAudioModule.h
@@ -18,8 +18,14 @@ namespace UAudio
 	public:
 		virtual ~IUnrealAudioModule() {}
 
-		virtual void Initialize()
+		virtual bool Initialize()
 		{
+			return false;
+		}
+
+		virtual bool Initialize(const FString& DeviceModuleName)
+		{
+			return false;
 		}
 
 		virtual void Shutdown()
diff --git a/Engine/Source/Runtime/UnrealAudio/Public/UnrealAudioTests.h b/Engine/Source/Runtime/UnrealAudio/Public/UnrealAudioTests.h
--- a/Engine/Source/Runtime/UnrealAudio/Public/UnrealAudioTests.h
+++ b/Engine/Source/Runtime/UnrealAudio/Public/UnrealAudioTests.h
@@ -23,17 +23,7 @@ namespace UAudio
 	* @return True if test succeeds.
 	*/
 	bool UNREALAUDIO_API TestDeviceOutputSimple(double TestTime);
-	
-	/**
-	* TestDeviceInputSimple
-	*
-	* Test which test input device functionality by feeding input device data to output device data (pass through).
-	* If you have a connected microphone on your system (set to default) you should hear the mic input feed through output speakers.
-	* @param TestTime The amount of time to run the test in seconds (negative value means to run indefinitely)
-	* @return True if test succeeds.
-	*/
-	bool UNREALAUDIO_API TestDeviceInputSimple(double TestTime);
-	
+
 	/**
 	* TestDeviceOutputRandomizedFm
 	*
@@ -53,14 +43,5 @@ namespace UAudio
 	*/
 	bool UNREALAUDIO_API TestDeviceOutputNoisePan(double TestTime);
 
-	/**
-	* TestDeviceInputRandomizedDelay
-	*
-	* Test input device functionality with a bit more fun that the simple test. You should hear the input data randomly delayed by various amounts
-	* with randomized feedback and wetness. Should sound robotic and "interesting".
-	* @param TestTime The amount of time to run the test in seconds (negative value means to run indefinitely)
-	* @return True if test succeeds.
-	*/
-	bool UNREALAUDIO_API TestDeviceInputRandomizedDelay(double TestTime);
 }
 
diff --git a/Engine/Source/Runtime/Windows/UnrealAudioWasapi/Private/UnrealAudioDeviceWasapi.cpp b/Engine/Source/Runtime/Windows/UnrealAudioWasapi/Private/UnrealAudioDeviceWasapi.cpp
--- a/Engine/Source/Runtime/Windows/UnrealAudioWasapi/Private/UnrealAudioDeviceWasapi.cpp
+++ b/Engine/Source/Runtime/Windows/UnrealAudioWasapi/Private/UnrealAudioDeviceWasapi.cpp
@@ -155,14 +155,11 @@ namespace UAudio
 		bool GetDevicePlatformApi(EDeviceApi::Type & OutType) const override;
 		bool GetNumOutputDevices(uint32& OutNumDevices) const override;
 		bool GetOutputDeviceInfo(const uint32 DeviceIndex, FDeviceInfo& OutInfo) const override;
-		bool GetNumInputDevices(uint32& OutNumDevices) const override;
-		bool GetInputDeviceInfo(const uint32 DeviceIndex, FDeviceInfo& OutInfo) const override;
 		bool GetDefaultOutputDeviceIndex(uint32& OutDefaultIndex) const override;
-		bool GetDefaultInputDeviceIndex(uint32& OutDefaultIndex) const override;
 		bool StartStream() override;
 		bool StopStream() override;
 		bool ShutdownStream() override;
-		bool GetLatency(uint32& OutputDeviceLatency, uint32& InputDeviceLatency) const override;
+		bool GetLatency(uint32& OutputDeviceLatency) const override;
 		bool GetFrameRate(uint32& OutFrameRate) const override;
 
 		// FRunnable
@@ -183,27 +180,15 @@ namespace UAudio
 			/** The windows device enumerator. Used to query connected audio devices. */
 			IMMDeviceEnumerator* DeviceEnumerator;
 
-			/** Allows creation and initialization of audio input stream from audio engine and hardware buffer of audio endpoint device*/
-			IAudioClient* CaptureClient;
-
 			/** Allows creation and initialization of audio output stream from audio engine and hardware buffer of audio endpoint device*/
 			IAudioClient* RenderClient;
 
-			/** Allows reading data input from capture endpoint buffer */
-			IAudioCaptureClient* CaptureService;
-
 			/** Allows writing data output to capture endpoint buffer */
 			IAudioRenderClient* RenderService;
 
-			/** Handle used to notify when new input data is ready to be read in from hardware */
-			HANDLE CaptureEvent;
-
 			/** Handle used to notify when hardware is ready for new audio data to be written to it*/
 			HANDLE RenderEvent;
 
-			/** Intermediate buffer used to store audio data from input device, converted to user data format (float)*/
-			IIntermediateBuffer* CaptureIntermediateBuffer;
-
 			/** Intermediate buffer used to store audio data from user callback, converted to hardware native format*/
 			IIntermediateBuffer* RenderIntermediateBuffer;
 
@@ -212,22 +197,15 @@ namespace UAudio
 
 			FWasapiInfo()
 				: DeviceEnumerator(nullptr)
-				, CaptureClient(nullptr)
 				, RenderClient(nullptr)
-				, CaptureService(nullptr)
 				, RenderService(nullptr)
-				, CaptureEvent(0)
 				, RenderEvent(0)
-				, CaptureIntermediateBuffer(nullptr)
 				, RenderIntermediateBuffer(nullptr)
 				, bDevicesOpen(false)
 			{
 			}
 		};
 
-		/** The number of active input devices detected by WASAPI */
-		uint32 InputDeviceCount;
-
 		/** The number of active output devices detected by WASAPI */
 		uint32 OutputDeviceCount;
 
@@ -243,14 +221,14 @@ namespace UAudio
 	private:
 
 		// IUnrealAudioDeviceModule
-		bool OpenDevices(const FCreateStreamParams& CreateStreamParams) override;
+		bool OpenDevice(const FCreateStreamParams& CreateStreamParams) override;
 
 		// Helper functions
-		bool OpenDevice(uint32 DeviceIndex, EStreamType::Type StreamType);
-		bool GetDeviceCount(EDataFlow DataFlow, uint32& NumDevices) const;
-		bool GetDeviceInfo(EDataFlow DataFlow, const uint32 DeviceINdex, FDeviceInfo& DeviceInfo) const;
-		bool GetDefaultDeviceIndex(EDataFlow DataFlow, uint32& OutDeviceIndex) const;
-		void GetDeviceBufferSize(uint32& DeviceBufferSize, EStreamType::Type StreamType);
+		bool OpenDevice(uint32 DeviceIndex);
+		bool GetDeviceCount(uint32& NumDevices) const;
+		bool GetDeviceInfo(const uint32 DeviceINdex, FDeviceInfo& DeviceInfo) const;
+		bool GetDefaultDeviceIndex(uint32& OutDeviceIndex) const;
+		void GetDeviceBufferSize(uint32& DeviceBufferSize);
 	};
 
 	/** 
@@ -258,8 +236,7 @@ namespace UAudio
 	*/
 
 	FUnrealAudioWasapi::FUnrealAudioWasapi()
-		: InputDeviceCount(0)
-		, OutputDeviceCount(0)
+		: OutputDeviceCount(0)
 		, bInitialized(false)
 		, bComInitialized(false)
 	{
@@ -304,15 +281,10 @@ namespace UAudio
 
 	bool FUnrealAudioWasapi::GetNumOutputDevices(uint32& OutNumDevices) const
 	{
-		return GetDeviceCount(eRender, OutNumDevices);
+		return GetDeviceCount(OutNumDevices);
 	}
 
-	bool FUnrealAudioWasapi::GetNumInputDevices(uint32& OutNumDevices) const
-	{
-		return GetDeviceCount(eCapture, OutNumDevices);
-	}
-
-	bool FUnrealAudioWasapi::GetDeviceCount(EDataFlow DataFlow, uint32& OutNumDevices) const
+	bool FUnrealAudioWasapi::GetDeviceCount(uint32& OutNumDevices) const
 	{
 		if (!bInitialized || !WasapiInfo.DeviceEnumerator)
 		{
@@ -322,7 +294,7 @@ namespace UAudio
 		IMMDeviceCollection* Devices = nullptr;
 		HRESULT Result = S_OK;
 
-		Result = WasapiInfo.DeviceEnumerator->EnumAudioEndpoints(DataFlow, DEVICE_STATE_ACTIVE, &Devices);
+		Result = WasapiInfo.DeviceEnumerator->EnumAudioEndpoints(eRender, DEVICE_STATE_ACTIVE, &Devices);
 		CLEANUP_ON_FAIL(Result);
 
 		Result = Devices->GetCount(&OutNumDevices);
@@ -335,15 +307,10 @@ namespace UAudio
 
 	bool FUnrealAudioWasapi::GetOutputDeviceInfo(const uint32 DeviceIndex, FDeviceInfo& OutInfo) const
 	{
-		return GetDeviceInfo(eRender, DeviceIndex, OutInfo);
+		return GetDeviceInfo(DeviceIndex, OutInfo);
 	}
 
-	bool FUnrealAudioWasapi::GetInputDeviceInfo(const uint32 DeviceIndex, FDeviceInfo& OutInfo) const
-	{
-		return GetDeviceInfo(eCapture, DeviceIndex, OutInfo);
-	}
-
-	bool FUnrealAudioWasapi::GetDeviceInfo(EDataFlow DataFlow, const uint32 DeviceIndex, FDeviceInfo& DeviceInfo) const
+	bool FUnrealAudioWasapi::GetDeviceInfo(const uint32 DeviceIndex, FDeviceInfo& DeviceInfo) const
 	{
 		if (!bInitialized || !WasapiInfo.DeviceEnumerator)
 		{
@@ -361,14 +328,14 @@ namespace UAudio
 		WAVEFORMATEX* WaveFormatEx = nullptr;
 		HRESULT Result = S_OK;
 
-		Result = WasapiInfo.DeviceEnumerator->EnumAudioEndpoints(DataFlow, DEVICE_STATE_ACTIVE, &Devices);
+		Result = WasapiInfo.DeviceEnumerator->EnumAudioEndpoints(eRender, DEVICE_STATE_ACTIVE, &Devices);
 		CLEANUP_ON_FAIL(Result);
 
 		uint32 NumDevices = 0;
 		Result = Devices->GetCount(&NumDevices);
 		CLEANUP_ON_FAIL(Result);
 
-		Result = WasapiInfo.DeviceEnumerator->GetDefaultAudioEndpoint(DataFlow, eConsole, &DefaultDevice);
+		Result = WasapiInfo.DeviceEnumerator->GetDefaultAudioEndpoint(eRender, eConsole, &DefaultDevice);
 		CLEANUP_ON_FAIL(Result);
 
 		Result = Devices->Item(DeviceIndex, &Device);
@@ -396,7 +363,6 @@ namespace UAudio
 		DeviceInfo.FriendlyName = FString(DeviceNameProperty.pwszVal);
 		DeviceInfo.bIsSystemDefault = (DeviceInfo.FriendlyName == FString(DefaultDeviceNameProperty.pwszVal));
 		DeviceInfo.NumChannels = WaveFormatEx->nChannels;
-		DeviceInfo.StreamType = (DataFlow == eCapture) ? EStreamType::INPUT : EStreamType::OUTPUT;
 		DeviceInfo.FrameRate = WaveFormatEx->nSamplesPerSec;
 
 		// Figure out native supported data formats of the device
@@ -478,18 +444,13 @@ Cleanup:
 
 	bool FUnrealAudioWasapi::GetDefaultOutputDeviceIndex(uint32& OutDefaultIndex) const
 	{
-		return GetDefaultDeviceIndex(eRender, OutDefaultIndex);
-	}
-
-	bool FUnrealAudioWasapi::GetDefaultInputDeviceIndex(uint32& OutDefaultIndex) const
-	{
-		return GetDefaultDeviceIndex(eCapture, OutDefaultIndex);
+		return GetDefaultDeviceIndex(OutDefaultIndex);
 	}
 
-	bool FUnrealAudioWasapi::GetDefaultDeviceIndex(EDataFlow DataFlow, uint32& OutDeviceIndex) const
+	bool FUnrealAudioWasapi::GetDefaultDeviceIndex(uint32& OutDeviceIndex) const
 	{
 		uint32 NumDevices = 0;
-		if (!GetDeviceCount(DataFlow, NumDevices))
+		if (!GetDeviceCount(NumDevices))
 		{
 			return false;
 		}
@@ -499,7 +460,7 @@ Cleanup:
 		for (uint32 DeviceIndex = 0; DeviceIndex < NumDevices; ++DeviceIndex)
 		{
 			FDeviceInfo DeviceInfo;
-			if (!GetDeviceInfo(DataFlow, DeviceIndex, DeviceInfo))
+			if (!GetDeviceInfo(DeviceIndex, DeviceInfo))
 			{
 				return false;
 			}
@@ -515,7 +476,7 @@ Cleanup:
 		return bFoundDeviceIndex;
 	}
 
-	bool FUnrealAudioWasapi::OpenDevices(const FCreateStreamParams& CreateStreamParams)
+	bool FUnrealAudioWasapi::OpenDevice(const FCreateStreamParams& CreateStreamParams)
 	{
 		if (!bInitialized || WasapiInfo.bDevicesOpen)
 		{
@@ -526,29 +487,16 @@ Cleanup:
 
 		StreamInfo.FrameRate = INDEX_NONE;
 
-		// If we have input stream params, we try to open an input device
-		if (CreateStreamParams.InputDeviceIndex != INDEX_NONE)
-		{
-			if (!OpenDevice(CreateStreamParams.InputDeviceIndex, EStreamType::INPUT))
-			{
-				return false;
-			}
-			FStreamDeviceInfo& StreamDeviceInfo = StreamInfo.GetDeviceInfo(EStreamType::INPUT);
-			StreamInfo.FrameRate = FMath::Min(StreamInfo.FrameRate, StreamDeviceInfo.FrameRate);
-		}
-		
 		check(CreateStreamParams.OutputDeviceIndex != INDEX_NONE);
-		if (!OpenDevice(CreateStreamParams.OutputDeviceIndex, EStreamType::OUTPUT))
+		if (!OpenDevice(CreateStreamParams.OutputDeviceIndex))
 		{
 			return false;
 		}
-		FStreamDeviceInfo& StreamDeviceInfo = StreamInfo.GetDeviceInfo(EStreamType::OUTPUT);
-		StreamInfo.FrameRate = FMath::Min(StreamInfo.FrameRate, StreamDeviceInfo.FrameRate);
-
+		StreamInfo.FrameRate = FMath::Min(StreamInfo.FrameRate, StreamInfo.DeviceInfo.FrameRate);
 		return true;
 	}
 
-	bool FUnrealAudioWasapi::OpenDevice(uint32 DeviceIndex, EStreamType::Type StreamType)
+	bool FUnrealAudioWasapi::OpenDevice(uint32 DeviceIndex)
 	{
 		check(WasapiInfo.DeviceEnumerator);
 
@@ -558,12 +506,10 @@ Cleanup:
 		FDeviceInfo DeviceInfo;
 		HRESULT Result = S_OK;
 
-		EDataFlow DataFlow = (StreamType == EStreamType::INPUT) ? eCapture : eRender;
-
-		Result = WasapiInfo.DeviceEnumerator->EnumAudioEndpoints(DataFlow, DEVICE_STATE_ACTIVE, &DeviceList);
+		Result = WasapiInfo.DeviceEnumerator->EnumAudioEndpoints(eRender, DEVICE_STATE_ACTIVE, &DeviceList);
 		CLEANUP_ON_FAIL(Result);
 
-		if (!GetDeviceInfo(DataFlow, DeviceIndex, DeviceInfo))
+		if (!GetDeviceInfo(DeviceIndex, DeviceInfo))
 		{
 			Result = S_FALSE;
 			goto Cleanup;
@@ -572,11 +518,11 @@ Cleanup:
 		Result = DeviceList->Item(DeviceIndex, &Device);
 		CLEANUP_ON_FAIL(Result);
 
-		FStreamDeviceInfo& StreamDeviceInfo = StreamInfo.GetDeviceInfo(StreamType);
+		FStreamDeviceInfo& StreamDeviceInfo = StreamInfo.DeviceInfo;
 
 		StreamDeviceInfo.DeviceIndex = DeviceIndex;
 
-		IAudioClient*& AudioClient = (StreamType == EStreamType::INPUT) ? WasapiInfo.CaptureClient : WasapiInfo.RenderClient;
+		IAudioClient*& AudioClient = WasapiInfo.RenderClient;
 		check(AudioClient == nullptr);
 
 		Result = Device->Activate(__uuidof(IAudioClient), CLSCTX_ALL, nullptr, (void**)&AudioClient);
@@ -595,10 +541,10 @@ Cleanup:
 		{
 			StreamDeviceInfo.bPerformFormatConversion = true;
 		}
-		SetupBufferFormatConvertInfo(StreamType);
+		SetupBufferFormatConvertInfo();
 
 		// Initialize the raw buffer used to write into
-		uint32 BufferSize = StreamDeviceInfo.NumChannels * StreamInfo.BlockSize * GetNumBytesForFormat(EStreamFormat::FLT);
+		uint32 BufferSize = StreamDeviceInfo.NumChannels * StreamInfo.BlockSize * sizeof(float);
 		StreamDeviceInfo.UserBuffer.Init(0, BufferSize);
 
 	Cleanup:
@@ -634,12 +580,12 @@ Cleanup:
 		}
 
 		check(StreamInfo.Thread != nullptr);
-		check(WasapiInfo.CaptureClient != nullptr || WasapiInfo.RenderClient != nullptr);
+		check(WasapiInfo.RenderClient != nullptr);
 
 		// Tell the stream update that we are stopping
 		Stop();
 
-			// Wait for the thread to finish
+		// Wait for the thread to finish
 		StreamInfo.Thread->WaitForCompletion();
 
 		check(StreamInfo.State == EStreamState::STOPPED);
@@ -652,16 +598,6 @@ Cleanup:
 		HRESULT Result = S_OK;
 		bool bSuccess = true;
 
-		if (WasapiInfo.CaptureClient)
-		{
-			Result = WasapiInfo.CaptureClient->Stop();
-			if (FAILED(Result))
-			{
-				bSuccess = false;
-				UA_DEVICE_PLATFORM_ERROR("Failed to stop capture client.");
-			}
-		}
-
 		if (WasapiInfo.RenderClient)
 		{
 			Result = WasapiInfo.RenderClient->Stop();
@@ -687,17 +623,9 @@ Cleanup:
 			StopStream();
 		}
 
-		SAFE_RELEASE(WasapiInfo.CaptureClient);
-		SAFE_RELEASE(WasapiInfo.CaptureService);
 		SAFE_RELEASE(WasapiInfo.RenderClient);
 		SAFE_RELEASE(WasapiInfo.RenderService);
 
-		if (WasapiInfo.CaptureEvent)
-		{
-			CloseHandle(WasapiInfo.CaptureEvent);
-			WasapiInfo.CaptureEvent = nullptr;
-		}
-
 		if (WasapiInfo.RenderEvent)
 		{
 			CloseHandle(WasapiInfo.RenderEvent);
@@ -708,10 +636,9 @@ Cleanup:
 		return true;
 	}
 
-	bool FUnrealAudioWasapi::GetLatency(uint32& OutputDeviceLatency, uint32& InputDeviceLatency) const
+	bool FUnrealAudioWasapi::GetLatency(uint32& OutputDeviceLatency) const
 	{
-		OutputDeviceLatency = StreamInfo.GetDeviceInfo(EStreamType::OUTPUT).Latency;
-		InputDeviceLatency = StreamInfo.GetDeviceInfo(EStreamType::INPUT).Latency;
+		OutputDeviceLatency = StreamInfo.DeviceInfo.Latency;
 		return true;
 	}
 
@@ -729,21 +656,15 @@ Cleanup:
 		bool bThreadComInitialized = FWindowsPlatformMisc::CoInitialize();
 
 		// Get local stack versions of various wasapi info we need for the device thread
-		WAVEFORMATEX* InputFormat = nullptr;
 		WAVEFORMATEX* OutputFormat = nullptr;
-		bool bWasInputBufferRead = false;
-		bool bWasOutputBufferWritten = false;
-		bool bWasInputBufferWritten = false;
-		bool bOutputCallbackCalled = false;
-		BYTE* DeviceByteBuffer = nullptr;		// byte buffer to hold input data from input device
-		uint32 CaptureBufferFrameCount = 0;		// how many frames read in from input device
-		::DWORD CaptureFlags = 0;				// flags from wasapi on audio input
-		FStreamDeviceInfo& InputDeviceInfo = StreamInfo.GetDeviceInfo(EStreamType::INPUT);
-		FStreamDeviceInfo& OutputDeviceInfo = StreamInfo.GetDeviceInfo(EStreamType::OUTPUT);
+		BYTE* DeviceByteBuffer = nullptr;
+		FStreamDeviceInfo& OutputDeviceInfo = StreamInfo.DeviceInfo;
 		uint32 NumBytesRenderFormat = GetNumBytesForFormat(OutputDeviceInfo.DeviceDataFormat);
-		uint32 NumBytesCaptureFormat = GetNumBytesForFormat(InputDeviceInfo.DeviceDataFormat);
-		uint32 InputDeviceSamples = StreamInfo.BlockSize * InputDeviceInfo.NumChannels;
 		uint32 OutputDeviceSamples = StreamInfo.BlockSize * OutputDeviceInfo.NumChannels;
+
+		bool bWasUserCallbackMade = false;
+		bool bWasUserBufferWritten = false;
+
 		HRESULT Result = S_OK;
 		FCallbackInfo CallbackInfo;
 
@@ -752,360 +673,165 @@ Cleanup:
 		// Number used to compute the REFERENCE_TIME for callback periods (REF_TIME is 100 nanoseconds)
 		static const uint64 REF_TIMES_PER_SECOND = 10000000;
 
-		// If we are trying to get audio from an input device, set up a capture stream
-		if (WasapiInfo.CaptureClient)
-		{
-
-			// We haven't set up an actual capture service yet
-			check(WasapiInfo.CaptureService == nullptr);
-			
-			Result = WasapiInfo.CaptureClient->GetMixFormat(&InputFormat);
-			CLEANUP_ON_FAIL(Result);
-
-			// The actual buffer size we're going to use, independent of the actual buffer size we want. 
-			// Compute the callback buffer period using the buffer size
-			REFERENCE_TIME CallbackBufferPeriod = (REFERENCE_TIME)((StreamInfo.BlockSize * REF_TIMES_PER_SECOND) / InputFormat->nSamplesPerSec);
-
-			Result = WasapiInfo.CaptureClient->Initialize(
-				AUDCLNT_SHAREMODE_SHARED,				// Other clients can use this input device
-				AUDCLNT_STREAMFLAGS_EVENTCALLBACK,		// Processing of the buffer by client is event driven
-				CallbackBufferPeriod,					// the size of the buffer in 100-nanosecond units (REF_TIME)
-				0,										// always 0 in shared mode 
-				InputFormat,							// the input format to use
-				nullptr									// audio session guid, we're ignoring
-			);
-			CLEANUP_ON_FAIL(Result);
-
-			REFERENCE_TIME InputLatency = 0;
-			WasapiInfo.CaptureClient->GetStreamLatency(&InputLatency);
+		// We haven't set up an actual render service yet
+		check(WasapiInfo.RenderService == nullptr);
 
-			// Get the reported device latency -- note for WASAPI, this can only be called after Initialize on capture client
-			InputDeviceInfo.Latency = (1000 * InputLatency) / REF_TIMES_PER_SECOND;
-
-			Result = WasapiInfo.CaptureClient->GetService(__uuidof(IAudioCaptureClient), (void**)&WasapiInfo.CaptureService);
-			CLEANUP_ON_FAIL(Result);
-
-			WasapiInfo.CaptureEvent = CreateEvent(nullptr, false, false, nullptr);
-			if (!WasapiInfo.CaptureClient)
-			{
-				Result = HRESULT_FROM_WIN32(GetLastError());
-				CLEANUP_ON_FAIL(Result);
-			}
-			Result = WasapiInfo.CaptureClient->SetEventHandle(WasapiInfo.CaptureEvent);
-			CLEANUP_ON_FAIL(Result);
-
-			uint32 ReadBufferSize = 0;
-			Result = WasapiInfo.CaptureClient->GetBufferSize(&ReadBufferSize);
-			CLEANUP_ON_FAIL(Result);
-
-			uint32 WriteBufferSize = StreamInfo.BlockSize * InputDeviceInfo.NumChannels;
-			ReadBufferSize *= InputDeviceInfo.NumChannels;
+		Result = WasapiInfo.RenderClient->GetMixFormat(&OutputFormat);
+		CLEANUP_ON_FAIL(Result);
 
-			WasapiInfo.CaptureIntermediateBuffer = IIntermediateBuffer::CreateIntermediateBuffer(InputDeviceInfo.DeviceDataFormat);
-			check(WasapiInfo.CaptureIntermediateBuffer);
-			WasapiInfo.CaptureIntermediateBuffer->Initialize(ReadBufferSize + WriteBufferSize, NumBytesCaptureFormat);
+		// Compute the callback buffer period using the buffer size
+		REFERENCE_TIME CallbackBufferPeriod = (REFERENCE_TIME)((StreamInfo.BlockSize * REF_TIMES_PER_SECOND) / OutputFormat->nSamplesPerSec);
+
+		Result = WasapiInfo.RenderClient->Initialize(
+			AUDCLNT_SHAREMODE_SHARED,				// Other clients can use this input device
+			AUDCLNT_STREAMFLAGS_EVENTCALLBACK,		// Processing of the buffer by client is event driven
+			CallbackBufferPeriod,					// the size of the buffer in 100-nanosecond units (REF_TIME)
+			0,										// always 0 in shared mode 
+			OutputFormat,							// the output format to use
+			nullptr									// audio session guid, we're ignoring
+		);
+		CLEANUP_ON_FAIL(Result);
 
-			Result = WasapiInfo.CaptureClient->Reset();
-			CLEANUP_ON_FAIL(Result);
+		// Get the output latency
+		REFERENCE_TIME OutputLatency = 0;
+		WasapiInfo.RenderClient->GetStreamLatency(&OutputLatency);
+		OutputDeviceInfo.Latency = (1000 * OutputLatency) / REF_TIMES_PER_SECOND;
 
-			Result = WasapiInfo.CaptureClient->Start();
-			CLEANUP_ON_FAIL(Result);
-		}
+		Result = WasapiInfo.RenderClient->GetService(__uuidof(IAudioRenderClient), (void**)&WasapiInfo.RenderService);
+		CLEANUP_ON_FAIL(Result);
 
-		if (WasapiInfo.RenderClient)
+		WasapiInfo.RenderEvent = CreateEvent(nullptr, false, false, nullptr);
+		if (!WasapiInfo.RenderEvent)
 		{
-			// We haven't set up an actual render service yet
-			check(WasapiInfo.RenderService == nullptr);
-
-			Result = WasapiInfo.RenderClient->GetMixFormat(&OutputFormat);
-			CLEANUP_ON_FAIL(Result);
-
-			// Compute the callback buffer period using the buffer size
-			REFERENCE_TIME CallbackBufferPeriod = (REFERENCE_TIME)((StreamInfo.BlockSize * REF_TIMES_PER_SECOND) / OutputFormat->nSamplesPerSec);
-
-			Result = WasapiInfo.RenderClient->Initialize(
-				AUDCLNT_SHAREMODE_SHARED,				// Other clients can use this input device
-				AUDCLNT_STREAMFLAGS_EVENTCALLBACK,		// Processing of the buffer by client is event driven
-				CallbackBufferPeriod,					// the size of the buffer in 100-nanosecond units (REF_TIME)
-				0,										// always 0 in shared mode 
-				OutputFormat,							// the output format to use
-				nullptr									// audio session guid, we're ignoring
-			);
+			Result = HRESULT_FROM_WIN32(GetLastError());
 			CLEANUP_ON_FAIL(Result);
+		}
 
-			// Get the output latency
-			REFERENCE_TIME OutputLatency = 0;
-			WasapiInfo.RenderClient->GetStreamLatency(&OutputLatency);
-			OutputDeviceInfo.Latency = (1000 * OutputLatency) / REF_TIMES_PER_SECOND;
-
-			Result = WasapiInfo.RenderClient->GetService(__uuidof(IAudioRenderClient), (void**)&WasapiInfo.RenderService);
-			CLEANUP_ON_FAIL(Result);
-
-			WasapiInfo.RenderEvent = CreateEvent(nullptr, false, false, nullptr);
-			if (!WasapiInfo.RenderEvent)
-			{
-				Result = HRESULT_FROM_WIN32(GetLastError());
-				CLEANUP_ON_FAIL(Result);
-			}
-
-			Result = WasapiInfo.RenderClient->SetEventHandle(WasapiInfo.RenderEvent);
-			CLEANUP_ON_FAIL(Result);
+		Result = WasapiInfo.RenderClient->SetEventHandle(WasapiInfo.RenderEvent);
+		CLEANUP_ON_FAIL(Result);
 
-			uint32 ReadBufferSize = 0;
-			Result = WasapiInfo.RenderClient->GetBufferSize(&ReadBufferSize);
-			CLEANUP_ON_FAIL(Result);
+		uint32 ReadBufferSize = 0;
+		Result = WasapiInfo.RenderClient->GetBufferSize(&ReadBufferSize);
+		CLEANUP_ON_FAIL(Result);
 
-			uint32 WriteBufferSize = StreamInfo.BlockSize* OutputDeviceInfo.NumChannels;
+		uint32 WriteBufferSize = StreamInfo.BlockSize* OutputDeviceInfo.NumChannels;
 
-			ReadBufferSize *= OutputDeviceInfo.NumChannels;
+		ReadBufferSize *= OutputDeviceInfo.NumChannels;
 
-			WasapiInfo.RenderIntermediateBuffer = IIntermediateBuffer::CreateIntermediateBuffer(OutputDeviceInfo.DeviceDataFormat);
-			check(WasapiInfo.RenderIntermediateBuffer);
-			WasapiInfo.RenderIntermediateBuffer->Initialize(ReadBufferSize + WriteBufferSize, NumBytesRenderFormat);
+		WasapiInfo.RenderIntermediateBuffer = IIntermediateBuffer::CreateIntermediateBuffer(OutputDeviceInfo.DeviceDataFormat);
+		check(WasapiInfo.RenderIntermediateBuffer);
+		WasapiInfo.RenderIntermediateBuffer->Initialize(ReadBufferSize + WriteBufferSize);
 
-			Result = WasapiInfo.RenderClient->Reset();
-			CLEANUP_ON_FAIL(Result);
+		Result = WasapiInfo.RenderClient->Reset();
+		CLEANUP_ON_FAIL(Result);
 
-			Result = WasapiInfo.RenderClient->Start();
-			CLEANUP_ON_FAIL(Result);
-		}
+		Result = WasapiInfo.RenderClient->Start();
+		CLEANUP_ON_FAIL(Result);
 
 		// Set up the device buffer (buffer read from and written directly to the device in the native device format)
 
 		// If we have an input stream, set up our device buffer size to accommodate
 		uint32 DeviceBufferSize = 0;
-		if (StreamInfo.StreamType == EStreamType::INPUT)
-		{
-			uint32 DeviceBufferInput = 0;
-			GetDeviceBufferSize(DeviceBufferInput, EStreamType::INPUT);
-
-			uint32 DeviceBufferOutput = 0;
-			GetDeviceBufferSize(DeviceBufferOutput, EStreamType::OUTPUT);
-
-			DeviceBufferSize = FMath::Max(DeviceBufferInput, DeviceBufferOutput);
-		}
-		else
-		{
-			GetDeviceBufferSize(DeviceBufferSize, StreamInfo.StreamType);
-		}
+		GetDeviceBufferSize(DeviceBufferSize);
 
 		StreamInfo.DeviceBuffer.Init(0, DeviceBufferSize);
 		uint32 Count = 0;
 
 		// Prepare the struct which will be used to make audio callbacks
 		CallbackInfo.OutBuffer = (float *)OutputDeviceInfo.UserBuffer.GetData();;
-		CallbackInfo.InBuffer = (float *)InputDeviceInfo.UserBuffer.GetData();;
 		CallbackInfo.NumFrames = StreamInfo.BlockSize;
 		CallbackInfo.NumChannels = OutputDeviceInfo.NumChannels;
-		CallbackInfo.NumInputChannels = InputDeviceInfo.NumChannels;
 		CallbackInfo.StreamTime = 0.0;
 		CallbackInfo.UserData = StreamInfo.UserData;
 		CallbackInfo.StatusFlags = 0;
 		CallbackInfo.OutputSpeakers = OutputDeviceInfo.Speakers;
 		CallbackInfo.FrameRate = OutputDeviceInfo.FrameRate;
 
-		/**
-		
-		The update loop performs the following stages:
-
-		1. Prepare input audio to feed into user callback
-			- input was pushed into intermediate buffer from previous update tick
-			- input is retrieved from intermediate buffer
-			- input is converted from device format from device buffer to user format in user buffer
-
-		2. Perform the user callback
-			- user callback is called with possible input data from stage 1
-			- return value of user callback is evaluated (success or failure) to determine the state of the thread loop
-
-		3. Convert user callback data buffer format to device data buffer format
-			- write results to intermediate output buffer
-
-		4. Get new audio input from device
-			- get input data from audio device
-			- push the input data to the intermediate buffer so that in next update it can be given to user callback
-
-		5. Write audio to output buffer
-			- get output data from intermediate buffer
-			- write output data to render device
-
-		*/
-
 		while (StreamInfo.State != EStreamState::STOPPING)
 		{
-			//UE_LOG(LogUnrealAudioDevice, Display, TEXT("Device Callback (%d)"), ++Count);
-
-			// Stage 1
-
-			if (!bWasInputBufferRead)
+			if (!bWasUserCallbackMade)
 			{
-				if (WasapiInfo.CaptureClient)
+				CallbackInfo.StatusFlags = 0;
+				CallbackInfo.StreamTime = StreamInfo.StreamTime;
+				memset((void*)CallbackInfo.OutBuffer, 0, StreamInfo.BlockSize*OutputDeviceInfo.NumChannels*sizeof(float));
+				if (!StreamInfo.CallbackFunction(CallbackInfo))
 				{
-					// Read data from the intermediate buffer so that it can be fed to the user callback
-					if (InputDeviceInfo.bPerformFormatConversion)
-					{
-						bWasInputBufferRead = WasapiInfo.CaptureIntermediateBuffer->Read(StreamInfo.DeviceBuffer.GetData(), InputDeviceSamples);
-						if (bWasInputBufferRead)
-						{
-							ConvertBufferFormat(InputDeviceInfo.UserBuffer, StreamInfo.DeviceBuffer, EStreamType::INPUT);
-						}
-					}
-					else
-					{
-						bWasInputBufferRead = WasapiInfo.CaptureIntermediateBuffer->Read(InputDeviceInfo.UserBuffer.GetData(), InputDeviceSamples);
-					}
-				}
-				else
-				{
-					// If there was no audio to read in (i.e. we are in device output only) pretend like we read audio input
-					bWasInputBufferRead = true;
-				}
-
-				// Stage 2
-
-				if (bWasInputBufferRead)
-				{
-					// Report to callback if the input device reported any discontinuities in the input stream
-					CallbackInfo.StatusFlags = 0;
-					if (CaptureFlags & AUDCLNT_BUFFERFLAGS_DATA_DISCONTINUITY)
-					{
-						CallbackInfo.StatusFlags |= EStreamFlowStatus::INPUT_OVERFLOW;
-					}
-
-					CallbackInfo.StreamTime = StreamInfo.StreamTime;
-
-					// Zero out the output buffer before calling the user callback
-					memset((void*)CallbackInfo.OutBuffer, 0, StreamInfo.BlockSize*OutputDeviceInfo.NumChannels*sizeof(float));
-
-					// Make the user callback
-					bool bSuccess = StreamInfo.CallbackFunction(CallbackInfo);
-
-					UpdateStreamTimeTick();
-
-					// Check if we need to stop the device thread
-					if (!bSuccess)
-					{
-						// Now stop the thread...
-						StreamInfo.State = EStreamState::STOPPING;
-					}
+					StreamInfo.State = EStreamState::STOPPING;
 				}
+				UpdateStreamTimeTick();
+				bWasUserCallbackMade = true;
 			}
 
-			// Stage 3
-
-			if (WasapiInfo.RenderClient && bWasInputBufferRead)
+			if (bWasUserCallbackMade)
 			{
 				if (OutputDeviceInfo.bPerformFormatConversion)
 				{
-					ConvertBufferFormat(StreamInfo.DeviceBuffer, OutputDeviceInfo.UserBuffer, EStreamType::OUTPUT);
-					bWasOutputBufferWritten = WasapiInfo.RenderIntermediateBuffer->Write(StreamInfo.DeviceBuffer.GetData(), OutputDeviceSamples);
+					ConvertBufferFormat(StreamInfo.DeviceBuffer, OutputDeviceInfo.UserBuffer);
+					bWasUserBufferWritten = WasapiInfo.RenderIntermediateBuffer->Write(StreamInfo.DeviceBuffer.GetData(), OutputDeviceSamples);
 				}
 				else
 				{
-					bWasOutputBufferWritten = WasapiInfo.RenderIntermediateBuffer->Write(OutputDeviceInfo.UserBuffer.GetData(), OutputDeviceSamples);
+					bWasUserBufferWritten = WasapiInfo.RenderIntermediateBuffer->Write(OutputDeviceInfo.UserBuffer.GetData(), OutputDeviceSamples);
 				}
 			}
 			else
 			{
-				bWasOutputBufferWritten = true;
+				bWasUserBufferWritten = true;
 			}
 
-			// Stage 4
-
-			if (WasapiInfo.CaptureClient)
+			// If the user output buffer was not pushed to the intermediate buffer, wait for the 
+			// next render event from the device before going on
+			if (bWasUserCallbackMade && !bWasUserBufferWritten)
 			{
-				// if we didn't read from the input buffer then, pause this thread until we get an event
-				// signal from the input audio client (this will get sent as soon as the device has read an input buffer)
-				if (!bWasInputBufferRead)
-				{
-					WaitForSingleObject(WasapiInfo.CaptureEvent, INFINITE);
-				}
-
-				Result = WasapiInfo.CaptureService->GetBuffer(&DeviceByteBuffer, &CaptureBufferFrameCount, &CaptureFlags, nullptr, nullptr);
-				CLEANUP_ON_FAIL(Result);
-
-				bool bCaptureSuccess = false;
-				if (CaptureBufferFrameCount != 0)
-				{
-					uint32 WriteSamples = CaptureBufferFrameCount * InputDeviceInfo.NumChannels;
-					bWasInputBufferWritten = WasapiInfo.CaptureIntermediateBuffer->Write((uint8*)DeviceByteBuffer, WriteSamples);
-					if (bWasInputBufferWritten)
-					{
-						bCaptureSuccess = true;
-						Result = WasapiInfo.CaptureService->ReleaseBuffer(CaptureBufferFrameCount);
-						CLEANUP_ON_FAIL(Result);
-					}
-				}
-
-				if (!bCaptureSuccess)
-				{
-					// We failed to read data from input device, releasing 0 bytes will release the capture buffer
-					Result = WasapiInfo.CaptureService->ReleaseBuffer(0);
-					CLEANUP_ON_FAIL(Result);
-				}
+				WaitForSingleObject(WasapiInfo.RenderEvent, INFINITE);
 			}
 
-			// Stage 5
+			uint32 OutputBufferFrameCount = 0;
+			uint32 OutputBufferFramePadding = 0;
 
-			if (WasapiInfo.RenderClient)
-			{
-				// If the user output buffer was not pushed to the intermediate buffer, wait for the 
-				// next render event from the device before going on
-				if (bWasInputBufferRead && !bWasOutputBufferWritten)
-				{
-					WaitForSingleObject(WasapiInfo.RenderEvent, INFINITE);
-				}
+			// Get render buffer from stream
+			Result = WasapiInfo.RenderClient->GetBufferSize(&OutputBufferFrameCount);
+			CLEANUP_ON_FAIL(Result);
 
-				uint32 OutputBufferFrameCount = 0;
-				uint32 OutputBufferFramePadding = 0;
+			Result = WasapiInfo.RenderClient->GetCurrentPadding(&OutputBufferFramePadding);
+			CLEANUP_ON_FAIL(Result);
 
-				// Get render buffer from stream
-				Result = WasapiInfo.RenderClient->GetBufferSize(&OutputBufferFrameCount);
-				CLEANUP_ON_FAIL(Result);
+			OutputBufferFrameCount -= OutputBufferFramePadding;
 
-				Result = WasapiInfo.RenderClient->GetCurrentPadding(&OutputBufferFramePadding);
+			if (OutputBufferFrameCount != 0)
+			{
+				Result = WasapiInfo.RenderService->GetBuffer(OutputBufferFrameCount, &DeviceByteBuffer);
 				CLEANUP_ON_FAIL(Result);
 
-				OutputBufferFrameCount -= OutputBufferFramePadding;
-
-				if (OutputBufferFrameCount != 0)
+				// Read the next buffer from the intermediate output buffer
+				uint32 ReadSamples = OutputBufferFrameCount * OutputDeviceInfo.NumChannels;
+				bool bSuccess = WasapiInfo.RenderIntermediateBuffer->Read((uint8*)DeviceByteBuffer, ReadSamples);
+				if (bSuccess)
 				{
-					Result = WasapiInfo.RenderService->GetBuffer(OutputBufferFrameCount, &DeviceByteBuffer);
+					Result = WasapiInfo.RenderService->ReleaseBuffer(OutputBufferFrameCount, 0);
+					CLEANUP_ON_FAIL(Result);
+				}
+				else
+				{
+					Result = WasapiInfo.RenderService->ReleaseBuffer(0, 0);
 					CLEANUP_ON_FAIL(Result);
-
-					// Read the next buffer from the intermediate output buffer
-					uint32 ReadSamples = OutputBufferFrameCount * OutputDeviceInfo.NumChannels;
-					bool bSuccess = WasapiInfo.RenderIntermediateBuffer->Read((uint8*)DeviceByteBuffer, ReadSamples);
-					if (bSuccess)
-					{
-						Result = WasapiInfo.RenderService->ReleaseBuffer(OutputBufferFrameCount, 0);
-						CLEANUP_ON_FAIL(Result);
-					}
-					else
-					{
-						Result = WasapiInfo.RenderService->ReleaseBuffer(0, 0);
-						CLEANUP_ON_FAIL(Result);
-					}
 				}
 			}
 
-			if (bWasOutputBufferWritten)
+			if (bWasUserBufferWritten)
 			{
-				bWasInputBufferRead = false;
+				bWasUserCallbackMade = false;
 			}
 		}
 
 		Cleanup:
 
-		CoTaskMemFree(InputFormat);
 		CoTaskMemFree(OutputFormat);
-
 		if (bThreadComInitialized)
 		{
 			FWindowsPlatformMisc::CoUninitialize();
 		}
-
 		StreamInfo.State = EStreamState::STOPPED;
-
 		return 0;
 	}
 
@@ -1119,11 +845,10 @@ Cleanup:
 
 	}
 
-	void FUnrealAudioWasapi::GetDeviceBufferSize(uint32& DeviceBufferSize, EStreamType::Type StreamType)
+	void FUnrealAudioWasapi::GetDeviceBufferSize(uint32& DeviceBufferSize)
 	{
-		FStreamDeviceInfo& DeviceInfo = StreamInfo.GetDeviceInfo(StreamType);
-		uint32 NumBytes = GetNumBytesForFormat(DeviceInfo.DeviceDataFormat);
-		uint32 NumChannels = DeviceInfo.NumChannels;
+		uint32 NumBytes = GetNumBytesForFormat(StreamInfo.DeviceInfo.DeviceDataFormat);
+		uint32 NumChannels = StreamInfo.DeviceInfo.NumChannels;
 		DeviceBufferSize = StreamInfo.BlockSize * NumChannels * NumBytes;
 	}
 }
diff --git a/Engine/Source/Runtime/Windows/UnrealAudioXAudio2/Private/UnrealAudioDeviceXAudio2.cpp b/Engine/Source/Runtime/Windows/UnrealAudioXAudio2/Private/UnrealAudioDeviceXAudio2.cpp
--- a/Engine/Source/Runtime/Windows/UnrealAudioXAudio2/Private/UnrealAudioDeviceXAudio2.cpp
+++ b/Engine/Source/Runtime/Windows/UnrealAudioXAudio2/Private/UnrealAudioDeviceXAudio2.cpp
@@ -1,25 +1,25 @@
 // Copyright 1998-2015 Epic Games, Inc. All Rights Reserved.
 
 /**
-	Concrete implementation of IAudioDevice for XAudio2
+Concrete implementation of IAudioDevice for XAudio2
 
-	See https://msdn.microsoft.com/en-us/library/windows/desktop/hh405049%28v=vs.85%29.aspx
+See https://msdn.microsoft.com/en-us/library/windows/desktop/hh405049%28v=vs.85%29.aspx
 
 */
 
 #include "UnrealAudioDeviceModule.h"
 #include "UnrealAudioBuffer.h"
+#include "ModuleManager.h"
 
 #include "AllowWindowsPlatformTypes.h"
 #include <xaudio2.h>
-#if PLATFORM_WINDOWS
-#include <mmdeviceapi.h>					// For getting endpoint notifications about audio devices
-#include <functiondiscoverykeys_devpkey.h>	// Defines property keys for the Plug and Play Device Property API.
-#endif
 #include "HideWindowsPlatformTypes.h"
 
 #if ENABLE_UNREAL_AUDIO
 
+// Used to toggle on white noise testing for xaudio2 output
+#define UNREAL_AUDIO_TEST_WHITE_NOISE 0
+
 DEFINE_LOG_CATEGORY(LogUnrealAudioDevice);
 
 // See MSDN documentation for what these error codes mean in the context of the API call
@@ -45,24 +45,99 @@ static const TCHAR* GetXAudio2Error(HRESULT Result)
 
 #define CLEANUP_ON_FAIL(Result)						\
 	if (FAILED(Result))								\
-	{												\
+		{												\
 		const TCHAR* Err = GetXAudio2Error(Result);	\
 		UA_DEVICE_PLATFORM_ERROR(Err);				\
 		goto Cleanup;								\
-	}
+		}
 
 #define RETURN_FALSE_ON_FAIL(Result)				\
 	if (FAILED(Result))								\
-	{												\
+		{												\
 		const TCHAR* Err = GetXAudio2Error(Result);	\
 		UA_DEVICE_PLATFORM_ERROR(Err);				\
 		return false;								\
-	}
+		}
 
 namespace UAudio
 {
+	/** Struct used to map Unreal Audio enumerations for speaker types to device speaker types */
+	struct FSpeakerMaskMapping
+	{
+		ESpeaker::Type UnrealSpeaker;
+		long XAudio2SpeakerTypeFlag;
+	};
+
+	/** Mapping of XAudio2/Windows speaker type enumerations and unreal audio enumerations */
+	FSpeakerMaskMapping SpeakerMaskMapping[18] =
+	{
+		{ ESpeaker::FRONT_LEFT, SPEAKER_FRONT_LEFT },
+		{ ESpeaker::FRONT_RIGHT, SPEAKER_FRONT_RIGHT },
+		{ ESpeaker::FRONT_CENTER, SPEAKER_FRONT_CENTER },
+		{ ESpeaker::LOW_FREQUENCY, SPEAKER_LOW_FREQUENCY },
+		{ ESpeaker::BACK_LEFT, SPEAKER_BACK_LEFT },
+		{ ESpeaker::BACK_RIGHT, SPEAKER_BACK_RIGHT },
+		{ ESpeaker::FRONT_LEFT_OF_CENTER, SPEAKER_FRONT_LEFT_OF_CENTER },
+		{ ESpeaker::FRONT_RIGHT_OF_CENTER, SPEAKER_FRONT_RIGHT_OF_CENTER },
+		{ ESpeaker::BACK_CENTER, SPEAKER_BACK_CENTER },
+		{ ESpeaker::SIDE_LEFT, SPEAKER_SIDE_LEFT },
+		{ ESpeaker::SIDE_RIGHT, SPEAKER_SIDE_RIGHT },
+		{ ESpeaker::TOP_CENTER, SPEAKER_TOP_CENTER },
+		{ ESpeaker::TOP_FRONT_LEFT, SPEAKER_TOP_FRONT_LEFT },
+		{ ESpeaker::TOP_FRONT_CENTER, SPEAKER_TOP_FRONT_CENTER },
+		{ ESpeaker::TOP_FRONT_RIGHT, SPEAKER_TOP_FRONT_RIGHT },
+		{ ESpeaker::TOP_BACK_LEFT, SPEAKER_TOP_BACK_LEFT },
+		{ ESpeaker::TOP_BACK_CENTER, SPEAKER_TOP_BACK_CENTER },
+		{ ESpeaker::TOP_BACK_RIGHT, SPEAKER_TOP_BACK_RIGHT },
+	};
+
+	/**
+	* FXAudio2VoiceCallback
+	* XAudio2 implementation of IXAudio2VoiceCallback
+	* This callback class is used to get event notifications on buffer end (when a buffer has finished processing).
+	* This is used to signal the I/O thread that it can request another buffer from the user callback.
+	*/
+	class FXAudio2VoiceCallback : public IXAudio2VoiceCallback
+	{
+	public:
+		FXAudio2VoiceCallback()
+			: BufferEndEvent(nullptr)
+		{
+		}
+
+		~FXAudio2VoiceCallback()
+		{
+		}
+
+		/** Sets the windows handle/event to use on buffer end */
+		void SetBufferEndEvent(HANDLE InBufferEndEvent)
+		{
+			BufferEndEvent = InBufferEndEvent;
+		}
+
+	private:
+		void STDCALL OnVoiceProcessingPassStart(UINT32 BytesRequired) {}
+		void STDCALL OnVoiceProcessingPassEnd() {}
+		void STDCALL OnStreamEnd() {}
+		void STDCALL OnBufferStart(void* BufferContext) {}
+		void STDCALL OnLoopEnd(void* BufferContext) {}
+		void STDCALL OnVoiceError(void* BufferContext, HRESULT Error) {}
+		void STDCALL OnBufferEnd(void* BufferContext)
+		{
+			SetEvent(BufferEndEvent);
+		}
+
+	private:
+		HANDLE BufferEndEvent;
+	};
+
+	/**
+	* FUnrealAudioXAudio2
+	* XAudio2 implementation of IUnrealAudioDeviceModule
+	* Inherits from FRunnable so that audio device I/O can be run on separate thread.
+	*/
 	class FUnrealAudioXAudio2 : public IUnrealAudioDeviceModule,
-		public FRunnable
+								public FRunnable
 	{
 	public:
 		FUnrealAudioXAudio2();
@@ -73,45 +148,77 @@ namespace UAudio
 		bool Shutdown() override;
 		bool GetDevicePlatformApi(EDeviceApi::Type & OutType) const override;
 		bool GetNumOutputDevices(uint32& OutNumDevices) const override;
-		bool GetOutputDeviceInfo(const uint32 DeviceIndex, FDeviceInfo& OutInfo) const override;
-		bool GetNumInputDevices(uint32& OutNumDevices) const override;
-		bool GetInputDeviceInfo(const uint32 DeviceIndex, FDeviceInfo& OutInfo) const override;
+		bool GetOutputDeviceInfo(const uint32 DeviceIndex, FDeviceInfo& DeviceInfo) const override;
 		bool GetDefaultOutputDeviceIndex(uint32& OutDefaultIndex) const override;
-		bool GetDefaultInputDeviceIndex(uint32& OutDefaultIndex) const override;
 		bool StartStream() override;
 		bool StopStream() override;
 		bool ShutdownStream() override;
-		bool GetLatency(uint32& OutputDeviceLatency, uint32& InputDeviceLatency) const override;
+		bool GetLatency(uint32& OutputDeviceLatency) const override;
 		bool GetFrameRate(uint32& OutFrameRate) const override;
 
 		// FRunnable
 		uint32 Run() override;
 		void Stop() override;
-		void Exit() override;
 
 	private:
-		struct FXAudio2Info
-		{
-			/** XAudio2 system object */
-			IXAudio2* XAudio2System;
 
-#if PLATFORM_WINDOWS
-			IMMDeviceEnumerator* DeviceEnumerator;
-#endif
+		/**  Simple struct to wrap a user buffer. */
+		struct FOutputBuffer
+		{
+			FOutputBuffer(int32 Size)
+			{
+				Buffer.Init(0.0f, Size);
+			}
+			TArray<float> Buffer;
+		};
 
+		/**
+		* FXAudio2Info
+		* 
+		* Structure for holding XAudio2 specific resources.
+		*/
+		struct FXAudio2Info
+		{
 			FXAudio2Info()
 				: XAudio2System(nullptr)
-#if PLATFORM_WINDOWS
-				, DeviceEnumerator(nullptr)
-#endif
+				, MasteringVoice(nullptr)
+				, OutputStreamSourceVoice(nullptr)
+				, OutputBufferEndEvent(nullptr)
+				, MaxQueuedBuffers(3)
+				, CurrentBufferIndex(0)
+				, bDeviceOpen(false)
 			{
 			}
+
+			/** XAudio2 system object */
+			IXAudio2* XAudio2System;
+
+			/** Mastering voice is the connection to a specific audio device. */
+			IXAudio2MasteringVoice* MasteringVoice;
+
+			/** Single source voice that will serve as our user-callback to platform-independent mixing code */
+			IXAudio2SourceVoice* OutputStreamSourceVoice;
+
+			/** Callback object to get notifications on buffer end. */
+			FXAudio2VoiceCallback OutputVoiceCallback;
+
+			/** Handle used to synchronize buffer end events and new buffers */
+			HANDLE OutputBufferEndEvent;
+
+			/** Max number of buffers to use for "ping-pong" buffer design (i.e. getting new user-buffer while writing out old buffer to hardware) */
+			int32 MaxQueuedBuffers;
+
+			/** Current index in the array of user buffers. */
+			int32 CurrentBufferIndex;
+
+			/** The array of user buffers */
+			TArray<FOutputBuffer> OutputBuffers;
+			bool bDeviceOpen;
 		};
 
 		/** XAudio2-specific data */
 		FXAudio2Info XAudio2Info;
 
-
 		/** Whether or not the device api has been initialized */
 		bool bInitialized;
 
@@ -119,10 +226,10 @@ namespace UAudio
 		bool bComInitialized;
 
 	private:
-		bool OpenDevices(const FCreateStreamParams& Params) override;
+		bool OpenDevice(const FCreateStreamParams& Params) override;
 
 		// Helper functions
-		EStreamFormat::Type GetStreamFormat(WAVEFORMATEX* WaveFormatEx) const;
+		void GetDeviceInfoFromFormat(const WAVEFORMATEX* WaveFormatEx, FDeviceInfo& DeviceInfo) const;
 
 	};
 
@@ -138,21 +245,17 @@ namespace UAudio
 
 	FUnrealAudioXAudio2::~FUnrealAudioXAudio2()
 	{
+		if (bInitialized)
+		{
+			Shutdown();
+		}
 	}
 
 	bool FUnrealAudioXAudio2::Initialize()
 	{
 		bComInitialized = FWindowsPlatformMisc::CoInitialize();
-
-		// Create the xaudio2 system
 		HRESULT Result = XAudio2Create(&XAudio2Info.XAudio2System);
 		RETURN_FALSE_ON_FAIL(Result);
-
-#if PLATFORM_WINDOWS
-		Result = CoCreateInstance(__uuidof(MMDeviceEnumerator), NULL, CLSCTX_ALL, IID_PPV_ARGS(&XAudio2Info.DeviceEnumerator));
-		RETURN_FALSE_ON_FAIL(Result);
-#endif
-
 		bInitialized = true;
 		return true;
 	}
@@ -170,7 +273,7 @@ namespace UAudio
 		{
 			FWindowsPlatformMisc::CoUninitialize();
 		}
-
+		bInitialized = false;
 		return true;
 	}
 
@@ -182,143 +285,179 @@ namespace UAudio
 
 	bool FUnrealAudioXAudio2::GetNumOutputDevices(uint32& OutNumDevices) const
 	{
-		if (!bInitialized || !XAudio2Info.XAudio2System)
+		if (!bInitialized)
 		{
 			return false;
 		}
 
-		// Note: XAudio2 only supports output devices, so GetDeviceCount() is only returning output devices
 		HRESULT Result = XAudio2Info.XAudio2System->GetDeviceCount(&OutNumDevices);
 		RETURN_FALSE_ON_FAIL(Result);
-
 		return true;
 	}
 
-	bool FUnrealAudioXAudio2::GetNumInputDevices(uint32& OutNumDevices) const
+	void FUnrealAudioXAudio2::GetDeviceInfoFromFormat(const WAVEFORMATEX* WaveFormatEx, FDeviceInfo& DeviceInfo) const
 	{
-#if PLATFORM_WINDOWS
-		// On windows we're going to use WASAPI/WINMM to do audio input since XAudio2 doesn't support capture devices
-		if (!bInitialized || !XAudio2Info.DeviceEnumerator)
-		{
-			return false;
-		}
-
-		IMMDeviceCollection* Devices = nullptr;
-		HRESULT Result = S_OK;
+		check(WaveFormatEx);
+		DeviceInfo.FrameRate = WaveFormatEx->nSamplesPerSec;
+		DeviceInfo.NumChannels = WaveFormatEx->nChannels;
 
-		Result = XAudio2Info.DeviceEnumerator->EnumAudioEndpoints(eCapture, DEVICE_STATE_ACTIVE, &Devices);
-		CLEANUP_ON_FAIL(Result);
-
-		Result = Devices->GetCount(&OutNumDevices);
-		CLEANUP_ON_FAIL(Result);
-
-	Cleanup:
-		SAFE_RELEASE(Devices);
-		return SUCCEEDED(Result);
-#else
-		// TODO: Figure out how to get input devices on XBOX api
-		// Look at using XInput and XAudio2 http://blogs.msdn.com/b/chuckw/archive/2012/05/03/xinput-and-xaudio2.aspx
-		OutNumDevices = 0;
-		return true;
-#endif
-	}
+		// XAudio2 supports converting formats to float so we don't need to do buffer conversion ourselves
+		DeviceInfo.StreamFormat = EStreamFormat::FLT;
 
-	EStreamFormat::Type FUnrealAudioXAudio2::GetStreamFormat(WAVEFORMATEX* WaveFormatEx) const
-	{
-		bool bIsFloat = (WaveFormatEx->wFormatTag == WAVE_FORMAT_IEEE_FLOAT);
-		bool bIsExtensible = (WaveFormatEx->wFormatTag == WAVE_FORMAT_EXTENSIBLE);
-		bool bIsSubtypeFloat = (bIsExtensible && (((WAVEFORMATEXTENSIBLE*)WaveFormatEx)->SubFormat == KSDATAFORMAT_SUBTYPE_IEEE_FLOAT));
-
-		if (bIsFloat || bIsSubtypeFloat)
+		if (WaveFormatEx->wFormatTag == WAVE_FORMAT_EXTENSIBLE)
 		{
-			if (WaveFormatEx->wBitsPerSample == 32)
+			// Build a flag field of the speaker outputs of this device
+			WAVEFORMATEXTENSIBLE* WaveFormatExtensible = (WAVEFORMATEXTENSIBLE*)WaveFormatEx;
+			DeviceInfo.Speakers.Reset();
+			uint32 ChanCount = 0;
+			for (uint32 SpeakerTypeIndex = 0; SpeakerTypeIndex < ESpeaker::SPEAKER_TYPE_COUNT, ChanCount < DeviceInfo.NumChannels; ++SpeakerTypeIndex)
 			{
-				return EStreamFormat::FLT;
-			}
-			else if (WaveFormatEx->wBitsPerSample == 64)
-			{
-				return EStreamFormat::DBL;
+				if (WaveFormatExtensible->dwChannelMask & SpeakerMaskMapping[SpeakerTypeIndex].XAudio2SpeakerTypeFlag)
+				{
+					++ChanCount;
+					DeviceInfo.Speakers.Add(SpeakerMaskMapping[SpeakerTypeIndex].UnrealSpeaker);
+				}
 			}
+			check(ChanCount == DeviceInfo.NumChannels);
 		}
 		else
 		{
-			bool bIsPcm = (WaveFormatEx->wFormatTag == WAVE_FORMAT_PCM);
-			bool bIsSubtypePcm = (bIsExtensible && (((WAVEFORMATEXTENSIBLE*)WaveFormatEx)->SubFormat == KSDATAFORMAT_SUBTYPE_PCM));
-			if (bIsPcm || bIsSubtypePcm)
+			// Non-extensible formats only support 1 or 2 channels
+			DeviceInfo.Speakers.Add(ESpeaker::FRONT_LEFT);
+			if (DeviceInfo.NumChannels == 2)
 			{
-				if (WaveFormatEx->wBitsPerSample == 16)
-				{
-					return EStreamFormat::INT_16;
-				}
-				else if (WaveFormatEx->wBitsPerSample == 24)
-				{
-					return EStreamFormat::INT_24;
-				}
-				else if (WaveFormatEx->wBitsPerSample == 32)
-				{
-					return EStreamFormat::INT_32;
-				}
+				DeviceInfo.Speakers.Add(ESpeaker::FRONT_RIGHT);
 			}
 		}
-		return EStreamFormat::UNSUPPORTED;
 	}
 
-	bool FUnrealAudioXAudio2::GetOutputDeviceInfo(const uint32 DeviceIndex, FDeviceInfo& OutInfo) const
+	bool FUnrealAudioXAudio2::GetOutputDeviceInfo(const uint32 DeviceIndex, FDeviceInfo& DeviceInfo) const
 	{
 		if (!bInitialized)
 		{
 			return false;
 		}
 
-		check(XAudio2Info.XAudio2System);
-
 		XAUDIO2_DEVICE_DETAILS DeviceDetails;
 		HRESULT Result = XAudio2Info.XAudio2System->GetDeviceDetails(DeviceIndex, &DeviceDetails);
 		RETURN_FALSE_ON_FAIL(Result);
 
-		OutInfo.FrameRate = DeviceDetails.OutputFormat.Format.nSamplesPerSec;
-		OutInfo.NumChannels = DeviceDetails.OutputFormat.Format.nChannels;
-		OutInfo.StreamFormat = GetStreamFormat(&DeviceDetails.OutputFormat.Format);
+		DeviceInfo.FriendlyName = FString(DeviceDetails.DisplayName);
+		DeviceInfo.bIsSystemDefault = (DeviceIndex == 0);
 
-		return true;
-	}
-
-	bool FUnrealAudioXAudio2::GetInputDeviceInfo(const uint32 DeviceIndex, FDeviceInfo& OutInfo) const
-	{
+		GetDeviceInfoFromFormat(&DeviceDetails.OutputFormat.Format, DeviceInfo);
 		return true;
 	}
 
 	bool FUnrealAudioXAudio2::GetDefaultOutputDeviceIndex(uint32& OutDefaultIndex) const
 	{
+		OutDefaultIndex = 0;
 		return true;
 	}
 
-	bool FUnrealAudioXAudio2::GetDefaultInputDeviceIndex(uint32& OutDefaultIndex) const
+	bool FUnrealAudioXAudio2::OpenDevice(const FCreateStreamParams& CreateStreamParams)
 	{
-		return true;
-	}
+		if (!bInitialized || XAudio2Info.bDeviceOpen)
+		{
+			return false;
+		}
+		check(CreateStreamParams.OutputDeviceIndex != INDEX_NONE);
+		check(XAudio2Info.XAudio2System && XAudio2Info.MasteringVoice == nullptr);
 
-	bool FUnrealAudioXAudio2::OpenDevices(const FCreateStreamParams& Params)
-	{
-		return true;
+		StreamInfo.BlockSize = CreateStreamParams.CallbackBlockSize;
+		StreamInfo.FrameRate = 44100;
+
+		HRESULT Result = S_OK;
+
+		FDeviceInfo DeviceInfo;
+		if (!GetOutputDeviceInfo(CreateStreamParams.OutputDeviceIndex, DeviceInfo))
+		{
+			Result = S_FALSE;
+			goto Cleanup;
+		}
+
+		Result = XAudio2Info.XAudio2System->CreateMasteringVoice(&XAudio2Info.MasteringVoice, DeviceInfo.NumChannels, StreamInfo.FrameRate, 0, CreateStreamParams.OutputDeviceIndex, nullptr);
+		CLEANUP_ON_FAIL(Result);
+
+		FStreamDeviceInfo& StreamDeviceInfo = StreamInfo.DeviceInfo;
+		StreamDeviceInfo.DeviceIndex = CreateStreamParams.OutputDeviceIndex;
+		StreamDeviceInfo.bPerformByteSwap = false;
+		StreamDeviceInfo.NumChannels = DeviceInfo.NumChannels;
+		StreamDeviceInfo.DeviceDataFormat = DeviceInfo.StreamFormat;
+		StreamDeviceInfo.FrameRate = DeviceInfo.FrameRate;
+		StreamDeviceInfo.Speakers = DeviceInfo.Speakers;
+
+		uint32 BufferSize = StreamDeviceInfo.NumChannels * StreamInfo.BlockSize * GetNumBytesForFormat(EStreamFormat::FLT);
+		StreamDeviceInfo.UserBuffer.Init(0, BufferSize);
+		XAudio2Info.bDeviceOpen = true;
+		XAudio2Info.XAudio2System->StartEngine();
+
+	Cleanup:
+		if (FAILED(Result))
+		{
+			ShutdownStream();
+		}
+		return SUCCEEDED(Result);
 	}
 
 	bool FUnrealAudioXAudio2::StartStream()
 	{
+		check(StreamInfo.Thread == nullptr);
+		StreamInfo.Thread = FRunnableThread::Create(this, TEXT("XAudio2DeviceThread"), 0, TPri_AboveNormal);
 		return true;
 	}
 
 	bool FUnrealAudioXAudio2::StopStream()
 	{
+		if (!bInitialized || StreamInfo.State == EStreamState::CLOSED || StreamInfo.State == EStreamState::STOPPED || StreamInfo.State == EStreamState::STOPPING)
+		{
+			return false;
+		}
+		check(StreamInfo.Thread != nullptr);
+		check(XAudio2Info.XAudio2System != nullptr);
+
+		// Tell the stream update that we are stopping
+		Stop();
+
+		// Wait for the thread to finish
+		StreamInfo.Thread->WaitForCompletion();
+
+		check(StreamInfo.State == EStreamState::STOPPED);
+
+		// delete the thread
+		delete StreamInfo.Thread;
+		StreamInfo.Thread = nullptr;
+
 		return true;
 	}
 
 	bool FUnrealAudioXAudio2::ShutdownStream()
 	{
+		if (!bInitialized || StreamInfo.State == EStreamState::CLOSED)
+		{
+			return false;
+		}
+
+		if (StreamInfo.State != EStreamState::STOPPED)
+		{
+			StopStream();
+		}
+
+		XAudio2Info.XAudio2System->StopEngine();
+		XAudio2Info.OutputStreamSourceVoice->DestroyVoice();
+		XAudio2Info.OutputStreamSourceVoice = nullptr;
+		XAudio2Info.MasteringVoice->DestroyVoice();
+		XAudio2Info.MasteringVoice = nullptr;
+		XAudio2Info.bDeviceOpen = false;
+
+		CloseHandle(XAudio2Info.OutputBufferEndEvent);
+		XAudio2Info.OutputBufferEndEvent = nullptr;
+
+		StreamInfo.State = EStreamState::CLOSED;
 		return true;
 	}
 
-	bool FUnrealAudioXAudio2::GetLatency(uint32& OutputDeviceLatency, uint32& InputDeviceLatency) const
+	bool FUnrealAudioXAudio2::GetLatency(uint32& OutputDeviceLatency) const
 	{
 		return true;
 	}
@@ -328,20 +467,125 @@ namespace UAudio
 		return true;
 	}
 
-	// FRunnable
-
 	uint32 FUnrealAudioXAudio2::Run()
 	{
+		bool bThreadComInitialized = FWindowsPlatformMisc::CoInitialize();
+
+		FCallbackInfo CallbackInfo;
+
+		FStreamDeviceInfo& OutputDeviceInfo = StreamInfo.DeviceInfo;
+		uint32 OutputDeviceSamples			= StreamInfo.BlockSize * OutputDeviceInfo.NumChannels;
+		HRESULT Result						= S_OK;
+		WAVEFORMATEX Format					= { 0 };
+
+		check(XAudio2Info.OutputBufferEndEvent == nullptr);
+		XAudio2Info.OutputBufferEndEvent = CreateEvent(0, false, false, nullptr);
+		if (!XAudio2Info.OutputBufferEndEvent)
+		{
+			Result = HRESULT_FROM_WIN32(GetLastError());
+			CLEANUP_ON_FAIL(Result);
+		}
+
+		XAudio2Info.OutputVoiceCallback.SetBufferEndEvent(XAudio2Info.OutputBufferEndEvent);
+
+		Format.nChannels		= OutputDeviceInfo.NumChannels;
+		Format.nSamplesPerSec	= StreamInfo.FrameRate;
+		Format.wFormatTag		= WAVE_FORMAT_IEEE_FLOAT;
+		Format.nAvgBytesPerSec	= Format.nSamplesPerSec * sizeof(float) * Format.nChannels;
+		Format.nBlockAlign		= sizeof(float) * Format.nChannels;
+		Format.wBitsPerSample	= sizeof(float) * 8;
+
+		Result = XAudio2Info.XAudio2System->CreateSourceVoice(
+			&XAudio2Info.OutputStreamSourceVoice,
+			&Format,
+			XAUDIO2_VOICE_NOSRC | XAUDIO2_VOICE_NOPITCH,
+			2.0f,
+			&XAudio2Info.OutputVoiceCallback
+		);
+		CLEANUP_ON_FAIL(Result);
+
+		XAudio2Info.OutputBuffers.Empty();
+		for (int32 Index = 0; Index < XAudio2Info.MaxQueuedBuffers; ++Index)
+		{
+			XAudio2Info.OutputBuffers.Add(FOutputBuffer(OutputDeviceSamples));
+		}
+
+		// Setup the callback info struct
+		CallbackInfo.OutBuffer		= nullptr;
+		CallbackInfo.NumFrames		= StreamInfo.BlockSize;
+		CallbackInfo.NumChannels	= OutputDeviceInfo.NumChannels;
+		CallbackInfo.StreamTime		= 0.0;
+		CallbackInfo.UserData		= StreamInfo.UserData;
+		CallbackInfo.StatusFlags	= 0;
+		CallbackInfo.OutputSpeakers = OutputDeviceInfo.Speakers;
+		CallbackInfo.FrameRate		= StreamInfo.FrameRate;
+
+		check(XAudio2Info.OutputStreamSourceVoice);
+		XAudio2Info.OutputStreamSourceVoice->Start();
+
+		StreamInfo.State = EStreamState::RUNNING;
+
+		while (StreamInfo.State != EStreamState::STOPPING)
+		{
+			CallbackInfo.StatusFlags = 0;
+			CallbackInfo.StreamTime = StreamInfo.StreamTime;
+			CallbackInfo.OutBuffer = XAudio2Info.OutputBuffers[XAudio2Info.CurrentBufferIndex].Buffer.GetData();
+
+#if UNREAL_AUDIO_TEST_WHITE_NOISE
+			uint32 NumSamples = StreamInfo.BlockSize*OutputDeviceInfo.NumChannels;
+			for (uint32 Sample = 0; Sample < NumSamples; ++Sample)
+			{
+				CallbackInfo.OutBuffer[Sample] = 0.5f * FMath::FRandRange(-1.0f, 1.0f);
+			}
+#else // #if UNREAL_AUDIO_TEST_WHITE_NOISE
+			memset((void*)CallbackInfo.OutBuffer, 0, StreamInfo.BlockSize*OutputDeviceInfo.NumChannels*sizeof(float));
+			if (!StreamInfo.CallbackFunction(CallbackInfo))
+			{
+				StreamInfo.State = EStreamState::STOPPING;
+			}
+#endif // #else // #if UNREAL_AUDIO_TEST_WHITE_NOISE
+
+			UpdateStreamTimeTick();
+
+			// If we read input but our xaudio2 voice has the max number of voices queued to playback, 
+			// then wait for the next buffer to finish playing out before writing into it new buffers
+			XAUDIO2_VOICE_STATE OutputVoiceState;
+			XAudio2Info.OutputStreamSourceVoice->GetState(&OutputVoiceState);
+			if (OutputVoiceState.BuffersQueued >= ((uint32)XAudio2Info.MaxQueuedBuffers - 1))
+			{
+				WaitForSingleObject(XAudio2Info.OutputBufferEndEvent, INFINITE);
+			}
+
+			float* Buffer = XAudio2Info.OutputBuffers[XAudio2Info.CurrentBufferIndex].Buffer.GetData();
+			// Write out the last user-buffer to the device
+			XAUDIO2_BUFFER SubmitBuffer = { 0 };
+			SubmitBuffer.AudioBytes = OutputDeviceSamples * sizeof(float);
+			SubmitBuffer.pAudioData = (const BYTE*)Buffer;
+			XAudio2Info.OutputStreamSourceVoice->SubmitSourceBuffer(&SubmitBuffer);
+
+			// Update the current buffer index (the last buffer is currently being written to device)
+			XAudio2Info.CurrentBufferIndex = (XAudio2Info.CurrentBufferIndex + 1) % XAudio2Info.MaxQueuedBuffers;
+		}
+
+		XAudio2Info.OutputStreamSourceVoice->Stop();
+		XAudio2Info.OutputStreamSourceVoice->FlushSourceBuffers();
+
+	Cleanup:
+		if (bThreadComInitialized)
+		{
+			FWindowsPlatformMisc::CoUninitialize();
+		}
+
+		StreamInfo.State = EStreamState::STOPPED;
 		return 0;
 	}
 
 	void FUnrealAudioXAudio2::Stop()
 	{
-	}
-
-	void FUnrealAudioXAudio2::Exit()
-	{
-
+		if (StreamInfo.State != EStreamState::STOPPED)
+		{
+			StreamInfo.State = EStreamState::STOPPING;
+		}
 	}
 }
 
