UE4 - Demacroize the RHI, step 1. Eliminate DYNAMIC_CAST_D3D11RESOURCE et al

diff --git a/Engine/Source/Runtime/EmptyRHI/Private/EmptyCommands.cpp b/Engine/Source/Runtime/EmptyRHI/Private/EmptyCommands.cpp
--- a/Engine/Source/Runtime/EmptyRHI/Private/EmptyCommands.cpp
+++ b/Engine/Source/Runtime/EmptyRHI/Private/EmptyCommands.cpp
@@ -20,7 +20,7 @@ void FEmptyDynamicRHI::RHIGpuTimeEnd(uint32 Hash, bool bCompute)
 
 void FEmptyDynamicRHI::RHISetStreamSource(uint32 StreamIndex,FVertexBufferRHIParamRef VertexBufferRHI,uint32 Stride,uint32 Offset)
 {
-	DYNAMIC_CAST_EMPTYRESOURCE(VertexBuffer,VertexBuffer);
+	FEmptyVertexBuffer* VertexBuffer = ResourceCast(VertexBufferRHI);
 
 }
 
@@ -31,13 +31,13 @@ void FEmptyDynamicRHI::RHISetStreamOutTargets(uint32 NumTargets, const FVertexBu
 
 void FEmptyDynamicRHI::RHISetRasterizerState(FRasterizerStateRHIParamRef NewStateRHI)
 {
-	DYNAMIC_CAST_EMPTYRESOURCE(RasterizerState,NewState);
+	FEmptyRasterizerState* NewState = ResourceCast(NewStateRHI);
 
 }
 
 void FEmptyDynamicRHI::RHISetComputeShader(FComputeShaderRHIParamRef ComputeShaderRHI)
 {
-	DYNAMIC_CAST_EMPTYRESOURCE(ComputeShader,ComputeShader);
+	FEmptyComputeShader* ComputeShader = ResourceCast(ComputeShaderRHI);
 
 }
 
@@ -48,7 +48,7 @@ void FEmptyDynamicRHI::RHIDispatchComputeShader(uint32 ThreadGroupCountX, uint32
 
 void FEmptyDynamicRHI::RHIDispatchIndirectComputeShader(FVertexBufferRHIParamRef ArgumentBufferRHI, uint32 ArgumentOffset) 
 { 
-	DYNAMIC_CAST_EMPTYRESOURCE(VertexBuffer,ArgumentBuffer);
+	FEmptyVertexBuffer* ArgumentBuffer = ResourceCast(ArgumentBufferRHI);
 
 }
 
@@ -69,20 +69,20 @@ void FEmptyDynamicRHI::RHISetScissorRect(bool bEnable,uint32 MinX,uint32 MinY,ui
 
 void FEmptyDynamicRHI::RHISetBoundShaderState( FBoundShaderStateRHIParamRef BoundShaderStateRHI)
 {
-	DYNAMIC_CAST_EMPTYRESOURCE(BoundShaderState,BoundShaderState);
+	FEmptyBoundShaderState* BoundShaderState = ResourceCast(BoundShaderStateRHI);
 
 }
 
 
 void FEmptyDynamicRHI::RHISetUAVParameter(FComputeShaderRHIParamRef ComputeShaderRHI, uint32 UAVIndex, FUnorderedAccessViewRHIParamRef UAVRHI)
 {
-	DYNAMIC_CAST_EMPTYRESOURCE(UnorderedAccessView, UAV);
+	FEmptyUnorderedAccessView* UAV = ResourceCast(UAVRHI);
 
 }
 
 void FEmptyDynamicRHI::RHISetUAVParameter(FComputeShaderRHIParamRef ComputeShaderRHI,uint32 UAVIndex,FUnorderedAccessViewRHIParamRef UAVRHI, uint32 InitialCount)
 {
-	DYNAMIC_CAST_EMPTYRESOURCE(UnorderedAccessView, UAV);
+	FEmptyUnorderedAccessView* UAV = ResourceCast(UAVRHI);
 
 }
 
@@ -151,38 +151,38 @@ void FEmptyDynamicRHI::RHISetShaderResourceViewParameter(FComputeShaderRHIParamR
 
 void FEmptyDynamicRHI::RHISetShaderSampler(FVertexShaderRHIParamRef VertexShaderRHI, uint32 SamplerIndex, FSamplerStateRHIParamRef NewStateRHI)
 {
-	DYNAMIC_CAST_EMPTYRESOURCE(SamplerState,NewState);
-	DYNAMIC_CAST_EMPTYRESOURCE(VertexShader,VertexShader);
+	FEmptySamplerState* NewState = ResourceCast(NewStateRHI);
+	FEmptyVertexShader* VertexShader = ResourceCast(VertexShaderRHI);
 
 }
 
 void FEmptyDynamicRHI::RHISetShaderSampler(FHullShaderRHIParamRef HullShader, uint32 SamplerIndex, FSamplerStateRHIParamRef NewStateRHI)
 {
-	DYNAMIC_CAST_EMPTYRESOURCE(SamplerState,NewState);
+	FEmptySamplerState* NewState = ResourceCast(NewStateRHI);
 
 }
 
 void FEmptyDynamicRHI::RHISetShaderSampler(FDomainShaderRHIParamRef DomainShader, uint32 SamplerIndex, FSamplerStateRHIParamRef NewStateRHI)
 {
-	DYNAMIC_CAST_EMPTYRESOURCE(SamplerState,NewState);
+	FEmptySamplerState* NewState = ResourceCast(NewStateRHI);
 
 }
 
 void FEmptyDynamicRHI::RHISetShaderSampler(FGeometryShaderRHIParamRef GeometryShader, uint32 SamplerIndex, FSamplerStateRHIParamRef NewStateRHI)
 {
-	DYNAMIC_CAST_EMPTYRESOURCE(SamplerState,NewState);
+	FEmptySamplerState* NewState = ResourceCast(NewStateRHI);
 
 }
 
 void FEmptyDynamicRHI::RHISetShaderSampler(FPixelShaderRHIParamRef PixelShader, uint32 SamplerIndex, FSamplerStateRHIParamRef NewStateRHI)
 {
-	DYNAMIC_CAST_EMPTYRESOURCE(SamplerState,NewState);
+	FEmptySamplerState* NewState = ResourceCast(NewStateRHI);
 
 }
 
 void FEmptyDynamicRHI::RHISetShaderSampler(FComputeShaderRHIParamRef ComputeShader, uint32 SamplerIndex, FSamplerStateRHIParamRef NewStateRHI)
 {
-	DYNAMIC_CAST_EMPTYRESOURCE(SamplerState,NewState);
+	FEmptySamplerState* NewState = ResourceCast(NewStateRHI);
 
 }
 
@@ -219,46 +219,46 @@ void FEmptyDynamicRHI::RHISetShaderParameter(FComputeShaderRHIParamRef ComputeSh
 
 void FEmptyDynamicRHI::RHISetShaderUniformBuffer(FVertexShaderRHIParamRef VertexShaderRHI, uint32 BufferIndex, FUniformBufferRHIParamRef BufferRHI)
 {
-	DYNAMIC_CAST_EMPTYRESOURCE(UniformBuffer,Buffer);
-	DYNAMIC_CAST_EMPTYRESOURCE(VertexShader,VertexShader);
+	FEmptyUniformBuffer* Buffer = ResourceCast(BufferRHI);
+	FEmptyVertexShader* VertexShader = ResourceCast(VertexShaderRHI);
 
 }
 
 void FEmptyDynamicRHI::RHISetShaderUniformBuffer(FHullShaderRHIParamRef HullShader, uint32 BufferIndex, FUniformBufferRHIParamRef BufferRHI)
 {
-	DYNAMIC_CAST_EMPTYRESOURCE(UniformBuffer,Buffer);
+	FEmptyUniformBuffer* Buffer = ResourceCast(BufferRHI);
 }
 
 void FEmptyDynamicRHI::RHISetShaderUniformBuffer(FDomainShaderRHIParamRef DomainShader, uint32 BufferIndex, FUniformBufferRHIParamRef BufferRHI)
 {
-	DYNAMIC_CAST_EMPTYRESOURCE(UniformBuffer,Buffer);
+	FEmptyUniformBuffer* Buffer = ResourceCast(BufferRHI);
 }
 
 void FEmptyDynamicRHI::RHISetShaderUniformBuffer(FGeometryShaderRHIParamRef GeometryShader, uint32 BufferIndex, FUniformBufferRHIParamRef BufferRHI)
 {
-	DYNAMIC_CAST_EMPTYRESOURCE(UniformBuffer,Buffer);
+	FEmptyUniformBuffer* Buffer = ResourceCast(BufferRHI);
 }
 
 void FEmptyDynamicRHI::RHISetShaderUniformBuffer(FPixelShaderRHIParamRef PixelShader, uint32 BufferIndex, FUniformBufferRHIParamRef BufferRHI)
 {
-	DYNAMIC_CAST_EMPTYRESOURCE(UniformBuffer,Buffer);
+	FEmptyUniformBuffer* Buffer = ResourceCast(BufferRHI);
 }
 
 void FEmptyDynamicRHI::RHISetShaderUniformBuffer(FComputeShaderRHIParamRef ComputeShader, uint32 BufferIndex, FUniformBufferRHIParamRef BufferRHI)
 {
-	DYNAMIC_CAST_EMPTYRESOURCE(UniformBuffer,Buffer);
+	FEmptyUniformBuffer* Buffer = ResourceCast(BufferRHI);
 }
 
 
 void FEmptyDynamicRHI::RHISetDepthStencilState(FDepthStencilStateRHIParamRef NewStateRHI, uint32 StencilRef)
 {
-	DYNAMIC_CAST_EMPTYRESOURCE(DepthStencilState,NewState);
+	FEmptyDepthStencilState* NewState = ResourceCast(NewStateRHI);
 
 }
 
 void FEmptyDynamicRHI::RHISetBlendState(FBlendStateRHIParamRef NewStateRHI, const FLinearColor& BlendFactor)
 {
-	DYNAMIC_CAST_EMPTYRESOURCE(BlendState,NewState);
+	FEmptyBlendState* NewState = ResourceCast(NewStateRHI);
 
 }
 
@@ -280,14 +280,14 @@ void FEmptyDynamicRHI::RHISetRenderTargetsAndClear(const FRHISetRenderTargetsInf
 // Occlusion/Timer queries.
 void FEmptyDynamicRHI::RHIBeginRenderQuery(FRenderQueryRHIParamRef QueryRHI)
 {
-	DYNAMIC_CAST_EMPTYRESOURCE(RenderQuery,Query);
+	FEmptyRenderQuery* Query = ResourceCast(QueryRHI);
 
 	Query->Begin();
 }
 
 void FEmptyDynamicRHI::RHIEndRenderQuery(FRenderQueryRHIParamRef QueryRHI)
 {
-	DYNAMIC_CAST_EMPTYRESOURCE(RenderQuery,Query);
+	FEmptyRenderQuery* Query = ResourceCast(QueryRHI);
 
 	Query->End();
 }
@@ -299,7 +299,7 @@ void FEmptyDynamicRHI::RHIDrawPrimitive(uint32 PrimitiveType, uint32 BaseVertexI
 
 void FEmptyDynamicRHI::RHIDrawPrimitiveIndirect(uint32 PrimitiveType, FVertexBufferRHIParamRef ArgumentBufferRHI, uint32 ArgumentOffset)
 {
-	DYNAMIC_CAST_EMPTYRESOURCE(VertexBuffer,ArgumentBuffer);
+	FEmptyVertexBuffer* ArgumentBuffer = ResourceCast(ArgumentBufferRHI);
 
 }
 
@@ -307,21 +307,21 @@ void FEmptyDynamicRHI::RHIDrawPrimitiveIndirect(uint32 PrimitiveType, FVertexBuf
 void FEmptyDynamicRHI::RHIDrawIndexedPrimitive(FIndexBufferRHIParamRef IndexBufferRHI, uint32 PrimitiveType, int32 BaseVertexIndex, uint32 FirstInstance,
 	uint32 NumVertices, uint32 StartIndex, uint32 NumPrimitives, uint32 NumInstances)
 {
-	DYNAMIC_CAST_EMPTYRESOURCE(IndexBuffer,IndexBuffer);
+	FEmptyIndexBuffer* IndexBuffer = ResourceCast(IndexBufferRHI);
 
 }
 
 void FEmptyDynamicRHI::RHIDrawIndexedIndirect(FIndexBufferRHIParamRef IndexBufferRHI, uint32 PrimitiveType, FStructuredBufferRHIParamRef ArgumentsBufferRHI, int32 DrawArgumentsIndex, uint32 NumInstances)
 {
-	DYNAMIC_CAST_EMPTYRESOURCE(IndexBuffer,IndexBuffer);
-	DYNAMIC_CAST_EMPTYRESOURCE(StructuredBuffer,ArgumentsBuffer);
+	FEmptyIndexBuffer* IndexBuffer = ResourceCast(IndexBufferRHI);
+	FEmptyStructuredBuffer* ArgumentsBuffer = ResourceCast(ArgumentsBufferRHI);
 
 }
 
 void FEmptyDynamicRHI::RHIDrawIndexedPrimitiveIndirect(uint32 PrimitiveType,FIndexBufferRHIParamRef IndexBufferRHI,FVertexBufferRHIParamRef ArgumentBufferRHI,uint32 ArgumentOffset)
 {
-	DYNAMIC_CAST_EMPTYRESOURCE(IndexBuffer,IndexBuffer);
-	DYNAMIC_CAST_EMPTYRESOURCE(VertexBuffer,ArgumentBuffer);
+	FEmptyIndexBuffer* IndexBuffer = ResourceCast(IndexBufferRHI);
+	FEmptyVertexBuffer* ArgumentBuffer = ResourceCast(ArgumentBufferRHI);
 
 }
 
diff --git a/Engine/Source/Runtime/EmptyRHI/Private/EmptyIndexBuffer.cpp b/Engine/Source/Runtime/EmptyRHI/Private/EmptyIndexBuffer.cpp
--- a/Engine/Source/Runtime/EmptyRHI/Private/EmptyIndexBuffer.cpp
+++ b/Engine/Source/Runtime/EmptyRHI/Private/EmptyIndexBuffer.cpp
@@ -49,14 +49,14 @@ FIndexBufferRHIRef FEmptyDynamicRHI::RHICreateIndexBuffer(uint32 Stride, uint32
 
 void* FEmptyDynamicRHI::RHILockIndexBuffer(FIndexBufferRHIParamRef IndexBufferRHI, uint32 Offset, uint32 Size, EResourceLockMode LockMode)
 {
-	DYNAMIC_CAST_EMPTYRESOURCE(IndexBuffer,IndexBuffer);
+	FEmptyIndexBuffer* IndexBuffer = ResourceCast(IndexBufferRHI);
 
 	return (uint8*)IndexBuffer->Lock(LockMode, Size) + Offset;
 }
 
 void FEmptyDynamicRHI::RHIUnlockIndexBuffer(FIndexBufferRHIParamRef IndexBufferRHI)
 {
-	DYNAMIC_CAST_EMPTYRESOURCE(IndexBuffer,IndexBuffer);
+	FEmptyIndexBuffer* IndexBuffer = ResourceCast(IndexBufferRHI);
 
 	IndexBuffer->Unlock();
 }
diff --git a/Engine/Source/Runtime/EmptyRHI/Private/EmptyQuery.cpp b/Engine/Source/Runtime/EmptyRHI/Private/EmptyQuery.cpp
--- a/Engine/Source/Runtime/EmptyRHI/Private/EmptyQuery.cpp
+++ b/Engine/Source/Runtime/EmptyRHI/Private/EmptyQuery.cpp
@@ -41,7 +41,7 @@ bool FEmptyDynamicRHI::RHIGetRenderQueryResult(FRenderQueryRHIParamRef QueryRHI,
 {
 	check(IsInRenderingThread());
 
-	DYNAMIC_CAST_EMPTYRESOURCE(RenderQuery,Query);
+	FEmptyRenderQuery* Query = ResourceCast(QueryRHI);
 
 	return false;
 }
diff --git a/Engine/Source/Runtime/EmptyRHI/Private/EmptyShaders.cpp b/Engine/Source/Runtime/EmptyRHI/Private/EmptyShaders.cpp
--- a/Engine/Source/Runtime/EmptyRHI/Private/EmptyShaders.cpp
+++ b/Engine/Source/Runtime/EmptyRHI/Private/EmptyShaders.cpp
@@ -76,12 +76,12 @@ FEmptyBoundShaderState::FEmptyBoundShaderState(
 			FGeometryShaderRHIParamRef InGeometryShaderRHI)
 	:	CacheLink(InVertexDeclarationRHI,InVertexShaderRHI,InPixelShaderRHI,InHullShaderRHI,InDomainShaderRHI,InGeometryShaderRHI,this)
 {
-	DYNAMIC_CAST_EMPTYRESOURCE(VertexDeclaration,InVertexDeclaration);
-	DYNAMIC_CAST_EMPTYRESOURCE(VertexShader,InVertexShader);
-	DYNAMIC_CAST_EMPTYRESOURCE(PixelShader,InPixelShader);
-	DYNAMIC_CAST_EMPTYRESOURCE(HullShader,InHullShader);
-	DYNAMIC_CAST_EMPTYRESOURCE(DomainShader,InDomainShader);
-	DYNAMIC_CAST_EMPTYRESOURCE(GeometryShader,InGeometryShader);
+	FEmptyVertexDeclaration* InVertexDeclaration = ResourceCast(InVertexDeclarationRHI);
+	FEmptyVertexShader* InVertexShader = ResourceCast(InVertexShaderRHI);
+	FEmptyPixelShader* InPixelShader = ResourceCast(InPixelShaderRHI);
+	FEmptyHullShader* InHullShader = ResourceCast(InHullShaderRHI);
+	FEmptyDomainShader* InDomainShader = ResourceCast(InDomainShaderRHI);
+	FEmptyGeometryShader* InGeometryShader = ResourceCast(InGeometryShaderRHI);
 
 	// cache everything
 	VertexDeclaration = InVertexDeclaration;
diff --git a/Engine/Source/Runtime/EmptyRHI/Private/EmptyStructuredBuffer.cpp b/Engine/Source/Runtime/EmptyRHI/Private/EmptyStructuredBuffer.cpp
--- a/Engine/Source/Runtime/EmptyRHI/Private/EmptyStructuredBuffer.cpp
+++ b/Engine/Source/Runtime/EmptyRHI/Private/EmptyStructuredBuffer.cpp
@@ -32,7 +32,7 @@ FStructuredBufferRHIRef FEmptyDynamicRHI::RHICreateStructuredBuffer(uint32 Strid
 
 void* FEmptyDynamicRHI::RHILockStructuredBuffer(FStructuredBufferRHIParamRef StructuredBufferRHI,uint32 Offset,uint32 Size,EResourceLockMode LockMode)
 {
-	DYNAMIC_CAST_EMPTYRESOURCE(StructuredBuffer,StructuredBuffer);
+	FEmptyStructuredBuffer* StructuredBuffer = ResourceCast(StructuredBufferRHI);
 
 	return NULL;
 }
diff --git a/Engine/Source/Runtime/EmptyRHI/Private/EmptyTexture.cpp b/Engine/Source/Runtime/EmptyRHI/Private/EmptyTexture.cpp
--- a/Engine/Source/Runtime/EmptyRHI/Private/EmptyTexture.cpp
+++ b/Engine/Source/Runtime/EmptyRHI/Private/EmptyTexture.cpp
@@ -149,7 +149,7 @@ void FEmptyDynamicRHI::RHIGenerateMips(FTextureRHIParamRef SourceSurfaceRHI)
 
 FTexture2DRHIRef FEmptyDynamicRHI::RHIAsyncReallocateTexture2D(FTexture2DRHIParamRef OldTextureRHI, int32 NewMipCount, int32 NewSizeX, int32 NewSizeY, FThreadSafeCounter* RequestStatus)
 {
-	DYNAMIC_CAST_EMPTYRESOURCE(Texture2D,OldTexture);
+	FEmptyTexture2D* OldTexture = ResourceCast(OldTextureRHI);
 
 	return NULL;
 }
@@ -167,37 +167,37 @@ ETextureReallocationStatus FEmptyDynamicRHI::RHICancelAsyncReallocateTexture2D(
 
 void* FEmptyDynamicRHI::RHILockTexture2D(FTexture2DRHIParamRef TextureRHI,uint32 MipIndex,EResourceLockMode LockMode,uint32& DestStride,bool bLockWithinMiptail)
 {
-	DYNAMIC_CAST_EMPTYRESOURCE(Texture2D,Texture);
+	FEmptyTexture2D* Texture = ResourceCast(TextureRHI);
 	return Texture->Surface.Lock(MipIndex, 0, LockMode, DestStride);
 }
 
 void FEmptyDynamicRHI::RHIUnlockTexture2D(FTexture2DRHIParamRef TextureRHI,uint32 MipIndex,bool bLockWithinMiptail)
 {
-	DYNAMIC_CAST_EMPTYRESOURCE(Texture2D,Texture);
+	FEmptyTexture2D* Texture = ResourceCast(TextureRHI);
 	Texture->Surface.Unlock(MipIndex, 0);
 }
 
 void* FEmptyDynamicRHI::RHILockTexture2DArray(FTexture2DArrayRHIParamRef TextureRHI,uint32 TextureIndex,uint32 MipIndex,EResourceLockMode LockMode,uint32& DestStride,bool bLockWithinMiptail)
 {
-	DYNAMIC_CAST_EMPTYRESOURCE(Texture2DArray,Texture);
+	FEmptyTexture2DArray* Texture = ResourceCast(TextureRHI);
 	return Texture->Surface.Lock(MipIndex, TextureIndex, LockMode, DestStride);
 }
 
 void FEmptyDynamicRHI::RHIUnlockTexture2DArray(FTexture2DArrayRHIParamRef TextureRHI,uint32 TextureIndex,uint32 MipIndex,bool bLockWithinMiptail)
 {
-	DYNAMIC_CAST_EMPTYRESOURCE(Texture2DArray,Texture);
+	FEmptyTexture2DArray* Texture = ResourceCast(TextureRHI);
 	Texture->Surface.Unlock(MipIndex, TextureIndex);
 }
 
 void FEmptyDynamicRHI::RHIUpdateTexture2D(FTexture2DRHIParamRef TextureRHI, uint32 MipIndex, const struct FUpdateTextureRegion2D& UpdateRegion, uint32 SourcePitch, const uint8* SourceData)
 {	
-	DYNAMIC_CAST_EMPTYRESOURCE(Texture3D,Texture);
+	FEmptyTexture3D* Texture = ResourceCast(TextureRHI);
 
 }
 
 void FEmptyDynamicRHI::RHIUpdateTexture3D(FTexture3DRHIParamRef TextureRHI,uint32 MipIndex,const FUpdateTextureRegion3D& UpdateRegion,uint32 SourceRowPitch,uint32 SourceDepthPitch,const uint8* SourceData)
 {	
-	DYNAMIC_CAST_EMPTYRESOURCE(Texture3D,Texture);
+	FEmptyTexture3D* Texture = ResourceCast(TextureRHI);
 
 }
 
@@ -217,14 +217,14 @@ FTextureCubeRHIRef FEmptyDynamicRHI::RHICreateTextureCubeArray(uint32 Size, uint
 
 void* FEmptyDynamicRHI::RHILockTextureCubeFace(FTextureCubeRHIParamRef TextureCubeRHI,uint32 FaceIndex,uint32 ArrayIndex,uint32 MipIndex,EResourceLockMode LockMode,uint32& DestStride,bool bLockWithinMiptail)
 {
-	DYNAMIC_CAST_EMPTYRESOURCE(TextureCube,TextureCube);
+	FEmptyTextureCube* TextureCube = ResourceCast(TextureCubeRHI);
 	uint32 EmptyFace = GetEmptyCubeFace((ECubeFace)FaceIndex);
 	return TextureCube->Surface.Lock(MipIndex, FaceIndex + 6 * ArrayIndex, LockMode, DestStride);
 }
 
 void FEmptyDynamicRHI::RHIUnlockTextureCubeFace(FTextureCubeRHIParamRef TextureCubeRHI,uint32 FaceIndex,uint32 ArrayIndex,uint32 MipIndex,bool bLockWithinMiptail)
 {
-	DYNAMIC_CAST_EMPTYRESOURCE(TextureCube,TextureCube);
+	FEmptyTextureCube* TextureCube = ResourceCast(TextureCubeRHI);
 	uint32 EmptyFace = GetEmptyCubeFace((ECubeFace)FaceIndex);
 	TextureCube->Surface.Unlock(MipIndex, FaceIndex + ArrayIndex * 6);
 }
diff --git a/Engine/Source/Runtime/EmptyRHI/Private/EmptyUAV.cpp b/Engine/Source/Runtime/EmptyRHI/Private/EmptyUAV.cpp
--- a/Engine/Source/Runtime/EmptyRHI/Private/EmptyUAV.cpp
+++ b/Engine/Source/Runtime/EmptyRHI/Private/EmptyUAV.cpp
@@ -13,7 +13,7 @@ FEmptyShaderResourceView::~FEmptyShaderResourceView()
 
 FUnorderedAccessViewRHIRef FEmptyDynamicRHI::RHICreateUnorderedAccessView(FStructuredBufferRHIParamRef StructuredBufferRHI, bool bUseUAVCounter, bool bAppendBuffer)
 {
-	DYNAMIC_CAST_EMPTYRESOURCE(StructuredBuffer,StructuredBuffer);
+	FEmptyStructuredBuffer* StructuredBuffer = ResourceCast(StructuredBufferRHI);
 
 	// create the UAV buffer to point to the structured buffer's memory
 	FEmptyUnorderedAccessView* UAV = new FEmptyUnorderedAccessView;
@@ -35,7 +35,7 @@ FUnorderedAccessViewRHIRef FEmptyDynamicRHI::RHICreateUnorderedAccessView(FTextu
 
 FUnorderedAccessViewRHIRef FEmptyDynamicRHI::RHICreateUnorderedAccessView(FVertexBufferRHIParamRef VertexBufferRHI, uint8 Format)
 {
-	DYNAMIC_CAST_EMPTYRESOURCE(VertexBuffer, VertexBuffer);
+	FEmptyVertexBuffer* VertexBuffer = ResourceCast(VertexBufferRHI);
 
 	// create the UAV buffer to point to the structured buffer's memory
 	FEmptyUnorderedAccessView* UAV = new FEmptyUnorderedAccessView;
@@ -46,7 +46,7 @@ FUnorderedAccessViewRHIRef FEmptyDynamicRHI::RHICreateUnorderedAccessView(FVerte
 
 FShaderResourceViewRHIRef FEmptyDynamicRHI::RHICreateShaderResourceView(FStructuredBufferRHIParamRef StructuredBufferRHI)
 {
-	DYNAMIC_CAST_EMPTYRESOURCE(StructuredBuffer, StructuredBuffer);
+	FEmptyStructuredBuffer* StructuredBuffer = ResourceCast(StructuredBufferRHI);
 
 	FEmptyShaderResourceView* SRV = new FEmptyShaderResourceView;
 	return SRV;
@@ -54,7 +54,7 @@ FShaderResourceViewRHIRef FEmptyDynamicRHI::RHICreateShaderResourceView(FStructu
 
 FShaderResourceViewRHIRef FEmptyDynamicRHI::RHICreateShaderResourceView(FVertexBufferRHIParamRef VertexBufferRHI, uint32 Stride, uint8 Format)
 {
-	DYNAMIC_CAST_EMPTYRESOURCE(VertexBuffer, VertexBuffer);
+	FEmptyVertexBuffer* VertexBuffer = ResourceCast(VertexBufferRHI);
 
 	FEmptyShaderResourceView* SRV = new FEmptyShaderResourceView;
 	SRV->SourceVertexBuffer = VertexBuffer;
@@ -77,7 +77,7 @@ FShaderResourceViewRHIRef FEmptyDynamicRHI::RHICreateShaderResourceView(FTexture
 
 void FEmptyDynamicRHI::RHIClearUAV(FUnorderedAccessViewRHIParamRef UnorderedAccessViewRHI, const uint32* Values)
 {
-	DYNAMIC_CAST_EMPTYRESOURCE(UnorderedAccessView, UnorderedAccessView);
+	FEmptyUnorderedAccessView* UnorderedAccessView = ResourceCast(UnorderedAccessViewRHI);
 
 }
 
diff --git a/Engine/Source/Runtime/EmptyRHI/Private/EmptyVertexBuffer.cpp b/Engine/Source/Runtime/EmptyRHI/Private/EmptyVertexBuffer.cpp
--- a/Engine/Source/Runtime/EmptyRHI/Private/EmptyVertexBuffer.cpp
+++ b/Engine/Source/Runtime/EmptyRHI/Private/EmptyVertexBuffer.cpp
@@ -48,7 +48,7 @@ FVertexBufferRHIRef FEmptyDynamicRHI::RHICreateVertexBuffer(uint32 Size, uint32
 
 void* FEmptyDynamicRHI::RHILockVertexBuffer(FVertexBufferRHIParamRef VertexBufferRHI, uint32 Offset, uint32 Size, EResourceLockMode LockMode)
 {
-	DYNAMIC_CAST_EMPTYRESOURCE(VertexBuffer,VertexBuffer);
+	FEmptyVertexBuffer* VertexBuffer = ResourceCast(VertexBufferRHI);
 
 	// default to vertex buffer memory
 	return (uint8*)VertexBuffer->Lock(LockMode, Size) + Offset;
@@ -56,7 +56,7 @@ void* FEmptyDynamicRHI::RHILockVertexBuffer(FVertexBufferRHIParamRef VertexBuffe
 
 void FEmptyDynamicRHI::RHIUnlockVertexBuffer(FVertexBufferRHIParamRef VertexBufferRHI)
 {
-	DYNAMIC_CAST_EMPTYRESOURCE(VertexBuffer,VertexBuffer);
+	FEmptyVertexBuffer* VertexBuffer = ResourceCast(VertexBufferRHI);
 
 	VertexBuffer->Unlock();
 }
diff --git a/Engine/Source/Runtime/EmptyRHI/Private/EmptyViewport.cpp b/Engine/Source/Runtime/EmptyRHI/Private/EmptyViewport.cpp
--- a/Engine/Source/Runtime/EmptyRHI/Private/EmptyViewport.cpp
+++ b/Engine/Source/Runtime/EmptyRHI/Private/EmptyViewport.cpp
@@ -30,7 +30,7 @@ void FEmptyDynamicRHI::RHIResizeViewport(FViewportRHIParamRef ViewportRHI,uint32
 {
 	check( IsInGameThread() );
 
-	DYNAMIC_CAST_EMPTYRESOURCE(Viewport,Viewport);
+	FEmptyViewport* Viewport = ResourceCast(ViewportRHI);
 }
 
 void FEmptyDynamicRHI::RHITick( float DeltaTime )
@@ -45,14 +45,14 @@ void FEmptyDynamicRHI::RHITick( float DeltaTime )
 
 void FEmptyDynamicRHI::RHIBeginDrawingViewport(FViewportRHIParamRef ViewportRHI, FTextureRHIParamRef RenderTargetRHI)
 {
-	DYNAMIC_CAST_EMPTYRESOURCE(Viewport,Viewport);
+	FEmptyViewport* Viewport = ResourceCast(ViewportRHI);
 
 	RHISetRenderTarget(RHIGetViewportBackBuffer(ViewportRHI), NULL);
 }
 
 void FEmptyDynamicRHI::RHIEndDrawingViewport(FViewportRHIParamRef ViewportRHI,bool bPresent,bool bLockToVsync)
 {
-	DYNAMIC_CAST_EMPTYRESOURCE(Viewport,Viewport);
+	FEmptyViewport* Viewport = ResourceCast(ViewportRHI);
 }
 
 bool FEmptyDynamicRHI::RHIIsDrawingViewport()
@@ -62,7 +62,7 @@ bool FEmptyDynamicRHI::RHIIsDrawingViewport()
 
 FTexture2DRHIRef FEmptyDynamicRHI::RHIGetViewportBackBuffer(FViewportRHIParamRef ViewportRHI)
 {
-	DYNAMIC_CAST_EMPTYRESOURCE(Viewport,Viewport);
+	FEmptyViewport* Viewport = ResourceCast(ViewportRHI);
 
 	return FTexture2DRHIRef();
 }
diff --git a/Engine/Source/Runtime/EmptyRHI/Public/EmptyRHI.h b/Engine/Source/Runtime/EmptyRHI/Public/EmptyRHI.h
--- a/Engine/Source/Runtime/EmptyRHI/Public/EmptyRHI.h
+++ b/Engine/Source/Runtime/EmptyRHI/Public/EmptyRHI.h
@@ -8,10 +8,6 @@
 
 DECLARE_LOG_CATEGORY_EXTERN(LogEmpty, Display, All)
 
-/** This is a macro that casts a dynamically bound RHI reference to the appropriate type. */
-#define DYNAMIC_CAST_EMPTYRESOURCE(Type,Name) \
-	FEmpty##Type* Name = (FEmpty##Type*)Name##RHI;
-
 // Empty RHI public headers.
 #include "EmptyState.h"
 #include "EmptyResources.h"
@@ -33,11 +29,15 @@ public:
 	virtual void Init();
 	virtual void Shutdown();
 
-	// The RHI methods are defined as virtual functions in URenderHardwareInterface.
+	template<typename TRHIType>
+	static FORCEINLINE typename TEmptyResourceTraits<TRHIType>::TConcreteType* ResourceCast(TRHIType* Resource)
+	{
+		return static_cast<TEmptyResourceTraits<TRHIType>::TConcreteType*>(Resource);
+	}
+
 	#define DEFINE_RHIMETHOD(Type,Name,ParameterTypesAndNames,ParameterNames,ReturnStatement,NullImplementation) virtual Type RHI##Name ParameterTypesAndNames
 	#include "RHIMethods.h"
 	#undef DEFINE_RHIMETHOD
-
 };
 
 /** Implements the Empty module as a dynamic RHI providing module. */
diff --git a/Engine/Source/Runtime/EmptyRHI/Public/EmptyResources.h b/Engine/Source/Runtime/EmptyRHI/Public/EmptyResources.h
--- a/Engine/Source/Runtime/EmptyRHI/Public/EmptyResources.h
+++ b/Engine/Source/Runtime/EmptyRHI/Public/EmptyResources.h
@@ -307,3 +307,133 @@ public:
 	~FEmptyShaderResourceView();
 };
 
+template<class T>
+struct TEmptyResourceTraits
+{
+};
+template<>
+struct TEmptyResourceTraits<FRHIVertexDeclaration>
+{
+	typedef FEmptyVertexDeclaration TConcreteType;
+};
+template<>
+struct TEmptyResourceTraits<FRHIVertexShader>
+{
+	typedef FEmptyVertexShader TConcreteType;
+};
+template<>
+struct TEmptyResourceTraits<FRHIGeometryShader>
+{
+	typedef FEmptyGeometryShader TConcreteType;
+};
+template<>
+struct TEmptyResourceTraits<FRHIHullShader>
+{
+	typedef FEmptyHullShader TConcreteType;
+};
+template<>
+struct TEmptyResourceTraits<FRHIDomainShader>
+{
+	typedef FEmptyDomainShader TConcreteType;
+};
+template<>
+struct TEmptyResourceTraits<FRHIPixelShader>
+{
+	typedef FEmptyPixelShader TConcreteType;
+};
+template<>
+struct TEmptyResourceTraits<FRHIComputeShader>
+{
+	typedef FEmptyComputeShader TConcreteType;
+};
+template<>
+struct TEmptyResourceTraits<FRHIBoundShaderState>
+{
+	typedef FEmptyBoundShaderState TConcreteType;
+};
+template<>
+struct TEmptyResourceTraits<FRHITexture3D>
+{
+	typedef FEmptyTexture3D TConcreteType;
+};
+template<>
+struct TEmptyResourceTraits<FRHITexture>
+{
+	typedef FEmptyTexture TConcreteType;
+};
+template<>
+struct TEmptyResourceTraits<FRHITexture2D>
+{
+	typedef FEmptyTexture2D TConcreteType;
+};
+template<>
+struct TEmptyResourceTraits<FRHITexture2DArray>
+{
+	typedef FEmptyTexture2DArray TConcreteType;
+};
+template<>
+struct TEmptyResourceTraits<FRHITextureCube>
+{
+	typedef FEmptyTextureCube TConcreteType;
+};
+template<>
+struct TEmptyResourceTraits<FRHIRenderQuery>
+{
+	typedef FEmptyRenderQuery TConcreteType;
+};
+template<>
+struct TEmptyResourceTraits<FRHIUniformBuffer>
+{
+	typedef FEmptyUniformBuffer TConcreteType;
+};
+template<>
+struct TEmptyResourceTraits<FRHIIndexBuffer>
+{
+	typedef FEmptyIndexBuffer TConcreteType;
+};
+template<>
+struct TEmptyResourceTraits<FRHIStructuredBuffer>
+{
+	typedef FEmptyStructuredBuffer TConcreteType;
+};
+template<>
+struct TEmptyResourceTraits<FRHIVertexBuffer>
+{
+	typedef FEmptyVertexBuffer TConcreteType;
+};
+template<>
+struct TEmptyResourceTraits<FRHIShaderResourceView>
+{
+	typedef FEmptyShaderResourceView TConcreteType;
+};
+template<>
+struct TEmptyResourceTraits<FRHIUnorderedAccessView>
+{
+	typedef FEmptyUnorderedAccessView TConcreteType;
+};
+
+template<>
+struct TEmptyResourceTraits<FRHISamplerState>
+{
+	typedef FEmptySamplerState TConcreteType;
+};
+template<>
+struct TEmptyResourceTraits<FRHIRasterizerState>
+{
+	typedef FEmptyRasterizerState TConcreteType;
+};
+template<>
+struct TEmptyResourceTraits<FRHIDepthStencilState>
+{
+	typedef FEmptyDepthStencilState TConcreteType;
+};
+template<>
+struct TEmptyResourceTraits<FRHIBlendState>
+{
+	typedef FEmptyBlendState TConcreteType;
+};
+template<>
+struct TEmptyResourceTraits<FRHIViewport>
+{
+	typedef FEmptyViewport TConcreteType;
+};
diff --git a/Engine/Source/Runtime/IOS/MetalRHI/Private/MetalCommands.cpp b/Engine/Source/Runtime/IOS/MetalRHI/Private/MetalCommands.cpp
--- a/Engine/Source/Runtime/IOS/MetalRHI/Private/MetalCommands.cpp
+++ b/Engine/Source/Runtime/IOS/MetalRHI/Private/MetalCommands.cpp
@@ -47,7 +47,7 @@ static MTLPrimitiveType TranslatePrimitiveType(uint32 PrimitiveType)
 
 void FMetalDynamicRHI::RHISetStreamSource(uint32 StreamIndex,FVertexBufferRHIParamRef VertexBufferRHI,uint32 Stride,uint32 Offset)
 {
-	DYNAMIC_CAST_METALRESOURCE(VertexBuffer,VertexBuffer);
+	FMetalVertexBuffer* VertexBuffer = ResourceCast(VertexBufferRHI);
 
 	if (VertexBuffer != NULL)
 	{
@@ -63,14 +63,14 @@ void FMetalDynamicRHI::RHISetStreamOutTargets(uint32 NumTargets, const FVertexBu
 
 void FMetalDynamicRHI::RHISetRasterizerState(FRasterizerStateRHIParamRef NewStateRHI)
 {
-	DYNAMIC_CAST_METALRESOURCE(RasterizerState,NewState);
+	FMetalRasterizerState* NewState = ResourceCast(NewStateRHI);
 
 	NewState->Set();
 }
 
 void FMetalDynamicRHI::RHISetComputeShader(FComputeShaderRHIParamRef ComputeShaderRHI)
 {
-	DYNAMIC_CAST_METALRESOURCE(ComputeShader, ComputeShader);
+	FMetalComputeShader* ComputeShader = ResourceCast(ComputeShaderRHI);
 	FMetalManager::Get()->SetComputeShader(ComputeShader);
 	
 	// set this compute shader pipeline as the current (this resets all state, so we need to set all resources after calling this)
@@ -130,7 +130,7 @@ void FMetalDynamicRHI::RHISetScissorRect(bool bEnable,uint32 MinX,uint32 MinY,ui
 
 void FMetalDynamicRHI::RHISetBoundShaderState( FBoundShaderStateRHIParamRef BoundShaderStateRHI)
 {
-	DYNAMIC_CAST_METALRESOURCE(BoundShaderState,BoundShaderState);
+	FMetalBoundShaderState* BoundShaderState = ResourceCast(BoundShaderStateRHI);
 
 	FMetalManager::Get()->SetBoundShaderState(BoundShaderState);
 	BoundShaderStateHistory.Add(BoundShaderState);
@@ -139,7 +139,7 @@ void FMetalDynamicRHI::RHISetBoundShaderState( FBoundShaderStateRHIParamRef Boun
 
 void FMetalDynamicRHI::RHISetUAVParameter(FComputeShaderRHIParamRef ComputeShaderRHI, uint32 UAVIndex, FUnorderedAccessViewRHIParamRef UAVRHI)
 {
-	DYNAMIC_CAST_METALRESOURCE(UnorderedAccessView, UAV);
+	FMetalUnorderedAccessView* UAV = ResourceCast(UAVRHI);
 
 	if (UAV)
 	{
@@ -211,7 +211,7 @@ void FMetalDynamicRHI::RHISetShaderTexture(FComputeShaderRHIParamRef ComputeShad
 
 void FMetalDynamicRHI::RHISetShaderResourceViewParameter(FVertexShaderRHIParamRef VertexShaderRHI, uint32 TextureIndex, FShaderResourceViewRHIParamRef SRVRHI)
 {
-	DYNAMIC_CAST_METALRESOURCE(ShaderResourceView, SRV);
+	FMetalShaderResourceView* SRV = ResourceCast(SRVRHI);
 	if (SRV)
 	{
 		FRHITexture* Texture = SRV->SourceTexture.GetReference();
@@ -259,7 +259,7 @@ void FMetalDynamicRHI::RHISetShaderResourceViewParameter(FGeometryShaderRHIParam
 
 void FMetalDynamicRHI::RHISetShaderResourceViewParameter(FPixelShaderRHIParamRef PixelShaderRHI,uint32 TextureIndex,FShaderResourceViewRHIParamRef SRVRHI)
 {
-	DYNAMIC_CAST_METALRESOURCE(ShaderResourceView, SRV);
+	FMetalShaderResourceView* SRV = ResourceCast(SRVRHI);
 	if (SRV)
 	{
 		FRHITexture* Texture = SRV->SourceTexture.GetReference();
@@ -292,7 +292,7 @@ void FMetalDynamicRHI::RHISetShaderResourceViewParameter(FPixelShaderRHIParamRef
 
 void FMetalDynamicRHI::RHISetShaderResourceViewParameter(FComputeShaderRHIParamRef ComputeShaderRHI,uint32 TextureIndex,FShaderResourceViewRHIParamRef SRVRHI)
 {
-	DYNAMIC_CAST_METALRESOURCE(ShaderResourceView, SRV);
+	FMetalShaderResourceView* SRV = ResourceCast(SRVRHI);
 
 	if (SRV != NULL)
 	{
@@ -323,42 +323,42 @@ void FMetalDynamicRHI::RHISetShaderResourceViewParameter(FComputeShaderRHIParamR
 
 void FMetalDynamicRHI::RHISetShaderSampler(FVertexShaderRHIParamRef VertexShaderRHI, uint32 SamplerIndex, FSamplerStateRHIParamRef NewStateRHI)
 {
-	DYNAMIC_CAST_METALRESOURCE(SamplerState,NewState);
+	FMetalSamplerState* NewState = ResourceCast(NewStateRHI);
 
 	[FMetalManager::GetContext() setVertexSamplerState:NewState->State atIndex:SamplerIndex];
 }
 
 void FMetalDynamicRHI::RHISetShaderSampler(FHullShaderRHIParamRef HullShader, uint32 SamplerIndex, FSamplerStateRHIParamRef NewStateRHI)
 {
-	DYNAMIC_CAST_METALRESOURCE(SamplerState,NewState);
+	FMetalSamplerState* NewState = ResourceCast(NewStateRHI);
 
 	NOT_SUPPORTED("RHISetSamplerState-Hull");
 }
 
 void FMetalDynamicRHI::RHISetShaderSampler(FDomainShaderRHIParamRef DomainShader, uint32 SamplerIndex, FSamplerStateRHIParamRef NewStateRHI)
 {
-	DYNAMIC_CAST_METALRESOURCE(SamplerState,NewState);
+	FMetalSamplerState* NewState = ResourceCast(NewStateRHI);
 
 	NOT_SUPPORTED("RHISetSamplerState-Domain");
 }
 
 void FMetalDynamicRHI::RHISetShaderSampler(FGeometryShaderRHIParamRef GeometryShader, uint32 SamplerIndex, FSamplerStateRHIParamRef NewStateRHI)
 {
-	DYNAMIC_CAST_METALRESOURCE(SamplerState,NewState);
+	FMetalSamplerState* NewState = ResourceCast(NewStateRHI);
 
 	NOT_SUPPORTED("RHISetSamplerState-Geometry");
 }
 
 void FMetalDynamicRHI::RHISetShaderSampler(FPixelShaderRHIParamRef PixelShader, uint32 SamplerIndex, FSamplerStateRHIParamRef NewStateRHI)
 {
-	DYNAMIC_CAST_METALRESOURCE(SamplerState,NewState);
+	FMetalSamplerState* NewState = ResourceCast(NewStateRHI);
 
 	[FMetalManager::GetContext() setFragmentSamplerState:NewState->State atIndex:SamplerIndex];
 }
 
 void FMetalDynamicRHI::RHISetShaderSampler(FComputeShaderRHIParamRef ComputeShader, uint32 SamplerIndex, FSamplerStateRHIParamRef NewStateRHI)
 {
-	DYNAMIC_CAST_METALRESOURCE(SamplerState, NewState);
+	FMetalSamplerState* NewState = ResourceCast(NewStateRHI);
 
 	[FMetalManager::GetComputeContext() setSamplerState:NewState->State atIndex : SamplerIndex];
 }
@@ -395,7 +395,7 @@ void FMetalDynamicRHI::RHISetShaderParameter(FComputeShaderRHIParamRef ComputeSh
 
 void FMetalDynamicRHI::RHISetShaderUniformBuffer(FVertexShaderRHIParamRef VertexShaderRHI, uint32 BufferIndex, FUniformBufferRHIParamRef BufferRHI)
 {
-	DYNAMIC_CAST_METALRESOURCE(VertexShader, VertexShader);
+	FMetalVertexShader* VertexShader = ResourceCast(VertexShaderRHI);
 	VertexShader->BoundUniformBuffers[BufferIndex] = BufferRHI;
 	VertexShader->DirtyUniformBuffers |= 1 << BufferIndex;
 
@@ -425,7 +425,7 @@ void FMetalDynamicRHI::RHISetShaderUniformBuffer(FGeometryShaderRHIParamRef Geom
 
 void FMetalDynamicRHI::RHISetShaderUniformBuffer(FPixelShaderRHIParamRef PixelShaderRHI, uint32 BufferIndex, FUniformBufferRHIParamRef BufferRHI)
 {
-	DYNAMIC_CAST_METALRESOURCE(PixelShader, PixelShader);
+	FMetalPixelShader* PixelShader = ResourceCast(PixelShaderRHI);
 	PixelShader->BoundUniformBuffers[BufferIndex] = BufferRHI;
 	PixelShader->DirtyUniformBuffers |= 1 << BufferIndex;
 
@@ -440,7 +440,7 @@ void FMetalDynamicRHI::RHISetShaderUniformBuffer(FPixelShaderRHIParamRef PixelSh
 
 void FMetalDynamicRHI::RHISetShaderUniformBuffer(FComputeShaderRHIParamRef ComputeShaderRHI, uint32 BufferIndex, FUniformBufferRHIParamRef BufferRHI)
 {
-	DYNAMIC_CAST_METALRESOURCE(ComputeShader, ComputeShader);
+	FMetalComputeShader* ComputeShader = ResourceCast(ComputeShaderRHI);
 	ComputeShader->BoundUniformBuffers[BufferIndex] = BufferRHI;
 	ComputeShader->DirtyUniformBuffers |= 1 << BufferIndex;
 
@@ -456,14 +456,14 @@ void FMetalDynamicRHI::RHISetShaderUniformBuffer(FComputeShaderRHIParamRef Compu
 
 void FMetalDynamicRHI::RHISetDepthStencilState(FDepthStencilStateRHIParamRef NewStateRHI, uint32 StencilRef)
 {
-	DYNAMIC_CAST_METALRESOURCE(DepthStencilState,NewState);
+	FMetalDepthStencilState* NewState = ResourceCast(NewStateRHI);
 
 	NewState->Set();
 }
 
 void FMetalDynamicRHI::RHISetBlendState(FBlendStateRHIParamRef NewStateRHI, const FLinearColor& BlendFactor)
 {
-	DYNAMIC_CAST_METALRESOURCE(BlendState,NewState);
+	FMetalBlendState* NewState = ResourceCast(NewStateRHI);
 
 	NewState->Set();
 }
@@ -552,14 +552,14 @@ void FMetalDynamicRHI::RHISetRenderTargetsAndClear(const FRHISetRenderTargetsInf
 // Occlusion/Timer queries.
 void FMetalDynamicRHI::RHIBeginRenderQuery(FRenderQueryRHIParamRef QueryRHI)
 {
-	DYNAMIC_CAST_METALRESOURCE(RenderQuery,Query);
+	FMetalRenderQuery* Query = ResourceCast(QueryRHI);
 
 	Query->Begin();
 }
 
 void FMetalDynamicRHI::RHIEndRenderQuery(FRenderQueryRHIParamRef QueryRHI)
 {
-	DYNAMIC_CAST_METALRESOURCE(RenderQuery,Query);
+	FMetalRenderQuery* Query = ResourceCast(QueryRHI);
 
 	Query->End();
 }
@@ -601,7 +601,7 @@ void FMetalDynamicRHI::RHIDrawIndexedPrimitive(FIndexBufferRHIParamRef IndexBuff
 
 	RHI_DRAW_CALL_STATS(PrimitiveType,NumInstances*NumPrimitives);
 
-	DYNAMIC_CAST_METALRESOURCE(IndexBuffer,IndexBuffer);
+	FMetalIndexBuffer* IndexBuffer = ResourceCast(IndexBufferRHI);
 
 	// finalize any pending state
 	FMetalManager::Get()->PrepareToDraw(NumVertices);
diff --git a/Engine/Source/Runtime/IOS/MetalRHI/Private/MetalIndexBuffer.cpp b/Engine/Source/Runtime/IOS/MetalRHI/Private/MetalIndexBuffer.cpp
--- a/Engine/Source/Runtime/IOS/MetalRHI/Private/MetalIndexBuffer.cpp
+++ b/Engine/Source/Runtime/IOS/MetalRHI/Private/MetalIndexBuffer.cpp
@@ -72,14 +72,14 @@ FIndexBufferRHIRef FMetalDynamicRHI::RHICreateIndexBuffer(uint32 Stride,uint32 S
 
 void* FMetalDynamicRHI::RHILockIndexBuffer(FIndexBufferRHIParamRef IndexBufferRHI, uint32 Offset, uint32 Size, EResourceLockMode LockMode)
 {
-	DYNAMIC_CAST_METALRESOURCE(IndexBuffer,IndexBuffer);
+	FMetalIndexBuffer* IndexBuffer = ResourceCast(IndexBufferRHI);
 
 	return (uint8*)IndexBuffer->Lock(LockMode, Size) + Offset;
 }
 
 void FMetalDynamicRHI::RHIUnlockIndexBuffer(FIndexBufferRHIParamRef IndexBufferRHI)
 {
-	DYNAMIC_CAST_METALRESOURCE(IndexBuffer,IndexBuffer);
+	FMetalIndexBuffer* IndexBuffer = ResourceCast(IndexBufferRHI);
 
 	IndexBuffer->Unlock();
 }
diff --git a/Engine/Source/Runtime/IOS/MetalRHI/Private/MetalQuery.cpp b/Engine/Source/Runtime/IOS/MetalRHI/Private/MetalQuery.cpp
--- a/Engine/Source/Runtime/IOS/MetalRHI/Private/MetalQuery.cpp
+++ b/Engine/Source/Runtime/IOS/MetalRHI/Private/MetalQuery.cpp
@@ -49,7 +49,7 @@ FRenderQueryRHIRef FMetalDynamicRHI::RHICreateRenderQuery(ERenderQueryType Query
 bool FMetalDynamicRHI::RHIGetRenderQueryResult(FRenderQueryRHIParamRef QueryRHI,uint64& OutNumPixels,bool bWait)
 {
 	check(IsInRenderingThread());
-	DYNAMIC_CAST_METALRESOURCE(RenderQuery, Query);
+	FMetalRenderQuery* Query = ResourceCast(QueryRHI);
 
 	uint32* QueryMemory = (uint32*)((uint8*)[FMetalManager::Get()->GetQueryBuffer() contents] + Query->Offset);
 
diff --git a/Engine/Source/Runtime/IOS/MetalRHI/Private/MetalRenderTarget.cpp b/Engine/Source/Runtime/IOS/MetalRHI/Private/MetalRenderTarget.cpp
--- a/Engine/Source/Runtime/IOS/MetalRHI/Private/MetalRenderTarget.cpp
+++ b/Engine/Source/Runtime/IOS/MetalRHI/Private/MetalRenderTarget.cpp
@@ -33,7 +33,7 @@ void FMetalDynamicRHI::RHIReadSurfaceFloatData(FTextureRHIParamRef TextureRHI, F
 	//kick the current command buffer.
 	FMetalManager::Get()->SubmitCommandBufferAndWait();
 
-	DYNAMIC_CAST_METALRESOURCE(Texture2D, Texture);
+	FMetalTexture2D* Texture = ResourceCast(TextureRHI);
 	
 	FMetalSurface* Surface = GetMetalSurfaceFromRHITexture(TextureRHI);
 	
diff --git a/Engine/Source/Runtime/IOS/MetalRHI/Private/MetalShaders.cpp b/Engine/Source/Runtime/IOS/MetalRHI/Private/MetalShaders.cpp
--- a/Engine/Source/Runtime/IOS/MetalRHI/Private/MetalShaders.cpp
+++ b/Engine/Source/Runtime/IOS/MetalRHI/Private/MetalShaders.cpp
@@ -249,12 +249,12 @@ FMetalBoundShaderState::FMetalBoundShaderState(
 			FGeometryShaderRHIParamRef InGeometryShaderRHI)
 	:	CacheLink(InVertexDeclarationRHI,InVertexShaderRHI,InPixelShaderRHI,InHullShaderRHI,InDomainShaderRHI,InGeometryShaderRHI,this)
 {
-	DYNAMIC_CAST_METALRESOURCE(VertexDeclaration,InVertexDeclaration);
-	DYNAMIC_CAST_METALRESOURCE(VertexShader,InVertexShader);
-	DYNAMIC_CAST_METALRESOURCE(PixelShader,InPixelShader);
-	DYNAMIC_CAST_METALRESOURCE(HullShader,InHullShader);
-	DYNAMIC_CAST_METALRESOURCE(DomainShader,InDomainShader);
-	DYNAMIC_CAST_METALRESOURCE(GeometryShader,InGeometryShader);
+	FMetalVertexDeclaration* InVertexDeclaration = FMetalDynamicRHI::ResourceCast(InVertexDeclarationRHI);
+	FMetalVertexShader* InVertexShader = FMetalDynamicRHI::ResourceCast(InVertexShaderRHI);
+	FMetalPixelShader* InPixelShader = FMetalDynamicRHI::ResourceCast(InPixelShaderRHI);
+	FMetalHullShader* InHullShader = FMetalDynamicRHI::ResourceCast(InHullShaderRHI);
+	FMetalDomainShader* InDomainShader = FMetalDynamicRHI::ResourceCast(InDomainShaderRHI);
+	FMetalGeometryShader* InGeometryShader = FMetalDynamicRHI::ResourceCast(InGeometryShaderRHI);
 
 	// cache everything
 	VertexDeclaration = InVertexDeclaration;
diff --git a/Engine/Source/Runtime/IOS/MetalRHI/Private/MetalStructuredBuffer.cpp b/Engine/Source/Runtime/IOS/MetalRHI/Private/MetalStructuredBuffer.cpp
--- a/Engine/Source/Runtime/IOS/MetalRHI/Private/MetalStructuredBuffer.cpp
+++ b/Engine/Source/Runtime/IOS/MetalRHI/Private/MetalStructuredBuffer.cpp
@@ -37,7 +37,7 @@ FStructuredBufferRHIRef FMetalDynamicRHI::RHICreateStructuredBuffer(uint32 Strid
 
 void* FMetalDynamicRHI::RHILockStructuredBuffer(FStructuredBufferRHIParamRef StructuredBufferRHI,uint32 Offset,uint32 Size,EResourceLockMode LockMode)
 {
-	DYNAMIC_CAST_METALRESOURCE(StructuredBuffer,StructuredBuffer);
+	FMetalStructuredBuffer* StructuredBuffer = ResourceCast(StructuredBufferRHI);
 
 	// just return the memory plus the offset
 	return ((uint8*)[StructuredBuffer->Buffer contents]) + Offset;
diff --git a/Engine/Source/Runtime/IOS/MetalRHI/Private/MetalTexture.cpp b/Engine/Source/Runtime/IOS/MetalRHI/Private/MetalTexture.cpp
--- a/Engine/Source/Runtime/IOS/MetalRHI/Private/MetalTexture.cpp
+++ b/Engine/Source/Runtime/IOS/MetalRHI/Private/MetalTexture.cpp
@@ -458,7 +458,7 @@ void FMetalDynamicRHI::RHIGenerateMips(FTextureRHIParamRef SourceSurfaceRHI)
 
 FTexture2DRHIRef FMetalDynamicRHI::RHIAsyncReallocateTexture2D(FTexture2DRHIParamRef OldTextureRHI, int32 NewMipCount, int32 NewSizeX, int32 NewSizeY, FThreadSafeCounter* RequestStatus)
 {
-	DYNAMIC_CAST_METALRESOURCE(Texture2D,OldTexture);
+	FMetalTexture2D* OldTexture = ResourceCast(OldTextureRHI);
 
 	FMetalTexture2D* NewTexture = new FMetalTexture2D(OldTexture->GetFormat(), NewSizeX, NewSizeY, NewMipCount, OldTexture->GetNumSamples(), OldTexture->GetFlags(), NULL);
 
@@ -521,37 +521,37 @@ ETextureReallocationStatus FMetalDynamicRHI::RHICancelAsyncReallocateTexture2D(
 
 void* FMetalDynamicRHI::RHILockTexture2D(FTexture2DRHIParamRef TextureRHI,uint32 MipIndex,EResourceLockMode LockMode,uint32& DestStride,bool bLockWithinMiptail)
 {
-	DYNAMIC_CAST_METALRESOURCE(Texture2D,Texture);
+	FMetalTexture2D* Texture = ResourceCast(TextureRHI);
 	return Texture->Surface.Lock(MipIndex, 0, LockMode, DestStride);
 }
 
 void FMetalDynamicRHI::RHIUnlockTexture2D(FTexture2DRHIParamRef TextureRHI,uint32 MipIndex,bool bLockWithinMiptail)
 {
-	DYNAMIC_CAST_METALRESOURCE(Texture2D,Texture);
+	FMetalTexture2D* Texture = ResourceCast(TextureRHI);
 	Texture->Surface.Unlock(MipIndex, 0);
 }
 
 void* FMetalDynamicRHI::RHILockTexture2DArray(FTexture2DArrayRHIParamRef TextureRHI,uint32 TextureIndex,uint32 MipIndex,EResourceLockMode LockMode,uint32& DestStride,bool bLockWithinMiptail)
 {
-	DYNAMIC_CAST_METALRESOURCE(Texture2DArray,Texture);
+	FMetalTexture2DArray* Texture = ResourceCast(TextureRHI);
 	return Texture->Surface.Lock(MipIndex, TextureIndex, LockMode, DestStride);
 }
 
 void FMetalDynamicRHI::RHIUnlockTexture2DArray(FTexture2DArrayRHIParamRef TextureRHI,uint32 TextureIndex,uint32 MipIndex,bool bLockWithinMiptail)
 {
-	DYNAMIC_CAST_METALRESOURCE(Texture2DArray,Texture);
+	FMetalTexture2DArray* Texture = ResourceCast(TextureRHI);
 	Texture->Surface.Unlock(MipIndex, TextureIndex);
 }
 
 void FMetalDynamicRHI::RHIUpdateTexture2D(FTexture2DRHIParamRef TextureRHI, uint32 MipIndex, const struct FUpdateTextureRegion2D& UpdateRegion, uint32 SourcePitch, const uint8* SourceData)
 {	
-	DYNAMIC_CAST_METALRESOURCE(Texture3D,Texture);
+	FMetalTexture3D* Texture = ResourceCast(TextureRHI);
 
 }
 
 void FMetalDynamicRHI::RHIUpdateTexture3D(FTexture3DRHIParamRef TextureRHI,uint32 MipIndex,const FUpdateTextureRegion3D& UpdateRegion,uint32 SourceRowPitch,uint32 SourceDepthPitch,const uint8* SourceData)
 {	
-	DYNAMIC_CAST_METALRESOURCE(Texture3D,Texture);
+	FMetalTexture3D* Texture = ResourceCast(TextureRHI);
 
 }
 
@@ -571,14 +571,14 @@ FTextureCubeRHIRef FMetalDynamicRHI::RHICreateTextureCubeArray(uint32 Size, uint
 
 void* FMetalDynamicRHI::RHILockTextureCubeFace(FTextureCubeRHIParamRef TextureCubeRHI,uint32 FaceIndex,uint32 ArrayIndex,uint32 MipIndex,EResourceLockMode LockMode,uint32& DestStride,bool bLockWithinMiptail)
 {
-	DYNAMIC_CAST_METALRESOURCE(TextureCube,TextureCube);
+	FMetalTextureCube* TextureCube = ResourceCast(TextureCubeRHI);
 	uint32 MetalFace = GetMetalCubeFace((ECubeFace)FaceIndex);
 	return TextureCube->Surface.Lock(MipIndex, FaceIndex + 6 * ArrayIndex, LockMode, DestStride);
 }
 
 void FMetalDynamicRHI::RHIUnlockTextureCubeFace(FTextureCubeRHIParamRef TextureCubeRHI,uint32 FaceIndex,uint32 ArrayIndex,uint32 MipIndex,bool bLockWithinMiptail)
 {
-	DYNAMIC_CAST_METALRESOURCE(TextureCube,TextureCube);
+	FMetalTextureCube* TextureCube = ResourceCast(TextureCubeRHI);
 	uint32 MetalFace = GetMetalCubeFace((ECubeFace)FaceIndex);
 	TextureCube->Surface.Unlock(MipIndex, FaceIndex + ArrayIndex * 6);
 }
diff --git a/Engine/Source/Runtime/IOS/MetalRHI/Private/MetalUAV.cpp b/Engine/Source/Runtime/IOS/MetalRHI/Private/MetalUAV.cpp
--- a/Engine/Source/Runtime/IOS/MetalRHI/Private/MetalUAV.cpp
+++ b/Engine/Source/Runtime/IOS/MetalRHI/Private/MetalUAV.cpp
@@ -41,7 +41,7 @@ void FMetalUnorderedAccessView::Set(uint32 ResourceIndex)
 
 FUnorderedAccessViewRHIRef FMetalDynamicRHI::RHICreateUnorderedAccessView(FStructuredBufferRHIParamRef StructuredBufferRHI, bool bUseUAVCounter, bool bAppendBuffer)
 {
-	DYNAMIC_CAST_METALRESOURCE(StructuredBuffer,StructuredBuffer);
+	FMetalStructuredBuffer* StructuredBuffer = ResourceCast(StructuredBufferRHI);
 
 	// create the UAV buffer to point to the structured buffer's memory
 	FMetalUnorderedAccessView* UAV = new FMetalUnorderedAccessView;
@@ -63,7 +63,7 @@ FUnorderedAccessViewRHIRef FMetalDynamicRHI::RHICreateUnorderedAccessView(FTextu
 
 FUnorderedAccessViewRHIRef FMetalDynamicRHI::RHICreateUnorderedAccessView(FVertexBufferRHIParamRef VertexBufferRHI, uint8 Format)
 {
-	DYNAMIC_CAST_METALRESOURCE(VertexBuffer, VertexBuffer);
+	FMetalVertexBuffer* VertexBuffer = ResourceCast(VertexBufferRHI);
 
 	// create the UAV buffer to point to the structured buffer's memory
 	FMetalUnorderedAccessView* UAV = new FMetalUnorderedAccessView;
@@ -74,7 +74,7 @@ FUnorderedAccessViewRHIRef FMetalDynamicRHI::RHICreateUnorderedAccessView(FVerte
 
 FShaderResourceViewRHIRef FMetalDynamicRHI::RHICreateShaderResourceView(FStructuredBufferRHIParamRef StructuredBufferRHI)
 {
-	DYNAMIC_CAST_METALRESOURCE(StructuredBuffer, StructuredBuffer);
+	FMetalStructuredBuffer* StructuredBuffer = ResourceCast(StructuredBufferRHI);
 
 	FMetalShaderResourceView* SRV = new FMetalShaderResourceView;
 	return SRV;
@@ -82,7 +82,7 @@ FShaderResourceViewRHIRef FMetalDynamicRHI::RHICreateShaderResourceView(FStructu
 
 FShaderResourceViewRHIRef FMetalDynamicRHI::RHICreateShaderResourceView(FVertexBufferRHIParamRef VertexBufferRHI, uint32 Stride, uint8 Format)
 {
-	DYNAMIC_CAST_METALRESOURCE(VertexBuffer, VertexBuffer);
+	FMetalVertexBuffer* VertexBuffer = ResourceCast(VertexBufferRHI);
 
 	FMetalShaderResourceView* SRV = new FMetalShaderResourceView;
 	SRV->SourceVertexBuffer = VertexBuffer;
@@ -105,7 +105,7 @@ FShaderResourceViewRHIRef FMetalDynamicRHI::RHICreateShaderResourceView(FTexture
 
 void FMetalDynamicRHI::RHIClearUAV(FUnorderedAccessViewRHIParamRef UnorderedAccessViewRHI, const uint32* Values)
 {
-	DYNAMIC_CAST_METALRESOURCE(UnorderedAccessView, UnorderedAccessView);
+	FMetalUnorderedAccessView* UnorderedAccessView = ResourceCast(UnorderedAccessViewRHI);
 
 }
 
diff --git a/Engine/Source/Runtime/IOS/MetalRHI/Private/MetalVertexBuffer.cpp b/Engine/Source/Runtime/IOS/MetalRHI/Private/MetalVertexBuffer.cpp
--- a/Engine/Source/Runtime/IOS/MetalRHI/Private/MetalVertexBuffer.cpp
+++ b/Engine/Source/Runtime/IOS/MetalRHI/Private/MetalVertexBuffer.cpp
@@ -72,7 +72,7 @@ FVertexBufferRHIRef FMetalDynamicRHI::RHICreateVertexBuffer(uint32 Size, uint32
 
 void* FMetalDynamicRHI::RHILockVertexBuffer(FVertexBufferRHIParamRef VertexBufferRHI, uint32 Offset, uint32 Size, EResourceLockMode LockMode)
 {
-	DYNAMIC_CAST_METALRESOURCE(VertexBuffer,VertexBuffer);
+	FMetalVertexBuffer* VertexBuffer = ResourceCast(VertexBufferRHI);
 
 	// default to vertex buffer memory
 	return (uint8*)VertexBuffer->Lock(LockMode, Size) + Offset;
@@ -80,7 +80,7 @@ void* FMetalDynamicRHI::RHILockVertexBuffer(FVertexBufferRHIParamRef VertexBuffe
 
 void FMetalDynamicRHI::RHIUnlockVertexBuffer(FVertexBufferRHIParamRef VertexBufferRHI)
 {
-	DYNAMIC_CAST_METALRESOURCE(VertexBuffer,VertexBuffer);
+	FMetalVertexBuffer* VertexBuffer = ResourceCast(VertexBufferRHI);
 
 	VertexBuffer->Unlock();
 }
diff --git a/Engine/Source/Runtime/IOS/MetalRHI/Private/MetalViewport.cpp b/Engine/Source/Runtime/IOS/MetalRHI/Private/MetalViewport.cpp
--- a/Engine/Source/Runtime/IOS/MetalRHI/Private/MetalViewport.cpp
+++ b/Engine/Source/Runtime/IOS/MetalRHI/Private/MetalViewport.cpp
@@ -34,7 +34,7 @@ void FMetalDynamicRHI::RHIResizeViewport(FViewportRHIParamRef ViewportRHI,uint32
 {
 	check( IsInGameThread() );
 
-	DYNAMIC_CAST_METALRESOURCE(Viewport,Viewport);
+	FMetalViewport* Viewport = ResourceCast(ViewportRHI);
 	Viewport->Resize(SizeX,SizeY,bIsFullscreen);
 }
 
diff --git a/Engine/Source/Runtime/IOS/MetalRHI/Public/MetalRHI.h b/Engine/Source/Runtime/IOS/MetalRHI/Public/MetalRHI.h
--- a/Engine/Source/Runtime/IOS/MetalRHI/Public/MetalRHI.h
+++ b/Engine/Source/Runtime/IOS/MetalRHI/Public/MetalRHI.h
@@ -8,10 +8,6 @@
 
 DECLARE_LOG_CATEGORY_EXTERN(LogMetal, Display, All);
 
-/** This is a macro that casts a dynamically bound RHI reference to the appropriate type. */
-#define DYNAMIC_CAST_METALRESOURCE(Type,Name) \
-	FMetal##Type* Name = (FMetal##Type*)Name##RHI;
-
 // Metal RHI public headers.
 #include <Metal/Metal.h>
 #include "MetalState.h"
@@ -36,6 +32,12 @@ public:
 	virtual void Init();
 	virtual void Shutdown() {}
 
+	template<typename TRHIType>
+	static FORCEINLINE typename TMetalResourceTraits<TRHIType>::TConcreteType* ResourceCast(TRHIType* Resource)
+	{
+		return static_cast<typename TMetalResourceTraits<TRHIType>::TConcreteType*>(Resource);
+	}
+
 	// The RHI methods are defined as virtual functions in URenderHardwareInterface.
 	#define DEFINE_RHIMETHOD(Type,Name,ParameterTypesAndNames,ParameterNames,ReturnStatement,NullImplementation) virtual Type RHI##Name ParameterTypesAndNames
 	#include "RHIMethods.h"
diff --git a/Engine/Source/Runtime/IOS/MetalRHI/Public/MetalResources.h b/Engine/Source/Runtime/IOS/MetalRHI/Public/MetalResources.h
--- a/Engine/Source/Runtime/IOS/MetalRHI/Public/MetalResources.h
+++ b/Engine/Source/Runtime/IOS/MetalRHI/Public/MetalResources.h
@@ -490,3 +490,129 @@ private:
 
 	int32 GlobalUniformArraySize;
 };
+
+template<class T>
+struct TMetalResourceTraits
+{
+};
+template<>
+struct TMetalResourceTraits<FRHIVertexDeclaration>
+{
+	typedef FMetalVertexDeclaration TConcreteType;
+};
+template<>
+struct TMetalResourceTraits<FRHIVertexShader>
+{
+	typedef FMetalVertexShader TConcreteType;
+};
+template<>
+struct TMetalResourceTraits<FRHIGeometryShader>
+{
+	typedef FMetalGeometryShader TConcreteType;
+};
+template<>
+struct TMetalResourceTraits<FRHIHullShader>
+{
+	typedef FMetalHullShader TConcreteType;
+};
+template<>
+struct TMetalResourceTraits<FRHIDomainShader>
+{
+	typedef FMetalDomainShader TConcreteType;
+};
+template<>
+struct TMetalResourceTraits<FRHIPixelShader>
+{
+	typedef FMetalPixelShader TConcreteType;
+};
+template<>
+struct TMetalResourceTraits<FRHIComputeShader>
+{
+	typedef FMetalComputeShader TConcreteType;
+};
+template<>
+struct TMetalResourceTraits<FRHIBoundShaderState>
+{
+	typedef FMetalBoundShaderState TConcreteType;
+};
+template<>
+struct TMetalResourceTraits<FRHITexture3D>
+{
+	typedef FMetalTexture3D TConcreteType;
+};
+template<>
+struct TMetalResourceTraits<FRHITexture2D>
+{
+	typedef FMetalTexture2D TConcreteType;
+};
+template<>
+struct TMetalResourceTraits<FRHITexture2DArray>
+{
+	typedef FMetalTexture2DArray TConcreteType;
+};
+template<>
+struct TMetalResourceTraits<FRHITextureCube>
+{
+	typedef FMetalTextureCube TConcreteType;
+};
+template<>
+struct TMetalResourceTraits<FRHIRenderQuery>
+{
+	typedef FMetalRenderQuery TConcreteType;
+};
+template<>
+struct TMetalResourceTraits<FRHIUniformBuffer>
+{
+	typedef FMetalUniformBuffer TConcreteType;
+};
+template<>
+struct TMetalResourceTraits<FRHIIndexBuffer>
+{
+	typedef FMetalIndexBuffer TConcreteType;
+};
+template<>
+struct TMetalResourceTraits<FRHIStructuredBuffer>
+{
+	typedef FMetalStructuredBuffer TConcreteType;
+};
+template<>
+struct TMetalResourceTraits<FRHIVertexBuffer>
+{
+	typedef FMetalVertexBuffer TConcreteType;
+};
+template<>
+struct TMetalResourceTraits<FRHIShaderResourceView>
+{
+	typedef FMetalShaderResourceView TConcreteType;
+};
+template<>
+struct TMetalResourceTraits<FRHIUnorderedAccessView>
+{
+	typedef FMetalUnorderedAccessView TConcreteType;
+};
+
+template<>
+struct TMetalResourceTraits<FRHISamplerState>
+{
+	typedef FMetalSamplerState TConcreteType;
+};
+template<>
+struct TMetalResourceTraits<FRHIRasterizerState>
+{
+	typedef FMetalRasterizerState TConcreteType;
+};
+template<>
+struct TMetalResourceTraits<FRHIDepthStencilState>
+{
+	typedef FMetalDepthStencilState TConcreteType;
+};
+template<>
+struct TMetalResourceTraits<FRHIBlendState>
+{
+	typedef FMetalBlendState TConcreteType;
+};
+template<>
+struct TMetalResourceTraits<FRHIViewport>
+{
+	typedef FMetalViewport TConcreteType;
+};
diff --git a/Engine/Source/Runtime/OpenGLDrv/Private/OpenGLCommands.cpp b/Engine/Source/Runtime/OpenGLDrv/Private/OpenGLCommands.cpp
--- a/Engine/Source/Runtime/OpenGLDrv/Private/OpenGLCommands.cpp
+++ b/Engine/Source/Runtime/OpenGLDrv/Private/OpenGLCommands.cpp
@@ -10,7 +10,7 @@
 
 #define DECLARE_ISBOUNDSHADER(ShaderType) inline void ValidateBoundShader(TRefCountPtr<FOpenGLBoundShaderState> InBoundShaderState, F##ShaderType##RHIParamRef ShaderType##RHI) \
 { \
-	DYNAMIC_CAST_OPENGLRESOURCE(ShaderType,ShaderType); \
+	FOpenGL##ShaderType* ShaderType = FOpenGLDynamicRHI::ResourceCast(ShaderType##RHI); \
 	ensureMsgf(InBoundShaderState && ShaderType == InBoundShaderState->ShaderType, TEXT("Parameters are being set for a %s which is not currently bound"), TEXT(#ShaderType)); \
 }
 
@@ -280,7 +280,7 @@ void FOpenGLDynamicRHI::RHIGpuTimeEnd(uint32 Hash, bool bCompute)
 // Vertex state.
 void FOpenGLDynamicRHI::RHISetStreamSource(uint32 StreamIndex,FVertexBufferRHIParamRef VertexBufferRHI,uint32 Stride,uint32 Offset)
 {
-	DYNAMIC_CAST_OPENGLRESOURCE(VertexBuffer,VertexBuffer);
+	FOpenGLVertexBuffer* VertexBuffer = ResourceCast(VertexBufferRHI);
 	PendingState.Streams[StreamIndex].VertexBuffer = VertexBuffer;
 	PendingState.Streams[StreamIndex].Stride = Stride;
 	PendingState.Streams[StreamIndex].Offset = Offset;
@@ -295,7 +295,7 @@ void FOpenGLDynamicRHI::RHISetStreamOutTargets(uint32 NumTargets, const FVertexB
 void FOpenGLDynamicRHI::RHISetRasterizerState(FRasterizerStateRHIParamRef NewStateRHI)
 {
 	VERIFY_GL_SCOPE();
-	DYNAMIC_CAST_OPENGLRESOURCE(RasterizerState,NewState);
+	FOpenGLRasterizerState* NewState = ResourceCast(NewStateRHI);
 	PendingState.RasterizerState = NewState->Data;
 	
 	FShaderCache::SetRasterizerState(NewStateRHI);
@@ -439,7 +439,7 @@ inline void FOpenGLDynamicRHI::UpdateScissorRectInOpenGLContext( FOpenGLContextS
 void FOpenGLDynamicRHI::RHISetBoundShaderState( FBoundShaderStateRHIParamRef BoundShaderStateRHI)
 {
 	VERIFY_GL_SCOPE();
-	DYNAMIC_CAST_OPENGLRESOURCE(BoundShaderState,BoundShaderState);
+	FOpenGLBoundShaderState* BoundShaderState = ResourceCast(BoundShaderStateRHI);
 	PendingState.BoundShaderState = BoundShaderState;
 
 	// Prevent transient bound shader states from being recreated for each use by keeping a history of the most recently used bound shader states.
@@ -456,7 +456,7 @@ void FOpenGLDynamicRHI::RHISetUAVParameter(FComputeShaderRHIParamRef ComputeShad
 	VERIFY_GL_SCOPE();
 	if(UnorderedAccessViewRHI)
 	{
-		DYNAMIC_CAST_OPENGLRESOURCE(UnorderedAccessView, UnorderedAccessView);
+		FOpenGLUnorderedAccessView* UnorderedAccessView = ResourceCast(UnorderedAccessViewRHI);
 		InternalSetShaderUAV(FOpenGL::GetFirstComputeUAVUnit() + UAVIndex, UnorderedAccessView->Format , UnorderedAccessView->Resource);
 	}
 	else
@@ -803,7 +803,7 @@ void FOpenGLDynamicRHI::RHISetShaderResourceViewParameter(FPixelShaderRHIParamRe
 	VERIFY_GL_SCOPE();
 	VALIDATE_BOUND_SHADER(PixelShaderRHI);
 	check(FOpenGL::SupportsResourceView());
-	DYNAMIC_CAST_OPENGLRESOURCE(ShaderResourceView,SRV);
+	FOpenGLShaderResourceView* SRV = ResourceCast(SRVRHI);
 	GLuint Resource = 0;
 	GLenum Target = GL_TEXTURE_BUFFER;
 	int32 LimitMip = -1;
@@ -825,7 +825,7 @@ void FOpenGLDynamicRHI::RHISetShaderResourceViewParameter(FVertexShaderRHIParamR
 
 	VERIFY_GL_SCOPE();
 	check(FOpenGL::SupportsResourceView());
-	DYNAMIC_CAST_OPENGLRESOURCE(ShaderResourceView,SRV);
+	FOpenGLShaderResourceView* SRV = ResourceCast(SRVRHI);
 	GLuint Resource = 0;
 	GLenum Target = GL_TEXTURE_BUFFER;
 	int32 LimitMip = -1;
@@ -846,7 +846,7 @@ void FOpenGLDynamicRHI::RHISetShaderResourceViewParameter(FComputeShaderRHIParam
 	check(GMaxRHIFeatureLevel >= ERHIFeatureLevel::SM5);
 	VERIFY_GL_SCOPE();
 	check(FOpenGL::SupportsResourceView());
-	DYNAMIC_CAST_OPENGLRESOURCE(ShaderResourceView,SRV);
+	FOpenGLShaderResourceView* SRV = ResourceCast(SRVRHI);
 	GLuint Resource = 0;
 	GLenum Target = GL_TEXTURE_BUFFER;
 	int32 LimitMip = -1;
@@ -869,7 +869,7 @@ void FOpenGLDynamicRHI::RHISetShaderResourceViewParameter(FHullShaderRHIParamRef
 	check(GMaxRHIFeatureLevel >= ERHIFeatureLevel::SM5);
 	VERIFY_GL_SCOPE();
 	check(FOpenGL::SupportsResourceView());
-	DYNAMIC_CAST_OPENGLRESOURCE(ShaderResourceView,SRV);
+	FOpenGLShaderResourceView* SRV = ResourceCast(SRVRHI);
 	GLuint Resource = 0;
 	GLenum Target = GL_TEXTURE_BUFFER;
 	int32 LimitMip = -1;
@@ -891,7 +891,7 @@ void FOpenGLDynamicRHI::RHISetShaderResourceViewParameter(FDomainShaderRHIParamR
 	check(GMaxRHIFeatureLevel >= ERHIFeatureLevel::SM5);
 	VERIFY_GL_SCOPE();
 	check(FOpenGL::SupportsResourceView());
-	DYNAMIC_CAST_OPENGLRESOURCE(ShaderResourceView,SRV);
+	FOpenGLShaderResourceView* SRV = ResourceCast(SRVRHI);
 	GLuint Resource = 0;
 	GLenum Target = GL_TEXTURE_BUFFER;
 	int32 LimitMip = -1;
@@ -912,7 +912,7 @@ void FOpenGLDynamicRHI::RHISetShaderResourceViewParameter(FGeometryShaderRHIPara
 
 	VERIFY_GL_SCOPE();
 	check(FOpenGL::SupportsResourceView());
-	DYNAMIC_CAST_OPENGLRESOURCE(ShaderResourceView,SRV);
+	FOpenGLShaderResourceView* SRV = ResourceCast(SRVRHI);
 	GLuint Resource = 0;
 	GLenum Target = GL_TEXTURE_BUFFER;
 	int32 LimitMip = -1;
@@ -1025,7 +1025,7 @@ void FOpenGLDynamicRHI::RHISetShaderSampler(FVertexShaderRHIParamRef VertexShade
 	VALIDATE_BOUND_SHADER(VertexShaderRHI);
 
 	VERIFY_GL_SCOPE();
-	DYNAMIC_CAST_OPENGLRESOURCE(SamplerState,NewState);
+	FOpenGLSamplerState* NewState = ResourceCast(NewStateRHI);
 	if (FOpenGL::SupportsSamplerObjects())
 	{
 		if ( OpenGLConsoleVariables::bBindlessTexture == 0 || !FOpenGL::SupportsBindlessTexture())
@@ -1051,7 +1051,7 @@ void FOpenGLDynamicRHI::RHISetShaderSampler(FHullShaderRHIParamRef HullShaderRHI
 
 	check(GMaxRHIFeatureLevel >= ERHIFeatureLevel::SM5);
 	VERIFY_GL_SCOPE();
-	DYNAMIC_CAST_OPENGLRESOURCE(SamplerState,NewState);
+	FOpenGLSamplerState* NewState = ResourceCast(NewStateRHI);
 	if (FOpenGL::SupportsSamplerObjects())
 	{
 		if ( OpenGLConsoleVariables::bBindlessTexture == 0 || !FOpenGL::SupportsBindlessTexture())
@@ -1076,7 +1076,7 @@ void FOpenGLDynamicRHI::RHISetShaderSampler(FDomainShaderRHIParamRef DomainShade
 
 	check(GMaxRHIFeatureLevel >= ERHIFeatureLevel::SM5);
 	VERIFY_GL_SCOPE();
-	DYNAMIC_CAST_OPENGLRESOURCE(SamplerState,NewState);
+	FOpenGLSamplerState* NewState = ResourceCast(NewStateRHI);
 	if (FOpenGL::SupportsSamplerObjects())
 	{
 		if ( OpenGLConsoleVariables::bBindlessTexture == 0 || !FOpenGL::SupportsBindlessTexture())
@@ -1101,7 +1101,7 @@ void FOpenGLDynamicRHI::RHISetShaderSampler(FGeometryShaderRHIParamRef GeometryS
 	VALIDATE_BOUND_SHADER(GeometryShaderRHI);
 
 	VERIFY_GL_SCOPE();
-	DYNAMIC_CAST_OPENGLRESOURCE(SamplerState,NewState);
+	FOpenGLSamplerState* NewState = ResourceCast(NewStateRHI);
 	if (FOpenGL::SupportsSamplerObjects())
 	{
 		if ( OpenGLConsoleVariables::bBindlessTexture == 0 || !FOpenGL::SupportsBindlessTexture())
@@ -1143,7 +1143,7 @@ void FOpenGLDynamicRHI::RHISetShaderSampler(FPixelShaderRHIParamRef PixelShaderR
 	VALIDATE_BOUND_SHADER(PixelShaderRHI);
 
 	VERIFY_GL_SCOPE();
-	DYNAMIC_CAST_OPENGLRESOURCE(SamplerState,NewState);
+	FOpenGLSamplerState* NewState = ResourceCast(NewStateRHI);
 	if ( FOpenGL::SupportsSamplerObjects() )
 	{
 		if ( OpenGLConsoleVariables::bBindlessTexture == 0 || !FOpenGL::SupportsBindlessTexture())
@@ -1197,7 +1197,7 @@ void FOpenGLDynamicRHI::RHISetShaderSampler(FComputeShaderRHIParamRef ComputeSha
 {
 	check(GMaxRHIFeatureLevel >= ERHIFeatureLevel::SM5);
 	VERIFY_GL_SCOPE();
-	DYNAMIC_CAST_OPENGLRESOURCE(SamplerState,NewState);
+	FOpenGLSamplerState* NewState = ResourceCast(NewStateRHI);
 	if ( OpenGLConsoleVariables::bBindlessTexture == 0 || !FOpenGL::SupportsBindlessTexture())
 	{
 		FOpenGL::BindSampler(FOpenGL::GetFirstComputeTextureUnit() + SamplerIndex, NewState->Resource);
@@ -1273,7 +1273,7 @@ void FOpenGLDynamicRHI::RHISetShaderParameter(FComputeShaderRHIParamRef ComputeS
 void FOpenGLDynamicRHI::RHISetDepthStencilState(FDepthStencilStateRHIParamRef NewStateRHI,uint32 StencilRef)
 {
 	VERIFY_GL_SCOPE();
-	DYNAMIC_CAST_OPENGLRESOURCE(DepthStencilState,NewState);
+	FOpenGLDepthStencilState* NewState = ResourceCast(NewStateRHI);
 	PendingState.DepthStencilState = NewState->Data;
 	PendingState.StencilRef = StencilRef;
 	
@@ -1634,7 +1634,7 @@ void FOpenGLDynamicRHI::SetPendingBlendStateForActiveRenderTargets( FOpenGLConte
 
 void FOpenGLDynamicRHI::RHISetBlendState(FBlendStateRHIParamRef NewStateRHI,const FLinearColor& BlendFactor)
 {
-	DYNAMIC_CAST_OPENGLRESOURCE(BlendState,NewState);
+	FOpenGLBlendState* NewState = ResourceCast(NewStateRHI);
 	FMemory::Memcpy(&PendingState.BlendState,&(NewState->Data),sizeof(FOpenGLBlendStateData));
 	
 	FShaderCache::SetBlendState(NewStateRHI);
@@ -2335,7 +2335,7 @@ void FOpenGLDynamicRHI::CommitComputeShaderConstants(FComputeShaderRHIParamRef C
 	VERIFY_GL_SCOPE();
 	check(GMaxRHIFeatureLevel >= ERHIFeatureLevel::SM5);
 
-	DYNAMIC_CAST_OPENGLRESOURCE(ComputeShader,ComputeShader);
+	FOpenGLComputeShader* ComputeShader = ResourceCast(ComputeShaderRHI);
 	const int32 Stage = CrossCompiler::SHADER_STAGE_COMPUTE;
 
 	FOpenGLShaderParameterCache& StageShaderParameters = PendingState.ShaderParameters[CrossCompiler::SHADER_STAGE_COMPUTE];
@@ -2610,7 +2610,7 @@ void FOpenGLDynamicRHI::RHIDrawPrimitiveIndirect(uint32 PrimitiveType,FVertexBuf
 			FOpenGL::PatchParameteri(GL_PATCH_VERTICES, PatchSize);
 		} 
 
-		DYNAMIC_CAST_OPENGLRESOURCE(VertexBuffer,ArgumentBuffer);
+		FOpenGLVertexBuffer* ArgumentBuffer = ResourceCast(ArgumentBufferRHI);
 
 
 		glBindBuffer( GL_DRAW_INDIRECT_BUFFER, ArgumentBuffer->Resource);
@@ -2635,7 +2635,7 @@ void FOpenGLDynamicRHI::RHIDrawIndexedIndirect(FIndexBufferRHIParamRef IndexBuff
 	{
 		VERIFY_GL_SCOPE();
 
-		DYNAMIC_CAST_OPENGLRESOURCE(IndexBuffer,IndexBuffer);
+		FOpenGLIndexBuffer* IndexBuffer = ResourceCast(IndexBufferRHI);
 	GPUProfilingData.RegisterGPUWork(1);
 
 		check(ArgumentsBufferRHI);
@@ -2670,7 +2670,7 @@ void FOpenGLDynamicRHI::RHIDrawIndexedIndirect(FIndexBufferRHIParamRef IndexBuff
 
 		GLenum IndexType = IndexBuffer->GetStride() == sizeof(uint32) ? GL_UNSIGNED_INT : GL_UNSIGNED_SHORT;
 
-		DYNAMIC_CAST_OPENGLRESOURCE(VertexBuffer,ArgumentsBuffer);
+		FOpenGLStructuredBuffer* ArgumentsBuffer = ResourceCast(ArgumentsBufferRHI);
 
 
 		glBindBuffer( GL_DRAW_INDIRECT_BUFFER, ArgumentsBuffer->Resource);
@@ -2695,7 +2695,7 @@ void FOpenGLDynamicRHI::RHIDrawIndexedPrimitive(FIndexBufferRHIParamRef IndexBuf
 	SCOPE_CYCLE_COUNTER_DETAILED(STAT_OpenGLDrawPrimitiveTime);
 	VERIFY_GL_SCOPE();
 
-	DYNAMIC_CAST_OPENGLRESOURCE(IndexBuffer,IndexBuffer);
+	FOpenGLIndexBuffer* IndexBuffer = ResourceCast(IndexBufferRHI);
 
 	RHI_DRAW_CALL_STATS(PrimitiveType,NumPrimitives*NumInstances);
 
@@ -2764,7 +2764,7 @@ void FOpenGLDynamicRHI::RHIDrawIndexedPrimitiveIndirect(uint32 PrimitiveType,FIn
 	{
 		VERIFY_GL_SCOPE();
 
-		DYNAMIC_CAST_OPENGLRESOURCE(IndexBuffer,IndexBuffer);
+		FOpenGLIndexBuffer* IndexBuffer = ResourceCast(IndexBufferRHI);
 		GPUProfilingData.RegisterGPUWork(1);
 
 		check(ArgumentBufferRHI);
@@ -2796,7 +2796,7 @@ void FOpenGLDynamicRHI::RHIDrawIndexedPrimitiveIndirect(uint32 PrimitiveType,FIn
 
 		GLenum IndexType = IndexBuffer->GetStride() == sizeof(uint32) ? GL_UNSIGNED_INT : GL_UNSIGNED_SHORT;
 
-		DYNAMIC_CAST_OPENGLRESOURCE(VertexBuffer,ArgumentBuffer);
+		FOpenGLVertexBuffer* ArgumentBuffer = ResourceCast(ArgumentBufferRHI);
 
 
 		glBindBuffer( GL_DRAW_INDIRECT_BUFFER, ArgumentBuffer->Resource);
@@ -3452,7 +3452,7 @@ void FOpenGLDynamicRHI::RHIDispatchComputeShader(uint32 ThreadGroupCountX, uint3
 		FComputeShaderRHIParamRef ComputeShaderRHI = PendingState.CurrentComputeShader;
 		check(ComputeShaderRHI);
 
-		DYNAMIC_CAST_OPENGLRESOURCE(ComputeShader,ComputeShader);
+		FOpenGLComputeShader* ComputeShader = ResourceCast(ComputeShaderRHI);
 
 		FOpenGLContextState& ContextState = GetContextStateForCurrentContext();
 
@@ -3484,8 +3484,8 @@ void FOpenGLDynamicRHI::RHIDispatchIndirectComputeShader(FVertexBufferRHIParamRe
 		FComputeShaderRHIParamRef ComputeShaderRHI = PendingState.CurrentComputeShader;
 		check(ComputeShaderRHI);
 
-		DYNAMIC_CAST_OPENGLRESOURCE(ComputeShader,ComputeShader);
-		DYNAMIC_CAST_OPENGLRESOURCE(VertexBuffer, ArgumentBuffer);
+		FOpenGLComputeShader* ComputeShader = ResourceCast(ComputeShaderRHI);
+		FOpenGLVertexBuffer* ArgumentBuffer = ResourceCast(ArgumentBufferRHI);
 
 		FOpenGLContextState& ContextState = GetContextStateForCurrentContext();
 
diff --git a/Engine/Source/Runtime/OpenGLDrv/Private/OpenGLIndexBuffer.cpp b/Engine/Source/Runtime/OpenGLDrv/Private/OpenGLIndexBuffer.cpp
--- a/Engine/Source/Runtime/OpenGLDrv/Private/OpenGLIndexBuffer.cpp
+++ b/Engine/Source/Runtime/OpenGLDrv/Private/OpenGLIndexBuffer.cpp
@@ -26,13 +26,13 @@ FIndexBufferRHIRef FOpenGLDynamicRHI::RHICreateIndexBuffer(uint32 Stride,uint32
 void* FOpenGLDynamicRHI::RHILockIndexBuffer(FIndexBufferRHIParamRef IndexBufferRHI,uint32 Offset,uint32 Size,EResourceLockMode LockMode)
 {
 	VERIFY_GL_SCOPE();
-	DYNAMIC_CAST_OPENGLRESOURCE(IndexBuffer,IndexBuffer);
+	FOpenGLIndexBuffer* IndexBuffer = ResourceCast(IndexBufferRHI);
 	return IndexBuffer->Lock(Offset, Size, LockMode == RLM_ReadOnly, IndexBuffer->IsDynamic());
 }
 
 void FOpenGLDynamicRHI::RHIUnlockIndexBuffer(FIndexBufferRHIParamRef IndexBufferRHI)
 {
 	VERIFY_GL_SCOPE();
-	DYNAMIC_CAST_OPENGLRESOURCE(IndexBuffer,IndexBuffer);
+	FOpenGLIndexBuffer* IndexBuffer = ResourceCast(IndexBufferRHI);
 	IndexBuffer->Unlock();
 }
diff --git a/Engine/Source/Runtime/OpenGLDrv/Private/OpenGLQuery.cpp b/Engine/Source/Runtime/OpenGLDrv/Private/OpenGLQuery.cpp
--- a/Engine/Source/Runtime/OpenGLDrv/Private/OpenGLQuery.cpp
+++ b/Engine/Source/Runtime/OpenGLDrv/Private/OpenGLQuery.cpp
@@ -36,7 +36,7 @@ void FOpenGLDynamicRHI::RHIBeginRenderQuery(FRenderQueryRHIParamRef QueryRHI)
 {
 	VERIFY_GL_SCOPE();
 
-	DYNAMIC_CAST_OPENGLRESOURCE(RenderQuery,Query);
+	FOpenGLRenderQuery* Query = ResourceCast(QueryRHI);
 	Query->bResultIsCached = false;
 	if(Query->QueryType == RQT_Occlusion)
 	{
@@ -69,7 +69,7 @@ void FOpenGLDynamicRHI::RHIEndRenderQuery(FRenderQueryRHIParamRef QueryRHI)
 {
 	VERIFY_GL_SCOPE();
 
-	DYNAMIC_CAST_OPENGLRESOURCE(RenderQuery,Query);
+	FOpenGLRenderQuery* Query = ResourceCast(QueryRHI);
 
 	if (Query)
 	{
@@ -101,7 +101,7 @@ bool FOpenGLDynamicRHI::RHIGetRenderQueryResult(FRenderQueryRHIParamRef QueryRHI
 	check(IsInRenderingThread());
 	VERIFY_GL_SCOPE();
 
-	DYNAMIC_CAST_OPENGLRESOURCE(RenderQuery,Query);
+	FOpenGLRenderQuery* Query = ResourceCast(QueryRHI);
 
 	if (!Query)
 	{
diff --git a/Engine/Source/Runtime/OpenGLDrv/Private/OpenGLShaders.cpp b/Engine/Source/Runtime/OpenGLDrv/Private/OpenGLShaders.cpp
--- a/Engine/Source/Runtime/OpenGLDrv/Private/OpenGLShaders.cpp
+++ b/Engine/Source/Runtime/OpenGLDrv/Private/OpenGLShaders.cpp
@@ -2065,12 +2065,12 @@ FBoundShaderStateRHIRef FOpenGLDynamicRHI::RHICreateBoundShaderState(
 	{
 		check(VertexDeclarationRHI);
 		
-		DYNAMIC_CAST_OPENGLRESOURCE(VertexDeclaration,VertexDeclaration);
-		DYNAMIC_CAST_OPENGLRESOURCE(VertexShader,VertexShader);
-		DYNAMIC_CAST_OPENGLRESOURCE(PixelShader,PixelShader);
-		DYNAMIC_CAST_OPENGLRESOURCE(HullShader,HullShader);
-		DYNAMIC_CAST_OPENGLRESOURCE(DomainShader,DomainShader);
-		DYNAMIC_CAST_OPENGLRESOURCE(GeometryShader,GeometryShader);
+		FOpenGLVertexDeclaration* VertexDeclaration = ResourceCast(VertexDeclarationRHI);
+		FOpenGLVertexShader* VertexShader = ResourceCast(VertexShaderRHI);
+		FOpenGLPixelShader* PixelShader = ResourceCast(PixelShaderRHI);
+		FOpenGLHullShader* HullShader = ResourceCast(HullShaderRHI);
+		FOpenGLDomainShader* DomainShader = ResourceCast(DomainShaderRHI);
+		FOpenGLGeometryShader* GeometryShader = ResourceCast(GeometryShaderRHI);
 
 		FOpenGLLinkedProgramConfiguration Config;
 
@@ -2430,12 +2430,12 @@ FOpenGLBoundShaderState::FOpenGLBoundShaderState(
 	:	CacheLink(InVertexDeclarationRHI, InVertexShaderRHI, InPixelShaderRHI,
 		InHullShaderRHI, InDomainShaderRHI,	InGeometryShaderRHI, this)
 {
-	DYNAMIC_CAST_OPENGLRESOURCE(VertexDeclaration,InVertexDeclaration);
-	DYNAMIC_CAST_OPENGLRESOURCE(VertexShader,InVertexShader);
-	DYNAMIC_CAST_OPENGLRESOURCE(PixelShader,InPixelShader);
-	DYNAMIC_CAST_OPENGLRESOURCE(HullShader,InHullShader);
-	DYNAMIC_CAST_OPENGLRESOURCE(DomainShader,InDomainShader);
-	DYNAMIC_CAST_OPENGLRESOURCE(GeometryShader,InGeometryShader);
+	FOpenGLVertexDeclaration* InVertexDeclaration = FOpenGLDynamicRHI::ResourceCast(InVertexDeclarationRHI);
+	FOpenGLVertexShader* InVertexShader = FOpenGLDynamicRHI::ResourceCast(InVertexShaderRHI);
+	FOpenGLPixelShader* InPixelShader = FOpenGLDynamicRHI::ResourceCast(InPixelShaderRHI);
+	FOpenGLHullShader* InHullShader = FOpenGLDynamicRHI::ResourceCast(InHullShaderRHI);
+	FOpenGLDomainShader* InDomainShader = FOpenGLDynamicRHI::ResourceCast(InDomainShaderRHI);
+	FOpenGLGeometryShader* InGeometryShader = FOpenGLDynamicRHI::ResourceCast(InGeometryShaderRHI);
 
 	VertexDeclaration = InVertexDeclaration;
 	VertexShader = InVertexShader;
@@ -2496,7 +2496,7 @@ bool FOpenGLComputeShader::NeedsUAVStage(int32 UAVStageIndex)
 void FOpenGLDynamicRHI::BindPendingComputeShaderState(FOpenGLContextState& ContextState, FComputeShaderRHIParamRef ComputeShaderRHI)
 {
 	VERIFY_GL_SCOPE();
-	DYNAMIC_CAST_OPENGLRESOURCE(ComputeShader,ComputeShader);
+	FOpenGLComputeShader* ComputeShader = ResourceCast(ComputeShaderRHI);
 	bool ForceUniformBindingUpdate = false;
 
 	GLuint PendingProgram = ComputeShader->LinkedProgram->Program;
diff --git a/Engine/Source/Runtime/OpenGLDrv/Private/OpenGLStructuredBuffer.cpp b/Engine/Source/Runtime/OpenGLDrv/Private/OpenGLStructuredBuffer.cpp
--- a/Engine/Source/Runtime/OpenGLDrv/Private/OpenGLStructuredBuffer.cpp
+++ b/Engine/Source/Runtime/OpenGLDrv/Private/OpenGLStructuredBuffer.cpp
@@ -26,13 +26,13 @@ FStructuredBufferRHIRef FOpenGLDynamicRHI::RHICreateStructuredBuffer(uint32 Stri
 void* FOpenGLDynamicRHI::RHILockStructuredBuffer(FStructuredBufferRHIParamRef StructuredBufferRHI,uint32 Offset,uint32 Size,EResourceLockMode LockMode)
 {
 	VERIFY_GL_SCOPE();
-	DYNAMIC_CAST_OPENGLRESOURCE(StructuredBuffer,StructuredBuffer);
+	FOpenGLStructuredBuffer* StructuredBuffer = ResourceCast(StructuredBufferRHI);
 	return StructuredBuffer->Lock(Offset, Size, LockMode == RLM_ReadOnly, StructuredBuffer->IsDynamic());
 }
 
 void FOpenGLDynamicRHI::RHIUnlockStructuredBuffer(FStructuredBufferRHIParamRef StructuredBufferRHI)
 {
 	VERIFY_GL_SCOPE();
-	DYNAMIC_CAST_OPENGLRESOURCE(StructuredBuffer,StructuredBuffer);
+	FOpenGLStructuredBuffer* StructuredBuffer = ResourceCast(StructuredBufferRHI);
 	StructuredBuffer->Unlock();
 }
diff --git a/Engine/Source/Runtime/OpenGLDrv/Private/OpenGLTexture.cpp b/Engine/Source/Runtime/OpenGLDrv/Private/OpenGLTexture.cpp
--- a/Engine/Source/Runtime/OpenGLDrv/Private/OpenGLTexture.cpp
+++ b/Engine/Source/Runtime/OpenGLDrv/Private/OpenGLTexture.cpp
@@ -1694,7 +1694,7 @@ void FOpenGLDynamicRHI::RHIGetResourceInfo(FTextureRHIParamRef Ref, FRHIResource
 
 FShaderResourceViewRHIRef FOpenGLDynamicRHI::RHICreateShaderResourceView(FTexture2DRHIParamRef Texture2DRHI, uint8 MipLevel)
 {
-	DYNAMIC_CAST_OPENGLRESOURCE(Texture2D,Texture2D);
+	FOpenGLTexture2D* Texture2D = ResourceCast(Texture2DRHI);
 
 	FOpenGLShaderResourceView *View = 0;
 
@@ -1724,7 +1724,7 @@ FShaderResourceViewRHIRef FOpenGLDynamicRHI::RHICreateShaderResourceView(FTextur
 
 FShaderResourceViewRHIRef FOpenGLDynamicRHI::RHICreateShaderResourceView(FTexture2DRHIParamRef Texture2DRHI, uint8 MipLevel, uint8 NumMipLevels, uint8 Format)
 {
-	DYNAMIC_CAST_OPENGLRESOURCE(Texture2D,Texture2D);
+	FOpenGLTexture2D* Texture2D = ResourceCast(Texture2DRHI);
 
 	FOpenGLShaderResourceView *View = 0;
 
@@ -1840,7 +1840,7 @@ FTexture2DRHIRef FOpenGLDynamicRHI::RHIAsyncReallocateTexture2D(FTexture2DRHIPar
 {
 	VERIFY_GL_SCOPE();
 
-	DYNAMIC_CAST_OPENGLRESOURCE(Texture2D,Texture2D);
+	FOpenGLTexture2D* Texture2D = ResourceCast(Texture2DRHI);
 
 	// Allocate a new texture.
 	FOpenGLTexture2D* NewTexture2D = (FOpenGLTexture2D*)CreateOpenGLTexture(NewSizeX,NewSizeY,false,false, Texture2D->GetFormat(),NewMipCount,1,1, Texture2D->GetFlags());
@@ -1932,25 +1932,25 @@ ETextureReallocationStatus FOpenGLDynamicRHI::RHICancelAsyncReallocateTexture2D(
 
 void* FOpenGLDynamicRHI::RHILockTexture2D(FTexture2DRHIParamRef TextureRHI,uint32 MipIndex,EResourceLockMode LockMode,uint32& DestStride,bool bLockWithinMiptail)
 {
-	DYNAMIC_CAST_OPENGLRESOURCE(Texture2D,Texture);
+	FOpenGLTexture2D* Texture = ResourceCast(TextureRHI);
 	return Texture->Lock(MipIndex,0,LockMode,DestStride);
 }
 
 void FOpenGLDynamicRHI::RHIUnlockTexture2D(FTexture2DRHIParamRef TextureRHI,uint32 MipIndex,bool bLockWithinMiptail)
 {
-	DYNAMIC_CAST_OPENGLRESOURCE(Texture2D,Texture);
+	FOpenGLTexture2D* Texture = ResourceCast(TextureRHI);
 	Texture->Unlock(MipIndex, 0);
 }
 
 void* FOpenGLDynamicRHI::RHILockTexture2DArray(FTexture2DArrayRHIParamRef TextureRHI,uint32 TextureIndex,uint32 MipIndex,EResourceLockMode LockMode,uint32& DestStride,bool bLockWithinMiptail)
 {
-	DYNAMIC_CAST_OPENGLRESOURCE(Texture2DArray,Texture);
+	FOpenGLTexture2DArray* Texture = ResourceCast(TextureRHI);
 	return Texture->Lock(MipIndex,TextureIndex,LockMode,DestStride);
 }
 
 void FOpenGLDynamicRHI::RHIUnlockTexture2DArray(FTexture2DArrayRHIParamRef TextureRHI,uint32 TextureIndex,uint32 MipIndex,bool bLockWithinMiptail)
 {
-	DYNAMIC_CAST_OPENGLRESOURCE(Texture2DArray,Texture);
+	FOpenGLTexture2DArray* Texture = ResourceCast(TextureRHI);
 	Texture->Unlock(MipIndex, TextureIndex);
 }
 
@@ -1959,7 +1959,7 @@ void FOpenGLDynamicRHI::RHIUpdateTexture2D(FTexture2DRHIParamRef TextureRHI,uint
 	VERIFY_GL_SCOPE();
 	check( FOpenGL::SupportsPixelBuffers() );
 
-	DYNAMIC_CAST_OPENGLRESOURCE(Texture2D,Texture);
+	FOpenGLTexture2D* Texture = ResourceCast(TextureRHI);
 
 	// Use a texture stage that's not likely to be used for draws, to avoid waiting
 	FOpenGLContextState& ContextState = GetContextStateForCurrentContext();
@@ -1991,7 +1991,7 @@ void FOpenGLDynamicRHI::RHIUpdateTexture3D(FTexture3DRHIParamRef TextureRHI,uint
 {
 	VERIFY_GL_SCOPE();
 	check( FOpenGL::SupportsPixelBuffers() && FOpenGL::SupportsTexture3D() );
-	DYNAMIC_CAST_OPENGLRESOURCE(Texture3D,Texture);
+	FOpenGLTexture3D* Texture = ResourceCast(TextureRHI);
 
 	// Use a texture stage that's not likely to be used for draws, to avoid waiting
 	FOpenGLContextState& ContextState = GetContextStateForCurrentContext();
@@ -2086,13 +2086,13 @@ FTextureCubeRHIRef FOpenGLDynamicRHI::RHICreateTextureCubeArray( uint32 Size, ui
 
 void* FOpenGLDynamicRHI::RHILockTextureCubeFace(FTextureCubeRHIParamRef TextureCubeRHI,uint32 FaceIndex,uint32 ArrayIndex,uint32 MipIndex,EResourceLockMode LockMode,uint32& DestStride,bool bLockWithinMiptail)
 {
-	DYNAMIC_CAST_OPENGLRESOURCE(TextureCube,TextureCube);
+	FOpenGLTextureCube* TextureCube = ResourceCast(TextureCubeRHI);
 	return TextureCube->Lock(MipIndex,FaceIndex + 6 * ArrayIndex,LockMode,DestStride);
 }
 
 void FOpenGLDynamicRHI::RHIUnlockTextureCubeFace(FTextureCubeRHIParamRef TextureCubeRHI,uint32 FaceIndex,uint32 ArrayIndex,uint32 MipIndex,bool bLockWithinMiptail)
 {
-	DYNAMIC_CAST_OPENGLRESOURCE(TextureCube,TextureCube);
+	FOpenGLTextureCube* TextureCube = ResourceCast(TextureCubeRHI);
 	TextureCube->Unlock(MipIndex,FaceIndex + ArrayIndex * 6);
 }
 
diff --git a/Engine/Source/Runtime/OpenGLDrv/Private/OpenGLUAV.cpp b/Engine/Source/Runtime/OpenGLDrv/Private/OpenGLUAV.cpp
--- a/Engine/Source/Runtime/OpenGLDrv/Private/OpenGLUAV.cpp
+++ b/Engine/Source/Runtime/OpenGLDrv/Private/OpenGLUAV.cpp
@@ -9,7 +9,7 @@ FShaderResourceViewRHIRef FOpenGLDynamicRHI::RHICreateShaderResourceView(FVertex
 	GLuint TextureID = 0;
 	if ( FOpenGL::SupportsResourceView() )
 	{
-		DYNAMIC_CAST_OPENGLRESOURCE(VertexBuffer,VertexBuffer);
+		FOpenGLVertexBuffer* VertexBuffer = ResourceCast(VertexBufferRHI);
 
 		const uint32 FormatBPP = GPixelFormats[Format].BlockBytes;
 
@@ -51,14 +51,14 @@ FOpenGLShaderResourceView::~FOpenGLShaderResourceView()
 
 FUnorderedAccessViewRHIRef FOpenGLDynamicRHI::RHICreateUnorderedAccessView(FStructuredBufferRHIParamRef StructuredBufferRHI, bool bUseUAVCounter, bool bAppendBuffer)
 {
-	DYNAMIC_CAST_OPENGLRESOURCE(StructuredBuffer,StructuredBuffer);
+	FOpenGLStructuredBuffer* StructuredBuffer = ResourceCast(StructuredBufferRHI);
 	UE_LOG(LogRHI, Fatal,TEXT("%s not implemented yet"),ANSI_TO_TCHAR(__FUNCTION__)); 
 	return new FOpenGLUnorderedAccessView();
 }
 
 FUnorderedAccessViewRHIRef FOpenGLDynamicRHI::RHICreateUnorderedAccessView(FTextureRHIParamRef TextureRHI)
 {
-	DYNAMIC_CAST_OPENGLRESOURCE(Texture,Texture);
+	FOpenGLTexture* Texture = ResourceCast(TextureRHI);
 	check(Texture->GetFlags() & TexCreate_UAV);
 	return new FOpenGLTextureUnorderedAccessView(TextureRHI);
 }
@@ -82,7 +82,7 @@ FOpenGLVertexBufferUnorderedAccessView::FOpenGLVertexBufferUnorderedAccessView(
 	OpenGLRHI(InOpenGLRHI)
 {
 	VERIFY_GL_SCOPE();
-	DYNAMIC_CAST_OPENGLRESOURCE(VertexBuffer, InVertexBuffer);
+	FOpenGLVertexBuffer* InVertexBuffer = FOpenGLDynamicRHI::ResourceCast(InVertexBufferRHI);
 
 
 	const FOpenGLTextureFormat& GLFormat = GOpenGLTextureFormats[Format];
@@ -114,13 +114,13 @@ FOpenGLVertexBufferUnorderedAccessView::~FOpenGLVertexBufferUnorderedAccessView(
 
 FUnorderedAccessViewRHIRef FOpenGLDynamicRHI::RHICreateUnorderedAccessView(FVertexBufferRHIParamRef VertexBufferRHI,uint8 Format)
 {
-	DYNAMIC_CAST_OPENGLRESOURCE(VertexBuffer,VertexBuffer);
+	FOpenGLVertexBuffer* VertexBuffer = ResourceCast(VertexBufferRHI);
 	return new FOpenGLVertexBufferUnorderedAccessView(this, VertexBufferRHI, Format);
 }
 
 FShaderResourceViewRHIRef FOpenGLDynamicRHI::RHICreateShaderResourceView(FStructuredBufferRHIParamRef StructuredBufferRHI)
 {
-	DYNAMIC_CAST_OPENGLRESOURCE(StructuredBuffer,StructuredBuffer);
+	FOpenGLStructuredBuffer* StructuredBuffer = ResourceCast(StructuredBufferRHI);
 	UE_LOG(LogRHI, Fatal,TEXT("OpenGL RHI doesn't support RHICreateShaderResourceView yet!"));
 	return new FOpenGLShaderResourceView(this,0,GL_TEXTURE_BUFFER);
 }
diff --git a/Engine/Source/Runtime/OpenGLDrv/Private/OpenGLVertexBuffer.cpp b/Engine/Source/Runtime/OpenGLDrv/Private/OpenGLVertexBuffer.cpp
--- a/Engine/Source/Runtime/OpenGLDrv/Private/OpenGLVertexBuffer.cpp
+++ b/Engine/Source/Runtime/OpenGLDrv/Private/OpenGLVertexBuffer.cpp
@@ -167,7 +167,7 @@ void* FOpenGLDynamicRHI::RHILockVertexBuffer(FVertexBufferRHIParamRef VertexBuff
 	check(Size > 0);
 
 	VERIFY_GL_SCOPE();
-	DYNAMIC_CAST_OPENGLRESOURCE(VertexBuffer,VertexBuffer);
+	FOpenGLVertexBuffer* VertexBuffer = ResourceCast(VertexBufferRHI);
 	if( !(FOpenGL::SupportsVertexAttribBinding() && OpenGLConsoleVariables::bUseVAB) && VertexBuffer->GetUsage() & BUF_ZeroStride )
 	{
 		check( Offset + Size <= VertexBuffer->GetSize() );
@@ -192,7 +192,7 @@ void* FOpenGLDynamicRHI::RHILockVertexBuffer(FVertexBufferRHIParamRef VertexBuff
 void FOpenGLDynamicRHI::RHIUnlockVertexBuffer(FVertexBufferRHIParamRef VertexBufferRHI)
 {
 	VERIFY_GL_SCOPE();
-	DYNAMIC_CAST_OPENGLRESOURCE(VertexBuffer,VertexBuffer);
+	FOpenGLVertexBuffer* VertexBuffer = ResourceCast(VertexBufferRHI);
 	if( (FOpenGL::SupportsVertexAttribBinding() && OpenGLConsoleVariables::bUseVAB) || !( VertexBuffer->GetUsage() & BUF_ZeroStride ) )
 	{
 		if (!RetireAllocation(VertexBuffer))
@@ -209,8 +209,8 @@ void FOpenGLDynamicRHI::RHICopyVertexBuffer(FVertexBufferRHIParamRef SourceBuffe
 {
 	VERIFY_GL_SCOPE();
 	check( FOpenGL::SupportsCopyBuffer() );
-	DYNAMIC_CAST_OPENGLRESOURCE(VertexBuffer,SourceBuffer);
-	DYNAMIC_CAST_OPENGLRESOURCE(VertexBuffer,DestBuffer);
+	FOpenGLVertexBuffer* SourceBuffer = ResourceCast(SourceBufferRHI);
+	FOpenGLVertexBuffer* DestBuffer = ResourceCast(DestBufferRHI);
 	check(SourceBuffer->GetSize() == DestBuffer->GetSize());
 
 	glBindBuffer(GL_COPY_READ_BUFFER,SourceBuffer->Resource);
diff --git a/Engine/Source/Runtime/OpenGLDrv/Private/OpenGLViewport.cpp b/Engine/Source/Runtime/OpenGLDrv/Private/OpenGLViewport.cpp
--- a/Engine/Source/Runtime/OpenGLDrv/Private/OpenGLViewport.cpp
+++ b/Engine/Source/Runtime/OpenGLDrv/Private/OpenGLViewport.cpp
@@ -74,7 +74,7 @@ FViewportRHIRef FOpenGLDynamicRHI::RHICreateViewport(void* WindowHandle,uint32 S
 
 void FOpenGLDynamicRHI::RHIResizeViewport(FViewportRHIParamRef ViewportRHI,uint32 SizeX,uint32 SizeY,bool bIsFullscreen)
 {
-	DYNAMIC_CAST_OPENGLRESOURCE(Viewport,Viewport);
+	FOpenGLViewport* Viewport = ResourceCast(ViewportRHI);
 	check( IsInGameThread() );
 
 //#if !PLATFORM_MAC
@@ -99,7 +99,7 @@ void FOpenGLDynamicRHI::RHIBeginDrawingViewport(FViewportRHIParamRef ViewportRHI
 {
 	VERIFY_GL_SCOPE();
 
-	DYNAMIC_CAST_OPENGLRESOURCE(Viewport,Viewport);
+	FOpenGLViewport* Viewport = ResourceCast(ViewportRHI);
 
 	SCOPE_CYCLE_COUNTER(STAT_OpenGLPresentTime);
 
@@ -138,7 +138,7 @@ void FOpenGLDynamicRHI::RHIEndDrawingViewport(FViewportRHIParamRef ViewportRHI,b
 {
 	VERIFY_GL_SCOPE();
 
-	DYNAMIC_CAST_OPENGLRESOURCE(Viewport,Viewport);
+	FOpenGLViewport* Viewport = ResourceCast(ViewportRHI);
 
 	SCOPE_CYCLE_COUNTER(STAT_OpenGLPresentTime);
 
@@ -212,7 +212,7 @@ bool FOpenGLDynamicRHI::RHIIsDrawingViewport()
 
 FTexture2DRHIRef FOpenGLDynamicRHI::RHIGetViewportBackBuffer(FViewportRHIParamRef ViewportRHI)
 {
-	DYNAMIC_CAST_OPENGLRESOURCE(Viewport,Viewport);
+	FOpenGLViewport* Viewport = ResourceCast(ViewportRHI);
 	return Viewport->GetBackBuffer();
 }
 
diff --git a/Engine/Source/Runtime/OpenGLDrv/Public/OpenGLDrv.h b/Engine/Source/Runtime/OpenGLDrv/Public/OpenGLDrv.h
--- a/Engine/Source/Runtime/OpenGLDrv/Public/OpenGLDrv.h
+++ b/Engine/Source/Runtime/OpenGLDrv/Public/OpenGLDrv.h
@@ -37,10 +37,6 @@
 #define OPENGL_USE_BINDABLE_UNIFORMS 0
 #define OPENGL_USE_BLIT_FOR_BACK_BUFFER 1
 
-/** This is a macro that casts a dynamically bound RHI reference to the appropriate OpenGL type. */
-#define DYNAMIC_CAST_OPENGLRESOURCE(Type,Name) \
-	FOpenGL##Type* Name = (FOpenGL##Type*)Name##RHI;
-
 // OpenGL RHI public headers.
 #include "OpenGLUtil.h"
 #include "OpenGLState.h"
@@ -305,6 +301,12 @@ public:
 	virtual void Init();
 	virtual void Shutdown();
 
+	template<typename TRHIType>
+	static FORCEINLINE typename TOpenGLResourceTraits<TRHIType>::TConcreteType* ResourceCast(TRHIType* Resource)
+	{
+		return static_cast<typename TOpenGLResourceTraits<TRHIType>::TConcreteType*>(Resource);
+	}
+
 	#define DEFINE_RHIMETHOD(Type,Name,ParameterTypesAndNames,ParameterNames,ReturnStatement,NullImplementation) virtual Type RHI##Name ParameterTypesAndNames
 	#include "RHIMethods.h"
 	#undef DEFINE_RHIMETHOD
diff --git a/Engine/Source/Runtime/OpenGLDrv/Public/OpenGLResources.h b/Engine/Source/Runtime/OpenGLDrv/Public/OpenGLResources.h
--- a/Engine/Source/Runtime/OpenGLDrv/Public/OpenGLResources.h
+++ b/Engine/Source/Runtime/OpenGLDrv/Public/OpenGLResources.h
@@ -1397,3 +1397,113 @@ private:
 	FCustomPresentRHIRef CustomPresent;
 };
 
+template<class T>
+struct TOpenGLResourceTraits
+{
+};
+template<>
+struct TOpenGLResourceTraits<FRHIVertexDeclaration>
+{
+	typedef FOpenGLVertexDeclaration TConcreteType;
+};
+template<>
+struct TOpenGLResourceTraits<FRHIVertexShader>
+{
+	typedef FOpenGLVertexShader TConcreteType;
+};
+template<>
+struct TOpenGLResourceTraits<FRHIGeometryShader>
+{
+	typedef FOpenGLGeometryShader TConcreteType;
+};
+template<>
+struct TOpenGLResourceTraits<FRHIHullShader>
+{
+	typedef FOpenGLHullShader TConcreteType;
+};
+template<>
+struct TOpenGLResourceTraits<FRHIDomainShader>
+{
+	typedef FOpenGLDomainShader TConcreteType;
+};
+template<>
+struct TOpenGLResourceTraits<FRHIPixelShader>
+{
+	typedef FOpenGLPixelShader TConcreteType;
+};
+template<>
+struct TOpenGLResourceTraits<FRHIComputeShader>
+{
+	typedef FOpenGLComputeShader TConcreteType;
+};
+template<>
+struct TOpenGLResourceTraits<FRHIBoundShaderState>
+{
+	typedef FOpenGLBoundShaderState TConcreteType;
+};
+template<>
+struct TOpenGLResourceTraits<FRHITexture3D>
+{
+	typedef FOpenGLTexture3D TConcreteType;
+};
+template<>
+struct TOpenGLResourceTraits<FRHITexture>
+{
+	typedef FOpenGLTexture TConcreteType;
+};
+template<>
+struct TOpenGLResourceTraits<FRHITexture2D>
+{
+	typedef FOpenGLTexture2D TConcreteType;
+};
+template<>
+struct TOpenGLResourceTraits<FRHITexture2DArray>
+{
+	typedef FOpenGLTexture2DArray TConcreteType;
+};
+template<>
+struct TOpenGLResourceTraits<FRHITextureCube>
+{
+	typedef FOpenGLTextureCube TConcreteType;
+};
+template<>
+struct TOpenGLResourceTraits<FRHIRenderQuery>
+{
+	typedef FOpenGLRenderQuery TConcreteType;
+};
+template<>
+struct TOpenGLResourceTraits<FRHIUniformBuffer>
+{
+	typedef FOpenGLUniformBuffer TConcreteType;
+};
+template<>
+struct TOpenGLResourceTraits<FRHIIndexBuffer>
+{
+	typedef FOpenGLIndexBuffer TConcreteType;
+};
+template<>
+struct TOpenGLResourceTraits<FRHIStructuredBuffer>
+{
+	typedef FOpenGLStructuredBuffer TConcreteType;
+};
+template<>
+struct TOpenGLResourceTraits<FRHIVertexBuffer>
+{
+	typedef FOpenGLVertexBuffer TConcreteType;
+};
+template<>
+struct TOpenGLResourceTraits<FRHIShaderResourceView>
+{
+	typedef FOpenGLShaderResourceView TConcreteType;
+};
+template<>
+struct TOpenGLResourceTraits<FRHIUnorderedAccessView>
+{
+	typedef FOpenGLUnorderedAccessView TConcreteType;
+};
+
+template<>
+struct TOpenGLResourceTraits<FRHIViewport>
+{
+	typedef FOpenGLViewport TConcreteType;
+};
diff --git a/Engine/Source/Runtime/OpenGLDrv/Public/OpenGLState.h b/Engine/Source/Runtime/OpenGLDrv/Public/OpenGLState.h
--- a/Engine/Source/Runtime/OpenGLDrv/Public/OpenGLState.h
+++ b/Engine/Source/Runtime/OpenGLDrv/Public/OpenGLState.h
@@ -508,3 +508,24 @@ struct FOpenGLRHIState : public FOpenGLCommonState
 		FOpenGLCommonState::CleanupResources();
 	}
 };
+
+template<>
+struct TOpenGLResourceTraits<FRHISamplerState>
+{
+	typedef FOpenGLSamplerState TConcreteType;
+};
+template<>
+struct TOpenGLResourceTraits<FRHIRasterizerState>
+{
+	typedef FOpenGLRasterizerState TConcreteType;
+};
+template<>
+struct TOpenGLResourceTraits<FRHIDepthStencilState>
+{
+	typedef FOpenGLDepthStencilState TConcreteType;
+};
+template<>
+struct TOpenGLResourceTraits<FRHIBlendState>
+{
+	typedef FOpenGLBlendState TConcreteType;
+};
diff --git a/Engine/Source/Runtime/Windows/D3D11RHI/Private/D3D11Commands.cpp b/Engine/Source/Runtime/Windows/D3D11RHI/Private/D3D11Commands.cpp
--- a/Engine/Source/Runtime/Windows/D3D11RHI/Private/D3D11Commands.cpp
+++ b/Engine/Source/Runtime/Windows/D3D11RHI/Private/D3D11Commands.cpp
@@ -35,7 +35,7 @@ TGlobalResource<FVector4VertexDeclaration> GD3D11Vector4VertexDeclaration;
 { \
 	ID3D11##ShaderType* CachedShader; \
 	InStateCache.Get##ShaderType(&CachedShader); \
-	DYNAMIC_CAST_D3D11RESOURCE(ShaderType,ShaderType); \
+	FD3D11##ShaderType* ShaderType = FD3D11DynamicRHI::ResourceCast(ShaderType##RHI); \
 	ensureMsgf(CachedShader == ShaderType->Resource, TEXT("Parameters are being set for a %s which is not currently bound"), TEXT( #ShaderType )); \
 	if (CachedShader) { CachedShader->Release(); } \
 }
@@ -47,6 +47,7 @@ DECLARE_ISBOUNDSHADER(HullShader)
 DECLARE_ISBOUNDSHADER(DomainShader)
 DECLARE_ISBOUNDSHADER(ComputeShader)
 
+
 #if DO_CHECK
 #define VALIDATE_BOUND_SHADER(s) ValidateBoundShader(StateCache, s)
 #else
@@ -82,7 +83,7 @@ void FD3D11DynamicRHI::RHIGpuTimeEnd(uint32 Hash, bool bCompute)
 // Vertex state.
 void FD3D11DynamicRHI::RHISetStreamSource(uint32 StreamIndex,FVertexBufferRHIParamRef VertexBufferRHI,uint32 Stride,uint32 Offset)
 {
-	DYNAMIC_CAST_D3D11RESOURCE(VertexBuffer,VertexBuffer);
+	FD3D11VertexBuffer* VertexBuffer = ResourceCast(VertexBufferRHI);
 
 	ID3D11Buffer* D3DBuffer = VertexBuffer ? VertexBuffer->Resource : NULL;
 	StateCache.SetStreamSource(D3DBuffer, StreamIndex, Stride, Offset);
@@ -106,20 +107,20 @@ void FD3D11DynamicRHI::RHISetStreamOutTargets(uint32 NumTargets, const FVertexBu
 // Rasterizer state.
 void FD3D11DynamicRHI::RHISetRasterizerState(FRasterizerStateRHIParamRef NewStateRHI)
 {
-	DYNAMIC_CAST_D3D11RESOURCE(RasterizerState,NewState);
+	FD3D11RasterizerState* NewState = ResourceCast(NewStateRHI);
 	StateCache.SetRasterizerState(NewState->Resource);
 }
 
 void FD3D11DynamicRHI::RHISetComputeShader(FComputeShaderRHIParamRef ComputeShaderRHI)
 {
-	DYNAMIC_CAST_D3D11RESOURCE(ComputeShader,ComputeShader);
+	FD3D11ComputeShader* ComputeShader = ResourceCast(ComputeShaderRHI);
 	SetCurrentComputeShader(ComputeShaderRHI);
 }
 
 void FD3D11DynamicRHI::RHIDispatchComputeShader(uint32 ThreadGroupCountX, uint32 ThreadGroupCountY, uint32 ThreadGroupCountZ) 
 { 
 	FComputeShaderRHIParamRef ComputeShaderRHI = GetCurrentComputeShader();
-	DYNAMIC_CAST_D3D11RESOURCE(ComputeShader,ComputeShader);
+	FD3D11ComputeShader* ComputeShader = ResourceCast(ComputeShaderRHI);
 
 	StateCache.SetComputeShader(ComputeShader->Resource);
 
@@ -138,8 +139,8 @@ void FD3D11DynamicRHI::RHIDispatchComputeShader(uint32 ThreadGroupCountX, uint32
 void FD3D11DynamicRHI::RHIDispatchIndirectComputeShader(FVertexBufferRHIParamRef ArgumentBufferRHI, uint32 ArgumentOffset) 
 { 
 	FComputeShaderRHIParamRef ComputeShaderRHI = GetCurrentComputeShader();
-	DYNAMIC_CAST_D3D11RESOURCE(ComputeShader,ComputeShader);
-	DYNAMIC_CAST_D3D11RESOURCE(VertexBuffer,ArgumentBuffer);
+	FD3D11ComputeShader* ComputeShader = ResourceCast(ComputeShaderRHI);
+	FD3D11VertexBuffer* ArgumentBuffer = ResourceCast(ArgumentBufferRHI);
 
 	GPUProfilingData.RegisterGPUWork(1);
 
@@ -200,7 +201,7 @@ void FD3D11DynamicRHI::RHISetScissorRect(bool bEnable,uint32 MinX,uint32 MinY,ui
 */
 void FD3D11DynamicRHI::RHISetBoundShaderState( FBoundShaderStateRHIParamRef BoundShaderStateRHI)
 {
-	DYNAMIC_CAST_D3D11RESOURCE(BoundShaderState,BoundShaderState);
+	FD3D11BoundShaderState* BoundShaderState = ResourceCast(BoundShaderStateRHI);
 
 	StateCache.SetInputLayout(BoundShaderState->InputLayout);
 	StateCache.SetVertexShader(BoundShaderState->VertexShader);
@@ -316,7 +317,7 @@ void FD3D11DynamicRHI::RHISetUAVParameter(FComputeShaderRHIParamRef ComputeShade
 {
 	//VALIDATE_BOUND_SHADER(ComputeShaderRHI);
 
-	DYNAMIC_CAST_D3D11RESOURCE(UnorderedAccessView,UAV);
+	FD3D11UnorderedAccessView* UAV = ResourceCast(UAVRHI);
 
 	if(UAV)
 	{
@@ -333,7 +334,7 @@ void FD3D11DynamicRHI::RHISetUAVParameter(FComputeShaderRHIParamRef ComputeShade
 {
 	//VALIDATE_BOUND_SHADER(ComputeShaderRHI);
 
-	DYNAMIC_CAST_D3D11RESOURCE(UnorderedAccessView,UAV);
+	FD3D11UnorderedAccessView* UAV = ResourceCast(UAVRHI);
 	
 	if(UAV)
 	{
@@ -348,7 +349,7 @@ void FD3D11DynamicRHI::RHISetShaderResourceViewParameter(FPixelShaderRHIParamRef
 {
 	VALIDATE_BOUND_SHADER(PixelShaderRHI);
 
-	DYNAMIC_CAST_D3D11RESOURCE(ShaderResourceView,SRV);
+	FD3D11ShaderResourceView* SRV = ResourceCast(SRVRHI);
 
 	FD3D11BaseShaderResource* Resource = nullptr;
 	ID3D11ShaderResourceView* D3D11SRV = nullptr;
@@ -366,7 +367,7 @@ void FD3D11DynamicRHI::RHISetShaderResourceViewParameter(FVertexShaderRHIParamRe
 {
 	VALIDATE_BOUND_SHADER(VertexShaderRHI);
 
-	DYNAMIC_CAST_D3D11RESOURCE(ShaderResourceView,SRV);
+	FD3D11ShaderResourceView* SRV = ResourceCast(SRVRHI);
 
 	FD3D11BaseShaderResource* Resource = nullptr;
 	ID3D11ShaderResourceView* D3D11SRV = nullptr;
@@ -384,7 +385,7 @@ void FD3D11DynamicRHI::RHISetShaderResourceViewParameter(FComputeShaderRHIParamR
 {
 	//VALIDATE_BOUND_SHADER(ComputeShaderRHI);
 
-	DYNAMIC_CAST_D3D11RESOURCE(ShaderResourceView,SRV);
+	FD3D11ShaderResourceView* SRV = ResourceCast(SRVRHI);
 
 	FD3D11BaseShaderResource* Resource = nullptr;
 	ID3D11ShaderResourceView* D3D11SRV = nullptr;
@@ -402,7 +403,7 @@ void FD3D11DynamicRHI::RHISetShaderResourceViewParameter(FHullShaderRHIParamRef
 {
 	VALIDATE_BOUND_SHADER(HullShaderRHI);
 
-	DYNAMIC_CAST_D3D11RESOURCE(ShaderResourceView,SRV);
+	FD3D11ShaderResourceView* SRV = ResourceCast(SRVRHI);
 
 	FD3D11BaseShaderResource* Resource = nullptr;
 	ID3D11ShaderResourceView* D3D11SRV = nullptr;
@@ -420,7 +421,7 @@ void FD3D11DynamicRHI::RHISetShaderResourceViewParameter(FDomainShaderRHIParamRe
 {
 	VALIDATE_BOUND_SHADER(DomainShaderRHI);
 
-	DYNAMIC_CAST_D3D11RESOURCE(ShaderResourceView,SRV);
+	FD3D11ShaderResourceView* SRV = ResourceCast(SRVRHI);
 
 	FD3D11BaseShaderResource* Resource = nullptr;
 	ID3D11ShaderResourceView* D3D11SRV = nullptr;
@@ -438,7 +439,7 @@ void FD3D11DynamicRHI::RHISetShaderResourceViewParameter(FGeometryShaderRHIParam
 {
 	VALIDATE_BOUND_SHADER(GeometryShaderRHI);
 
-	DYNAMIC_CAST_D3D11RESOURCE(ShaderResourceView,SRV);
+	FD3D11ShaderResourceView* SRV = ResourceCast(SRVRHI);
 
 	FD3D11BaseShaderResource* Resource = nullptr;
 	ID3D11ShaderResourceView* D3D11SRV = nullptr;
@@ -456,8 +457,8 @@ void FD3D11DynamicRHI::RHISetShaderSampler(FVertexShaderRHIParamRef VertexShader
 {
 	VALIDATE_BOUND_SHADER(VertexShaderRHI);
 
-	DYNAMIC_CAST_D3D11RESOURCE(VertexShader,VertexShader);
-	DYNAMIC_CAST_D3D11RESOURCE(SamplerState,NewState);
+	FD3D11VertexShader* VertexShader = ResourceCast(VertexShaderRHI);
+	FD3D11SamplerState* NewState = ResourceCast(NewStateRHI);
 
 	ID3D11SamplerState* StateResource = NewState->Resource;
 	StateCache.SetSamplerState<SF_Vertex>(StateResource, SamplerIndex);
@@ -467,8 +468,8 @@ void FD3D11DynamicRHI::RHISetShaderSampler(FHullShaderRHIParamRef HullShaderRHI,
 {
 	VALIDATE_BOUND_SHADER(HullShaderRHI);
 
-	DYNAMIC_CAST_D3D11RESOURCE(HullShader,HullShader);
-	DYNAMIC_CAST_D3D11RESOURCE(SamplerState,NewState);
+	FD3D11HullShader* HullShader = ResourceCast(HullShaderRHI);
+	FD3D11SamplerState* NewState = ResourceCast(NewStateRHI);
 
 	ID3D11SamplerState* StateResource = NewState->Resource;
 	StateCache.SetSamplerState<SF_Hull>(StateResource, SamplerIndex);
@@ -478,8 +479,8 @@ void FD3D11DynamicRHI::RHISetShaderSampler(FDomainShaderRHIParamRef DomainShader
 {
 	VALIDATE_BOUND_SHADER(DomainShaderRHI);
 
-	DYNAMIC_CAST_D3D11RESOURCE(DomainShader,DomainShader);
-	DYNAMIC_CAST_D3D11RESOURCE(SamplerState,NewState);
+	FD3D11DomainShader* DomainShader = ResourceCast(DomainShaderRHI);
+	FD3D11SamplerState* NewState = ResourceCast(NewStateRHI);
 
 	ID3D11SamplerState* StateResource = NewState->Resource;
 	StateCache.SetSamplerState<SF_Domain>(StateResource, SamplerIndex);
@@ -489,8 +490,8 @@ void FD3D11DynamicRHI::RHISetShaderSampler(FGeometryShaderRHIParamRef GeometrySh
 {
 	VALIDATE_BOUND_SHADER(GeometryShaderRHI);
 
-	DYNAMIC_CAST_D3D11RESOURCE(GeometryShader,GeometryShader);
-	DYNAMIC_CAST_D3D11RESOURCE(SamplerState,NewState);
+	FD3D11GeometryShader* GeometryShader = ResourceCast(GeometryShaderRHI);
+	FD3D11SamplerState* NewState = ResourceCast(NewStateRHI);
 
 	ID3D11SamplerState* StateResource = NewState->Resource;
 	StateCache.SetSamplerState<SF_Geometry>(StateResource, SamplerIndex);
@@ -500,8 +501,8 @@ void FD3D11DynamicRHI::RHISetShaderSampler(FPixelShaderRHIParamRef PixelShaderRH
 {
 	VALIDATE_BOUND_SHADER(PixelShaderRHI);
 
-	DYNAMIC_CAST_D3D11RESOURCE(PixelShader,PixelShader);
-	DYNAMIC_CAST_D3D11RESOURCE(SamplerState,NewState);
+	FD3D11PixelShader* PixelShader = ResourceCast(PixelShaderRHI);
+	FD3D11SamplerState* NewState = ResourceCast(NewStateRHI);
 
 	ID3D11SamplerState* StateResource = NewState->Resource;
 	StateCache.SetSamplerState<SF_Pixel>(StateResource, SamplerIndex);
@@ -510,8 +511,8 @@ void FD3D11DynamicRHI::RHISetShaderSampler(FPixelShaderRHIParamRef PixelShaderRH
 void FD3D11DynamicRHI::RHISetShaderSampler(FComputeShaderRHIParamRef ComputeShaderRHI,uint32 SamplerIndex,FSamplerStateRHIParamRef NewStateRHI)
 {
 	//VALIDATE_BOUND_SHADER(ComputeShaderRHI);
-	DYNAMIC_CAST_D3D11RESOURCE(ComputeShader,ComputeShader);
-	DYNAMIC_CAST_D3D11RESOURCE(SamplerState,NewState);
+	FD3D11ComputeShader* ComputeShader = ResourceCast(ComputeShaderRHI);
+	FD3D11SamplerState* NewState = ResourceCast(NewStateRHI);
 
 	ID3D11SamplerState* StateResource = NewState->Resource;
 	StateCache.SetSamplerState<SF_Compute>(StateResource, SamplerIndex);
@@ -520,7 +521,7 @@ void FD3D11DynamicRHI::RHISetShaderSampler(FComputeShaderRHIParamRef ComputeShad
 void FD3D11DynamicRHI::RHISetShaderUniformBuffer(FVertexShaderRHIParamRef VertexShader,uint32 BufferIndex,FUniformBufferRHIParamRef BufferRHI)
 {
 	VALIDATE_BOUND_SHADER(VertexShader);
-	DYNAMIC_CAST_D3D11RESOURCE(UniformBuffer,Buffer);
+	FD3D11UniformBuffer* Buffer = ResourceCast(BufferRHI);
 	{
 		ID3D11Buffer* ConstantBuffer = Buffer ? Buffer->Resource : NULL;
 		StateCache.SetConstantBuffer<SF_Vertex>(ConstantBuffer, BufferIndex);
@@ -533,7 +534,7 @@ void FD3D11DynamicRHI::RHISetShaderUniformBuffer(FVertexShaderRHIParamRef Vertex
 void FD3D11DynamicRHI::RHISetShaderUniformBuffer(FHullShaderRHIParamRef HullShader,uint32 BufferIndex,FUniformBufferRHIParamRef BufferRHI)
 {
 	VALIDATE_BOUND_SHADER(HullShader);
-	DYNAMIC_CAST_D3D11RESOURCE(UniformBuffer,Buffer);
+	FD3D11UniformBuffer* Buffer = ResourceCast(BufferRHI);
 	{
 		ID3D11Buffer* ConstantBuffer = Buffer ? Buffer->Resource : NULL;
 		StateCache.SetConstantBuffer<SF_Hull>(ConstantBuffer, BufferIndex);
@@ -546,7 +547,7 @@ void FD3D11DynamicRHI::RHISetShaderUniformBuffer(FHullShaderRHIParamRef HullShad
 void FD3D11DynamicRHI::RHISetShaderUniformBuffer(FDomainShaderRHIParamRef DomainShader,uint32 BufferIndex,FUniformBufferRHIParamRef BufferRHI)
 {
 	VALIDATE_BOUND_SHADER(DomainShader);
-	DYNAMIC_CAST_D3D11RESOURCE(UniformBuffer,Buffer);
+	FD3D11UniformBuffer* Buffer = ResourceCast(BufferRHI);
 	{
 		ID3D11Buffer* ConstantBuffer = Buffer ? Buffer->Resource : NULL;
 		StateCache.SetConstantBuffer<SF_Domain>(ConstantBuffer, BufferIndex);
@@ -559,7 +560,7 @@ void FD3D11DynamicRHI::RHISetShaderUniformBuffer(FDomainShaderRHIParamRef Domain
 void FD3D11DynamicRHI::RHISetShaderUniformBuffer(FGeometryShaderRHIParamRef GeometryShader,uint32 BufferIndex,FUniformBufferRHIParamRef BufferRHI)
 {
 	VALIDATE_BOUND_SHADER(GeometryShader);
-	DYNAMIC_CAST_D3D11RESOURCE(UniformBuffer,Buffer);
+	FD3D11UniformBuffer* Buffer = ResourceCast(BufferRHI);
 	{
 		ID3D11Buffer* ConstantBuffer = Buffer ? Buffer->Resource : NULL;
 		StateCache.SetConstantBuffer<SF_Geometry>(ConstantBuffer, BufferIndex);
@@ -572,7 +573,7 @@ void FD3D11DynamicRHI::RHISetShaderUniformBuffer(FGeometryShaderRHIParamRef Geom
 void FD3D11DynamicRHI::RHISetShaderUniformBuffer(FPixelShaderRHIParamRef PixelShader,uint32 BufferIndex,FUniformBufferRHIParamRef BufferRHI)
 {
 	VALIDATE_BOUND_SHADER(PixelShader);
-	DYNAMIC_CAST_D3D11RESOURCE(UniformBuffer,Buffer);
+	FD3D11UniformBuffer* Buffer = ResourceCast(BufferRHI);
 
 	{
 		ID3D11Buffer* ConstantBuffer = Buffer ? Buffer->Resource : NULL;
@@ -586,7 +587,7 @@ void FD3D11DynamicRHI::RHISetShaderUniformBuffer(FPixelShaderRHIParamRef PixelSh
 void FD3D11DynamicRHI::RHISetShaderUniformBuffer(FComputeShaderRHIParamRef ComputeShader,uint32 BufferIndex,FUniformBufferRHIParamRef BufferRHI)
 {
 	//VALIDATE_BOUND_SHADER(ComputeShader);
-	DYNAMIC_CAST_D3D11RESOURCE(UniformBuffer,Buffer);
+	FD3D11UniformBuffer* Buffer = ResourceCast(BufferRHI);
 	{
 		ID3D11Buffer* ConstantBuffer = Buffer ? Buffer->Resource : NULL;
 		StateCache.SetConstantBuffer<SF_Compute>(ConstantBuffer, BufferIndex);
@@ -659,7 +660,7 @@ void FD3D11DynamicRHI::ValidateExclusiveDepthStencilAccess(FExclusiveDepthStenci
 
 void FD3D11DynamicRHI::RHISetDepthStencilState(FDepthStencilStateRHIParamRef NewStateRHI,uint32 StencilRef)
 {
-	DYNAMIC_CAST_D3D11RESOURCE(DepthStencilState,NewState);
+	FD3D11DepthStencilState* NewState = ResourceCast(NewStateRHI);
 
 	ValidateExclusiveDepthStencilAccess(NewState->AccessType);
 
@@ -668,7 +669,7 @@ void FD3D11DynamicRHI::RHISetDepthStencilState(FDepthStencilStateRHIParamRef New
 
 void FD3D11DynamicRHI::RHISetBlendState(FBlendStateRHIParamRef NewStateRHI,const FLinearColor& BlendFactor)
 {
-	DYNAMIC_CAST_D3D11RESOURCE(BlendState,NewState);
+	FD3D11BlendState* NewState = ResourceCast(NewStateRHI);
 	StateCache.SetBlendState(NewState->Resource, (const float*)&BlendFactor, 0xffffffff);
 }
 
@@ -1231,7 +1232,7 @@ void FD3D11DynamicRHI::RHIDrawPrimitive(uint32 PrimitiveType,uint32 BaseVertexIn
 
 void FD3D11DynamicRHI::RHIDrawPrimitiveIndirect(uint32 PrimitiveType,FVertexBufferRHIParamRef ArgumentBufferRHI,uint32 ArgumentOffset)
 {
-	DYNAMIC_CAST_D3D11RESOURCE(VertexBuffer,ArgumentBuffer);
+	FD3D11VertexBuffer* ArgumentBuffer = ResourceCast(ArgumentBufferRHI);
 
 	RHI_DRAW_CALL_INC();
 
@@ -1246,8 +1247,8 @@ void FD3D11DynamicRHI::RHIDrawPrimitiveIndirect(uint32 PrimitiveType,FVertexBuff
 
 void FD3D11DynamicRHI::RHIDrawIndexedIndirect(FIndexBufferRHIParamRef IndexBufferRHI, uint32 PrimitiveType, FStructuredBufferRHIParamRef ArgumentsBufferRHI, int32 DrawArgumentsIndex, uint32 NumInstances)
 {
-	DYNAMIC_CAST_D3D11RESOURCE(IndexBuffer,IndexBuffer);
-	DYNAMIC_CAST_D3D11RESOURCE(StructuredBuffer,ArgumentsBuffer);
+	FD3D11IndexBuffer* IndexBuffer = ResourceCast(IndexBufferRHI);
+	FD3D11StructuredBuffer* ArgumentsBuffer = ResourceCast(ArgumentsBufferRHI);
 
 	RHI_DRAW_CALL_INC();
 
@@ -1275,7 +1276,7 @@ void FD3D11DynamicRHI::RHIDrawIndexedIndirect(FIndexBufferRHIParamRef IndexBuffe
 
 void FD3D11DynamicRHI::RHIDrawIndexedPrimitive(FIndexBufferRHIParamRef IndexBufferRHI,uint32 PrimitiveType,int32 BaseVertexIndex,uint32 FirstInstance,uint32 NumVertices,uint32 StartIndex,uint32 NumPrimitives,uint32 NumInstances)
 {
-	DYNAMIC_CAST_D3D11RESOURCE(IndexBuffer,IndexBuffer);
+	FD3D11IndexBuffer* IndexBuffer = ResourceCast(IndexBufferRHI);
 
 	// called should make sure the input is valid, this avoid hidden bugs
 	ensure(NumPrimitives > 0);
@@ -1313,8 +1314,8 @@ void FD3D11DynamicRHI::RHIDrawIndexedPrimitive(FIndexBufferRHIParamRef IndexBuff
 
 void FD3D11DynamicRHI::RHIDrawIndexedPrimitiveIndirect(uint32 PrimitiveType,FIndexBufferRHIParamRef IndexBufferRHI,FVertexBufferRHIParamRef ArgumentBufferRHI,uint32 ArgumentOffset)
 {
-	DYNAMIC_CAST_D3D11RESOURCE(IndexBuffer,IndexBuffer);
-	DYNAMIC_CAST_D3D11RESOURCE(VertexBuffer,ArgumentBuffer);
+	FD3D11IndexBuffer* IndexBuffer = ResourceCast(IndexBufferRHI);
+	FD3D11VertexBuffer* ArgumentBuffer = ResourceCast(ArgumentBufferRHI);
 
 	RHI_DRAW_CALL_INC();
 
@@ -1694,9 +1695,9 @@ void FD3D11DynamicRHI::RHIClearMRT(bool bClearColor,int32 NumClearColors,const F
 			ValidateExclusiveDepthStencilAccess(RequestedAccess);
 		}
 
-		DYNAMIC_CAST_D3D11RESOURCE(BlendState,BlendState);
-		DYNAMIC_CAST_D3D11RESOURCE(RasterizerState,RasterizerState);
-		DYNAMIC_CAST_D3D11RESOURCE(DepthStencilState,DepthStencilState);
+		FD3D11BlendState* BlendState = ResourceCast(BlendStateRHI);
+		FD3D11RasterizerState* RasterizerState = ResourceCast(RasterizerStateRHI);
+		FD3D11DepthStencilState* DepthStencilState = ResourceCast(DepthStencilStateRHI);
 
 		// Store the current device state
 		FDeviceStateHelper OriginalResourceState(Direct3DDeviceIMContext);
diff --git a/Engine/Source/Runtime/Windows/D3D11RHI/Private/D3D11IndexBuffer.cpp b/Engine/Source/Runtime/Windows/D3D11RHI/Private/D3D11IndexBuffer.cpp
--- a/Engine/Source/Runtime/Windows/D3D11RHI/Private/D3D11IndexBuffer.cpp
+++ b/Engine/Source/Runtime/Windows/D3D11RHI/Private/D3D11IndexBuffer.cpp
@@ -63,7 +63,7 @@ FIndexBufferRHIRef FD3D11DynamicRHI::RHICreateIndexBuffer(uint32 Stride,uint32 S
 
 void* FD3D11DynamicRHI::RHILockIndexBuffer(FIndexBufferRHIParamRef IndexBufferRHI,uint32 Offset,uint32 Size,EResourceLockMode LockMode)
 {
-	DYNAMIC_CAST_D3D11RESOURCE(IndexBuffer,IndexBuffer);
+	FD3D11IndexBuffer* IndexBuffer = ResourceCast(IndexBufferRHI);
 	
 	// If this resource is bound to the device, unbind it
 	ConditionalClearShaderResource(IndexBuffer);
@@ -128,7 +128,7 @@ void* FD3D11DynamicRHI::RHILockIndexBuffer(FIndexBufferRHIParamRef IndexBufferRH
 
 void FD3D11DynamicRHI::RHIUnlockIndexBuffer(FIndexBufferRHIParamRef IndexBufferRHI)
 {
-	DYNAMIC_CAST_D3D11RESOURCE(IndexBuffer,IndexBuffer);
+	FD3D11IndexBuffer* IndexBuffer = ResourceCast(IndexBufferRHI);
 
 	// Determine whether the index buffer is dynamic or not.
 	D3D11_BUFFER_DESC Desc;
diff --git a/Engine/Source/Runtime/Windows/D3D11RHI/Private/D3D11Query.cpp b/Engine/Source/Runtime/Windows/D3D11RHI/Private/D3D11Query.cpp
--- a/Engine/Source/Runtime/Windows/D3D11RHI/Private/D3D11Query.cpp
+++ b/Engine/Source/Runtime/Windows/D3D11RHI/Private/D3D11Query.cpp
@@ -40,7 +40,7 @@ FRenderQueryRHIRef FD3D11DynamicRHI::RHICreateRenderQuery(ERenderQueryType Query
 bool FD3D11DynamicRHI::RHIGetRenderQueryResult(FRenderQueryRHIParamRef QueryRHI,uint64& OutResult,bool bWait)
 {
 	check(IsInRenderingThread());
-	DYNAMIC_CAST_D3D11RESOURCE(RenderQuery,Query);
+	FD3D11RenderQuery* Query = ResourceCast(QueryRHI);
 
 	bool bSuccess = true;
 	if(!Query->bResultIsCached)
@@ -69,7 +69,7 @@ bool FD3D11DynamicRHI::RHIGetRenderQueryResult(FRenderQueryRHIParamRef QueryRHI,
 // Occlusion/Timer queries.
 void FD3D11DynamicRHI::RHIBeginRenderQuery(FRenderQueryRHIParamRef QueryRHI)
 {
-	DYNAMIC_CAST_D3D11RESOURCE(RenderQuery,Query);
+	FD3D11RenderQuery* Query = ResourceCast(QueryRHI);
 
 	if(Query->QueryType == RQT_Occlusion)
 	{
@@ -85,7 +85,7 @@ void FD3D11DynamicRHI::RHIBeginRenderQuery(FRenderQueryRHIParamRef QueryRHI)
 
 void FD3D11DynamicRHI::RHIEndRenderQuery(FRenderQueryRHIParamRef QueryRHI)
 {
-	DYNAMIC_CAST_D3D11RESOURCE(RenderQuery,Query);
+	FD3D11RenderQuery* Query = ResourceCast(QueryRHI);
 	Query->bResultIsCached = false; // for occlusion queries, this is redundant with the one in begin
 	Direct3DDeviceIMContext->End(Query->Resource);
 
diff --git a/Engine/Source/Runtime/Windows/D3D11RHI/Private/D3D11RHIPrivate.h b/Engine/Source/Runtime/Windows/D3D11RHI/Private/D3D11RHIPrivate.h
--- a/Engine/Source/Runtime/Windows/D3D11RHI/Private/D3D11RHIPrivate.h
+++ b/Engine/Source/Runtime/Windows/D3D11RHI/Private/D3D11RHIPrivate.h
@@ -20,10 +20,6 @@ DECLARE_LOG_CATEGORY_EXTERN(LogD3D11RHI, Log, All);
 #include "D3D11RHIBasePrivate.h"
 #include "StaticArray.h"
 
-/** This is a macro that casts a dynamically bound RHI reference to the appropriate D3D type. */
-#define DYNAMIC_CAST_D3D11RESOURCE(Type,Name) \
-	FD3D11##Type* Name = (FD3D11##Type*)Name##RHI;
-
 // D3D RHI public headers.
 #include "D3D11Util.h"
 #include "D3D11State.h"
@@ -291,6 +287,12 @@ public:
 	virtual void Init() override;
 	virtual void Shutdown() override;
 
+	template<typename TRHIType>
+	static FORCEINLINE typename TD3D11ResourceTraits<TRHIType>::TConcreteType* ResourceCast(TRHIType* Resource)
+	{
+		return static_cast<typename TD3D11ResourceTraits<TRHIType>::TConcreteType*>(Resource);
+	}
+
 	/**
 	 * Reads a D3D query's data into the provided buffer.
 	 * @param Query - The D3D query to read data from.
@@ -584,6 +586,7 @@ protected:
 	void ReadSurfaceDataMSAARaw(FRHICommandList_RecursiveHazardous& RHICmdList, FTextureRHIParamRef TextureRHI, FIntRect Rect, TArray<uint8>& OutData, FReadSurfaceDataFlags InFlags);
 
 	friend struct FD3DGPUProfiler;
+
 };
 
 struct FD3D11Adapter
diff --git a/Engine/Source/Runtime/Windows/D3D11RHI/Private/D3D11Shaders.cpp b/Engine/Source/Runtime/Windows/D3D11RHI/Private/D3D11Shaders.cpp
--- a/Engine/Source/Runtime/Windows/D3D11RHI/Private/D3D11Shaders.cpp
+++ b/Engine/Source/Runtime/Windows/D3D11RHI/Private/D3D11Shaders.cpp
@@ -203,12 +203,12 @@ FD3D11BoundShaderState::FD3D11BoundShaderState(
 {
 	INC_DWORD_STAT(STAT_D3D11NumBoundShaderState);
 
-	DYNAMIC_CAST_D3D11RESOURCE(VertexDeclaration,InVertexDeclaration);
-	DYNAMIC_CAST_D3D11RESOURCE(VertexShader,InVertexShader);
-	DYNAMIC_CAST_D3D11RESOURCE(PixelShader,InPixelShader);
-	DYNAMIC_CAST_D3D11RESOURCE(HullShader,InHullShader);
-	DYNAMIC_CAST_D3D11RESOURCE(DomainShader,InDomainShader);
-	DYNAMIC_CAST_D3D11RESOURCE(GeometryShader,InGeometryShader);
+	FD3D11VertexDeclaration* InVertexDeclaration = FD3D11DynamicRHI::ResourceCast(InVertexDeclarationRHI);
+	FD3D11VertexShader* InVertexShader = FD3D11DynamicRHI::ResourceCast(InVertexShaderRHI);
+	FD3D11PixelShader* InPixelShader = FD3D11DynamicRHI::ResourceCast(InPixelShaderRHI);
+	FD3D11HullShader* InHullShader = FD3D11DynamicRHI::ResourceCast(InHullShaderRHI);
+	FD3D11DomainShader* InDomainShader = FD3D11DynamicRHI::ResourceCast(InDomainShaderRHI);
+	FD3D11GeometryShader* InGeometryShader = FD3D11DynamicRHI::ResourceCast(InGeometryShaderRHI);
 
 	// Create an input layout for this combination of vertex declaration and vertex shader.
 	D3D11_INPUT_ELEMENT_DESC NullInputElement;
diff --git a/Engine/Source/Runtime/Windows/D3D11RHI/Private/D3D11StructuredBuffer.cpp b/Engine/Source/Runtime/Windows/D3D11RHI/Private/D3D11StructuredBuffer.cpp
--- a/Engine/Source/Runtime/Windows/D3D11RHI/Private/D3D11StructuredBuffer.cpp
+++ b/Engine/Source/Runtime/Windows/D3D11RHI/Private/D3D11StructuredBuffer.cpp
@@ -90,7 +90,7 @@ FStructuredBufferRHIRef FD3D11DynamicRHI::RHICreateStructuredBuffer(uint32 Strid
 
 void* FD3D11DynamicRHI::RHILockStructuredBuffer(FStructuredBufferRHIParamRef StructuredBufferRHI,uint32 Offset,uint32 Size,EResourceLockMode LockMode)
 {
-	DYNAMIC_CAST_D3D11RESOURCE(StructuredBuffer,StructuredBuffer);
+	FD3D11StructuredBuffer* StructuredBuffer = ResourceCast(StructuredBufferRHI);
 	
 	// If this resource is bound to the device, unbind it
 	ConditionalClearShaderResource(StructuredBuffer);
@@ -155,7 +155,7 @@ void* FD3D11DynamicRHI::RHILockStructuredBuffer(FStructuredBufferRHIParamRef Str
 
 void FD3D11DynamicRHI::RHIUnlockStructuredBuffer(FStructuredBufferRHIParamRef StructuredBufferRHI)
 {
-	DYNAMIC_CAST_D3D11RESOURCE(StructuredBuffer,StructuredBuffer);
+	FD3D11StructuredBuffer* StructuredBuffer = ResourceCast(StructuredBufferRHI);
 
 	// Determine whether the Structured buffer is dynamic or not.
 	D3D11_BUFFER_DESC Desc;
diff --git a/Engine/Source/Runtime/Windows/D3D11RHI/Private/D3D11Texture.cpp b/Engine/Source/Runtime/Windows/D3D11RHI/Private/D3D11Texture.cpp
--- a/Engine/Source/Runtime/Windows/D3D11RHI/Private/D3D11Texture.cpp
+++ b/Engine/Source/Runtime/Windows/D3D11RHI/Private/D3D11Texture.cpp
@@ -1089,8 +1089,8 @@ FTexture2DRHIRef FD3D11DynamicRHI::RHIAsyncCreateTexture2D(uint32 SizeX,uint32 S
 
 void FD3D11DynamicRHI::RHICopySharedMips(FTexture2DRHIParamRef DestTexture2DRHI,FTexture2DRHIParamRef SrcTexture2DRHI)
 {
-	DYNAMIC_CAST_D3D11RESOURCE(Texture2D,DestTexture2D);
-	DYNAMIC_CAST_D3D11RESOURCE(Texture2D,SrcTexture2D);
+	FD3D11Texture2D* DestTexture2D = ResourceCast(DestTexture2DRHI);
+	FD3D11Texture2D* SrcTexture2D = ResourceCast(SrcTexture2DRHI);
 
 	// Use the GPU to asynchronously copy the old mip-maps into the new texture.
 	const uint32 NumSharedMips = FMath::Min(DestTexture2D->GetNumMips(),SrcTexture2D->GetNumMips());
@@ -1134,7 +1134,7 @@ void FD3D11DynamicRHI::RHIGetResourceInfo(FTextureRHIParamRef Ref, FRHIResourceI
 
 FShaderResourceViewRHIRef FD3D11DynamicRHI::RHICreateShaderResourceView(FTexture2DRHIParamRef Texture2DRHI, uint8 MipLevel)
 {
-	DYNAMIC_CAST_D3D11RESOURCE(Texture2D,Texture2D);
+	FD3D11Texture2D* Texture2D = ResourceCast(Texture2DRHI);
 
 	D3D11_TEXTURE2D_DESC TextureDesc;
 	Texture2D->GetResource()->GetDesc(&TextureDesc);
@@ -1154,7 +1154,7 @@ FShaderResourceViewRHIRef FD3D11DynamicRHI::RHICreateShaderResourceView(FTexture
 
 FShaderResourceViewRHIRef FD3D11DynamicRHI::RHICreateShaderResourceView(FTexture2DRHIParamRef Texture2DRHI, uint8 MipLevel, uint8 NumMipLevels, uint8 Format)
 {
-	DYNAMIC_CAST_D3D11RESOURCE(Texture2D,Texture2D);
+	FD3D11Texture2D* Texture2D = ResourceCast(Texture2DRHI);
 
 	D3D11_TEXTURE2D_DESC TextureDesc;
 	Texture2D->GetResource()->GetDesc(&TextureDesc);
@@ -1241,7 +1241,7 @@ uint32 FD3D11DynamicRHI::RHIComputeMemorySize(FTextureRHIParamRef TextureRHI)
  */
 FTexture2DRHIRef FD3D11DynamicRHI::RHIAsyncReallocateTexture2D(FTexture2DRHIParamRef Texture2DRHI, int32 NewMipCount, int32 NewSizeX, int32 NewSizeY, FThreadSafeCounter* RequestStatus)
 {
-	DYNAMIC_CAST_D3D11RESOURCE(Texture2D,Texture2D);
+	FD3D11Texture2D* Texture2D = ResourceCast(Texture2DRHI);
 
 	// Allocate a new texture.
 	FRHIResourceCreateInfo CreateInfo;
@@ -1410,7 +1410,7 @@ void TD3D11Texture2D<RHIResourceType>::Unlock(uint32 MipIndex,uint32 ArrayIndex)
 void* FD3D11DynamicRHI::RHILockTexture2D(FTexture2DRHIParamRef TextureRHI,uint32 MipIndex,EResourceLockMode LockMode,uint32& DestStride,bool bLockWithinMiptail)
 {
 	check(TextureRHI);
-	DYNAMIC_CAST_D3D11RESOURCE(Texture2D,Texture);
+	FD3D11Texture2D* Texture = ResourceCast(TextureRHI);
 	ConditionalClearShaderResource(Texture);
 	return Texture->Lock(MipIndex,0,LockMode,DestStride);
 }
@@ -1418,26 +1418,26 @@ void* FD3D11DynamicRHI::RHILockTexture2D(FTexture2DRHIParamRef TextureRHI,uint32
 void FD3D11DynamicRHI::RHIUnlockTexture2D(FTexture2DRHIParamRef TextureRHI,uint32 MipIndex,bool bLockWithinMiptail)
 {
 	check(TextureRHI);
-	DYNAMIC_CAST_D3D11RESOURCE(Texture2D,Texture);
+	FD3D11Texture2D* Texture = ResourceCast(TextureRHI);
 	Texture->Unlock(MipIndex,0);
 }
 
 void* FD3D11DynamicRHI::RHILockTexture2DArray(FTexture2DArrayRHIParamRef TextureRHI,uint32 TextureIndex,uint32 MipIndex,EResourceLockMode LockMode,uint32& DestStride,bool bLockWithinMiptail)
 {
-	DYNAMIC_CAST_D3D11RESOURCE(Texture2DArray,Texture);
+	FD3D11Texture2DArray* Texture = ResourceCast(TextureRHI);
 	ConditionalClearShaderResource(Texture);
 	return Texture->Lock(MipIndex,TextureIndex,LockMode,DestStride);
 }
 
 void FD3D11DynamicRHI::RHIUnlockTexture2DArray(FTexture2DArrayRHIParamRef TextureRHI,uint32 TextureIndex,uint32 MipIndex,bool bLockWithinMiptail)
 {
-	DYNAMIC_CAST_D3D11RESOURCE(Texture2DArray,Texture);
+	FD3D11Texture2DArray* Texture = ResourceCast(TextureRHI);
 	Texture->Unlock(MipIndex,TextureIndex);
 }
 
 void FD3D11DynamicRHI::RHIUpdateTexture2D(FTexture2DRHIParamRef TextureRHI,uint32 MipIndex,const FUpdateTextureRegion2D& UpdateRegion,uint32 SourcePitch,const uint8* SourceData)
 {
-    DYNAMIC_CAST_D3D11RESOURCE(Texture2D,Texture);
+    FD3D11Texture2D* Texture = ResourceCast(TextureRHI);
 
     D3D11_BOX DestBox =
 	{
@@ -1453,7 +1453,7 @@ void FD3D11DynamicRHI::RHIUpdateTexture2D(FTexture2DRHIParamRef TextureRHI,uint3
 
 void FD3D11DynamicRHI::RHIUpdateTexture3D(FTexture3DRHIParamRef TextureRHI,uint32 MipIndex,const FUpdateTextureRegion3D& UpdateRegion,uint32 SourceRowPitch,uint32 SourceDepthPitch,const uint8* SourceData)
 {
-	DYNAMIC_CAST_D3D11RESOURCE(Texture3D,Texture);
+	FD3D11Texture3D* Texture = ResourceCast(TextureRHI);
 
 	D3D11_BOX DestBox =
 	{
@@ -1482,14 +1482,14 @@ FTextureCubeRHIRef FD3D11DynamicRHI::RHICreateTextureCubeArray(uint32 Size, uint
 
 void* FD3D11DynamicRHI::RHILockTextureCubeFace(FTextureCubeRHIParamRef TextureCubeRHI,uint32 FaceIndex,uint32 ArrayIndex,uint32 MipIndex,EResourceLockMode LockMode,uint32& DestStride,bool bLockWithinMiptail)
 {
-	DYNAMIC_CAST_D3D11RESOURCE(TextureCube,TextureCube);
+	FD3D11TextureCube* TextureCube = ResourceCast(TextureCubeRHI);
 	ConditionalClearShaderResource(TextureCube);
 	uint32 D3DFace = GetD3D11CubeFace((ECubeFace)FaceIndex);
 	return TextureCube->Lock(MipIndex,D3DFace + ArrayIndex * 6,LockMode,DestStride);
 }
 void FD3D11DynamicRHI::RHIUnlockTextureCubeFace(FTextureCubeRHIParamRef TextureCubeRHI,uint32 FaceIndex,uint32 ArrayIndex,uint32 MipIndex,bool bLockWithinMiptail)
 {
-	DYNAMIC_CAST_D3D11RESOURCE(TextureCube,TextureCube);
+	FD3D11TextureCube* TextureCube = ResourceCast(TextureCubeRHI);
 	uint32 D3DFace = GetD3D11CubeFace((ECubeFace)FaceIndex);
 	TextureCube->Unlock(MipIndex,D3DFace + ArrayIndex * 6);
 }
diff --git a/Engine/Source/Runtime/Windows/D3D11RHI/Private/D3D11UAV.cpp b/Engine/Source/Runtime/Windows/D3D11RHI/Private/D3D11UAV.cpp
--- a/Engine/Source/Runtime/Windows/D3D11RHI/Private/D3D11UAV.cpp
+++ b/Engine/Source/Runtime/Windows/D3D11RHI/Private/D3D11UAV.cpp
@@ -5,7 +5,7 @@
 
 FUnorderedAccessViewRHIRef FD3D11DynamicRHI::RHICreateUnorderedAccessView(FStructuredBufferRHIParamRef StructuredBufferRHI, bool bUseUAVCounter, bool bAppendBuffer)
 {
-	DYNAMIC_CAST_D3D11RESOURCE(StructuredBuffer,StructuredBuffer);
+	FD3D11StructuredBuffer* StructuredBuffer = ResourceCast(StructuredBufferRHI);
 
 	D3D11_BUFFER_DESC BufferDesc;
 	StructuredBuffer->Resource->GetDesc(&BufferDesc);
@@ -83,7 +83,7 @@ FUnorderedAccessViewRHIRef FD3D11DynamicRHI::RHICreateUnorderedAccessView(FTextu
 
 FUnorderedAccessViewRHIRef FD3D11DynamicRHI::RHICreateUnorderedAccessView(FVertexBufferRHIParamRef VertexBufferRHI, uint8 Format)
 {
-	DYNAMIC_CAST_D3D11RESOURCE(VertexBuffer,VertexBuffer);
+	FD3D11VertexBuffer* VertexBuffer = ResourceCast(VertexBufferRHI);
 
 	D3D11_BUFFER_DESC BufferDesc;
 	VertexBuffer->Resource->GetDesc(&BufferDesc);
@@ -112,7 +112,7 @@ FUnorderedAccessViewRHIRef FD3D11DynamicRHI::RHICreateUnorderedAccessView(FVerte
 
 FShaderResourceViewRHIRef FD3D11DynamicRHI::RHICreateShaderResourceView(FStructuredBufferRHIParamRef StructuredBufferRHI)
 {
-	DYNAMIC_CAST_D3D11RESOURCE(StructuredBuffer,StructuredBuffer);
+	FD3D11StructuredBuffer* StructuredBuffer = ResourceCast(StructuredBufferRHI);
 
 	D3D11_BUFFER_DESC BufferDesc;
 	StructuredBuffer->Resource->GetDesc(&BufferDesc);
@@ -146,7 +146,7 @@ FShaderResourceViewRHIRef FD3D11DynamicRHI::RHICreateShaderResourceView(FStructu
 
 FShaderResourceViewRHIRef FD3D11DynamicRHI::RHICreateShaderResourceView(FVertexBufferRHIParamRef VertexBufferRHI, uint32 Stride, uint8 Format)
 {
-	DYNAMIC_CAST_D3D11RESOURCE(VertexBuffer,VertexBuffer);
+	FD3D11VertexBuffer* VertexBuffer = ResourceCast(VertexBufferRHI);
 	check(VertexBuffer);
 	check(VertexBuffer->Resource);
 
@@ -183,7 +183,7 @@ FShaderResourceViewRHIRef FD3D11DynamicRHI::RHICreateShaderResourceView(FVertexB
 
 void FD3D11DynamicRHI::RHIClearUAV(FUnorderedAccessViewRHIParamRef UnorderedAccessViewRHI, const uint32* Values)
 {
-	DYNAMIC_CAST_D3D11RESOURCE(UnorderedAccessView,UnorderedAccessView);
+	FD3D11UnorderedAccessView* UnorderedAccessView = ResourceCast(UnorderedAccessViewRHI);
 	
 	Direct3DDeviceIMContext->ClearUnorderedAccessViewUint(UnorderedAccessView->View, Values);
 	
diff --git a/Engine/Source/Runtime/Windows/D3D11RHI/Private/D3D11VertexBuffer.cpp b/Engine/Source/Runtime/Windows/D3D11RHI/Private/D3D11VertexBuffer.cpp
--- a/Engine/Source/Runtime/Windows/D3D11RHI/Private/D3D11VertexBuffer.cpp
+++ b/Engine/Source/Runtime/Windows/D3D11RHI/Private/D3D11VertexBuffer.cpp
@@ -89,7 +89,7 @@ void* FD3D11DynamicRHI::RHILockVertexBuffer(FVertexBufferRHIParamRef VertexBuffe
 {
 	check(Size > 0);
 
-	DYNAMIC_CAST_D3D11RESOURCE(VertexBuffer,VertexBuffer);
+	FD3D11VertexBuffer* VertexBuffer = ResourceCast(VertexBufferRHI);
 	
 	// If this resource is bound to the device, unbind it
 	ConditionalClearShaderResource(VertexBuffer);
@@ -159,7 +159,7 @@ void* FD3D11DynamicRHI::RHILockVertexBuffer(FVertexBufferRHIParamRef VertexBuffe
 
 void FD3D11DynamicRHI::RHIUnlockVertexBuffer(FVertexBufferRHIParamRef VertexBufferRHI)
 {
-	DYNAMIC_CAST_D3D11RESOURCE(VertexBuffer,VertexBuffer);
+	FD3D11VertexBuffer* VertexBuffer = ResourceCast(VertexBufferRHI);
 
 	// Determine whether the vertex buffer is dynamic or not.
 	D3D11_BUFFER_DESC Desc;
@@ -202,8 +202,8 @@ void FD3D11DynamicRHI::RHIUnlockVertexBuffer(FVertexBufferRHIParamRef VertexBuff
 
 void FD3D11DynamicRHI::RHICopyVertexBuffer(FVertexBufferRHIParamRef SourceBufferRHI,FVertexBufferRHIParamRef DestBufferRHI)
 {
-	DYNAMIC_CAST_D3D11RESOURCE(VertexBuffer,SourceBuffer);
-	DYNAMIC_CAST_D3D11RESOURCE(VertexBuffer,DestBuffer);
+	FD3D11VertexBuffer* SourceBuffer = ResourceCast(SourceBufferRHI);
+	FD3D11VertexBuffer* DestBuffer = ResourceCast(DestBufferRHI);
 
 	D3D11_BUFFER_DESC SourceBufferDesc;
 	SourceBuffer->Resource->GetDesc(&SourceBufferDesc);
diff --git a/Engine/Source/Runtime/Windows/D3D11RHI/Private/D3D11Viewport.cpp b/Engine/Source/Runtime/Windows/D3D11RHI/Private/D3D11Viewport.cpp
--- a/Engine/Source/Runtime/Windows/D3D11RHI/Private/D3D11Viewport.cpp
+++ b/Engine/Source/Runtime/Windows/D3D11RHI/Private/D3D11Viewport.cpp
@@ -470,7 +470,7 @@ FViewportRHIRef FD3D11DynamicRHI::RHICreateViewport(void* WindowHandle,uint32 Si
 
 void FD3D11DynamicRHI::RHIResizeViewport(FViewportRHIParamRef ViewportRHI,uint32 SizeX,uint32 SizeY,bool bIsFullscreen)
 {
-	DYNAMIC_CAST_D3D11RESOURCE(Viewport,Viewport);
+	FD3D11Viewport* Viewport = ResourceCast(ViewportRHI);
 
 	check( IsInGameThread() );
 	Viewport->Resize(SizeX,SizeY,bIsFullscreen);
@@ -499,7 +499,7 @@ void FD3D11DynamicRHI::RHITick( float DeltaTime )
 
 void FD3D11DynamicRHI::RHIBeginDrawingViewport(FViewportRHIParamRef ViewportRHI, FTextureRHIParamRef RenderTarget)
 {
-	DYNAMIC_CAST_D3D11RESOURCE(Viewport,Viewport);
+	FD3D11Viewport* Viewport = ResourceCast(ViewportRHI);
 
 	SCOPE_CYCLE_COUNTER(STAT_D3D11PresentTime);
 
@@ -520,7 +520,7 @@ void FD3D11DynamicRHI::RHIBeginDrawingViewport(FViewportRHIParamRef ViewportRHI,
 
 void FD3D11DynamicRHI::RHIEndDrawingViewport(FViewportRHIParamRef ViewportRHI,bool bPresent,bool bLockToVsync)
 {
-	DYNAMIC_CAST_D3D11RESOURCE(Viewport,Viewport);
+	FD3D11Viewport* Viewport = ResourceCast(ViewportRHI);
 
 	SCOPE_CYCLE_COUNTER(STAT_D3D11PresentTime);
 
@@ -609,7 +609,7 @@ void FD3D11DynamicRHI::RHIAdvanceFrameForGetViewportBackBuffer()
 
 FTexture2DRHIRef FD3D11DynamicRHI::RHIGetViewportBackBuffer(FViewportRHIParamRef ViewportRHI)
 {
-	DYNAMIC_CAST_D3D11RESOURCE(Viewport,Viewport);
+	FD3D11Viewport* Viewport = ResourceCast(ViewportRHI);
 
 	return Viewport->GetBackBuffer();
 }
diff --git a/Engine/Source/Runtime/Windows/D3D11RHI/Public/D3D11Resources.h b/Engine/Source/Runtime/Windows/D3D11RHI/Public/D3D11Resources.h
--- a/Engine/Source/Runtime/Windows/D3D11RHI/Public/D3D11Resources.h
+++ b/Engine/Source/Runtime/Windows/D3D11RHI/Public/D3D11Resources.h
@@ -780,3 +780,130 @@ public:
 
 void ReturnPooledTexture2D(int32 MipCount, EPixelFormat PixelFormat, ID3D11Texture2D* InResource);
 void ReleasePooledTextures();
+
+template<class T>
+struct TD3D11ResourceTraits
+{
+};
+template<>
+struct TD3D11ResourceTraits<FRHIVertexDeclaration>
+{
+	typedef FD3D11VertexDeclaration TConcreteType;
+};
+template<>
+struct TD3D11ResourceTraits<FRHIVertexShader>
+{
+	typedef FD3D11VertexShader TConcreteType;
+};
+template<>
+struct TD3D11ResourceTraits<FRHIGeometryShader>
+{
+	typedef FD3D11GeometryShader TConcreteType;
+};
+template<>
+struct TD3D11ResourceTraits<FRHIHullShader>
+{
+	typedef FD3D11HullShader TConcreteType;
+};
+template<>
+struct TD3D11ResourceTraits<FRHIDomainShader>
+{
+	typedef FD3D11DomainShader TConcreteType;
+};
+template<>
+struct TD3D11ResourceTraits<FRHIPixelShader>
+{
+	typedef FD3D11PixelShader TConcreteType;
+};
+template<>
+struct TD3D11ResourceTraits<FRHIComputeShader>
+{
+	typedef FD3D11ComputeShader TConcreteType;
+};
+template<>
+struct TD3D11ResourceTraits<FRHIBoundShaderState>
+{
+	typedef FD3D11BoundShaderState TConcreteType;
+};
+template<>
+struct TD3D11ResourceTraits<FRHITexture3D>
+{
+	typedef FD3D11Texture3D TConcreteType;
+};
+template<>
+struct TD3D11ResourceTraits<FRHITexture>
+{
+	typedef FD3D11Texture TConcreteType;
+};
+template<>
+struct TD3D11ResourceTraits<FRHITexture2D>
+{
+	typedef FD3D11Texture2D TConcreteType;
+};
+template<>
+struct TD3D11ResourceTraits<FRHITexture2DArray>
+{
+	typedef FD3D11Texture2DArray TConcreteType;
+};
+template<>
+struct TD3D11ResourceTraits<FRHITextureCube>
+{
+	typedef FD3D11TextureCube TConcreteType;
+};
+template<>
+struct TD3D11ResourceTraits<FRHIRenderQuery>
+{
+	typedef FD3D11RenderQuery TConcreteType;
+};
+template<>
+struct TD3D11ResourceTraits<FRHIUniformBuffer>
+{
+	typedef FD3D11UniformBuffer TConcreteType;
+};
+template<>
+struct TD3D11ResourceTraits<FRHIIndexBuffer>
+{
+	typedef FD3D11IndexBuffer TConcreteType;
+};
+template<>
+struct TD3D11ResourceTraits<FRHIStructuredBuffer>
+{
+	typedef FD3D11StructuredBuffer TConcreteType;
+};
+template<>
+struct TD3D11ResourceTraits<FRHIVertexBuffer>
+{
+	typedef FD3D11VertexBuffer TConcreteType;
+};
+template<>
+struct TD3D11ResourceTraits<FRHIShaderResourceView>
+{
+	typedef FD3D11ShaderResourceView TConcreteType;
+};
+template<>
+struct TD3D11ResourceTraits<FRHIUnorderedAccessView>
+{
+	typedef FD3D11UnorderedAccessView TConcreteType;
+};
+
+template<>
+struct TD3D11ResourceTraits<FRHISamplerState>
+{
+	typedef FD3D11SamplerState TConcreteType;
+};
+template<>
+struct TD3D11ResourceTraits<FRHIRasterizerState>
+{
+	typedef FD3D11RasterizerState TConcreteType;
+};
+template<>
+struct TD3D11ResourceTraits<FRHIDepthStencilState>
+{
+	typedef FD3D11DepthStencilState TConcreteType;
+};
+template<>
+struct TD3D11ResourceTraits<FRHIBlendState>
+{
+	typedef FD3D11BlendState TConcreteType;
+};
+
diff --git a/Engine/Source/Runtime/Windows/D3D11RHI/Public/D3D11Viewport.h b/Engine/Source/Runtime/Windows/D3D11RHI/Public/D3D11Viewport.h
--- a/Engine/Source/Runtime/Windows/D3D11RHI/Public/D3D11Viewport.h
+++ b/Engine/Source/Runtime/Windows/D3D11RHI/Public/D3D11Viewport.h
@@ -130,3 +130,8 @@ private:
 	DXGI_MODE_DESC SetupDXGI_MODE_DESC() const;
 };
 
+template<>
+struct TD3D11ResourceTraits<FRHIViewport>
+{
+	typedef FD3D11Viewport TConcreteType;
+};
