Added assert to gather more info when GetDisplayModeList fails

diff --git a/Engine/Source/Runtime/Windows/D3D11RHI/Private/Windows/WindowsD3D11Device.cpp b/Engine/Source/Runtime/Windows/D3D11RHI/Private/Windows/WindowsD3D11Device.cpp
--- a/Engine/Source/Runtime/Windows/D3D11RHI/Private/Windows/WindowsD3D11Device.cpp
+++ b/Engine/Source/Runtime/Windows/D3D11RHI/Private/Windows/WindowsD3D11Device.cpp
@@ -743,6 +743,8 @@ bool FD3D11DynamicRHI::RHIGetAvailableResolutions(FScreenResolutionArray& Resolu
 			return false;
 		}
 
+		checkf(NumModes > 0, TEXT("No display modes found for the standard format DXGI_FORMAT_R8G8B8A8_UNORM!"));
+
 		DXGI_MODE_DESC* ModeList = new DXGI_MODE_DESC[ NumModes ];
 		VERIFYD3D11RESULT(Output->GetDisplayModeList(Format, 0, &NumModes, ModeList));
 
