UE4 - Added missing vertex element formats UShort2/4/2N/4N; fixed Apex Cloth using wrong format

diff --git a/Engine/Source/Runtime/Engine/Private/SkeletalRenderGPUSkin.cpp b/Engine/Source/Runtime/Engine/Private/SkeletalRenderGPUSkin.cpp
--- a/Engine/Source/Runtime/Engine/Private/SkeletalRenderGPUSkin.cpp
+++ b/Engine/Source/Runtime/Engine/Private/SkeletalRenderGPUSkin.cpp
@@ -585,7 +585,7 @@ void InitAPEXClothVertexFactoryComponents(typename VertexFactoryType::DataType*
 		VertexBuffers.APEXClothVertexBuffer,STRUCT_OFFSET(FApexClothPhysToRenderVertData,TangentBaryCoordsAndDist),sizeof(FApexClothPhysToRenderVertData),VET_Float4);
 	// indices for reference physics mesh vertices
 	VertexFactoryData->SimulIndicesComponent = FVertexStreamComponent(
-		VertexBuffers.APEXClothVertexBuffer,STRUCT_OFFSET(FApexClothPhysToRenderVertData,SimulMeshVertIndices),sizeof(FApexClothPhysToRenderVertData),VET_Short4);
+		VertexBuffers.APEXClothVertexBuffer,STRUCT_OFFSET(FApexClothPhysToRenderVertData,SimulMeshVertIndices),sizeof(FApexClothPhysToRenderVertData),VET_UShort4);
 }
 
 /** 
diff --git a/Engine/Source/Runtime/IOS/MetalRHI/Private/MetalVertexDeclaration.cpp b/Engine/Source/Runtime/IOS/MetalRHI/Private/MetalVertexDeclaration.cpp
--- a/Engine/Source/Runtime/IOS/MetalRHI/Private/MetalVertexDeclaration.cpp
+++ b/Engine/Source/Runtime/IOS/MetalRHI/Private/MetalVertexDeclaration.cpp
@@ -10,20 +10,24 @@ static MTLVertexFormat TranslateElementTypeToMTLType(EVertexElementType Type)
 {
 	switch (Type)
 	{
-		case VET_Float1:		return MTLVertexFormatFloat;
-		case VET_Float2:		return MTLVertexFormatFloat2;
-		case VET_Float3:		return MTLVertexFormatFloat3;
-		case VET_Float4:		return MTLVertexFormatFloat4;
+		case VET_Float1:			return MTLVertexFormatFloat;
+		case VET_Float2:			return MTLVertexFormatFloat2;
+		case VET_Float3:			return MTLVertexFormatFloat3;
+		case VET_Float4:			return MTLVertexFormatFloat4;
 		case VET_PackedNormal:	return MTLVertexFormatUChar4Normalized;
-		case VET_UByte4:		return MTLVertexFormatUChar4;
+		case VET_UByte4:			return MTLVertexFormatUChar4;
 		case VET_UByte4N:		return MTLVertexFormatUChar4Normalized;
 		case VET_Color:			return MTLVertexFormatUChar4Normalized;
-		case VET_Short2:		return MTLVertexFormatShort2;
-		case VET_Short4:		return MTLVertexFormatShort4;
+		case VET_Short2:			return MTLVertexFormatShort2;
+		case VET_Short4:			return MTLVertexFormatShort4;
 		case VET_Short2N:		return MTLVertexFormatShort2Normalized;
 		case VET_Half2:			return MTLVertexFormatHalf2;
 		case VET_Half4:			return MTLVertexFormatHalf4;
 		case VET_Short4N:		return MTLVertexFormatShort4Normalized;
+		case VET_UShort2:		return MTLVertexFormatUShort2;
+		case VET_UShort4:		return MTLVertexFormatUShort4;
+		case VET_UShort2N:		return MTLVertexFormatUShort2Normalized;
+		case VET_UShort4N:		return MTLVertexFormatUShort4Normalized;
 		default:				UE_LOG(LogMetal, Fatal, TEXT("Unknown vertex element type!")); return MTLVertexFormatFloat;
 	};
 
@@ -33,16 +37,16 @@ static uint32 TranslateElementTypeToSize(EVertexElementType Type)
 {
 	switch (Type)
 	{
-		case VET_Float1:		return 4;
-		case VET_Float2:		return 8;
-		case VET_Float3:		return 12;
-		case VET_Float4:		return 16;
+		case VET_Float1:			return 4;
+		case VET_Float2:			return 8;
+		case VET_Float3:			return 12;
+		case VET_Float4:			return 16;
 		case VET_PackedNormal:	return 4;
-		case VET_UByte4:		return 4;
+		case VET_UByte4:			return 4;
 		case VET_UByte4N:		return 4;
 		case VET_Color:			return 4;
-		case VET_Short2:		return 4;
-		case VET_Short4:		return 8;
+		case VET_Short2:			return 4;
+		case VET_Short4:			return 8;
 		case VET_Short2N:		return 4;
 		case VET_Half2:			return 4;
 		case VET_Half4:			return 8;
diff --git a/Engine/Source/Runtime/OpenGLDrv/Private/OpenGLVertexDeclaration.cpp b/Engine/Source/Runtime/OpenGLDrv/Private/OpenGLVertexDeclaration.cpp
--- a/Engine/Source/Runtime/OpenGLDrv/Private/OpenGLVertexDeclaration.cpp
+++ b/Engine/Source/Runtime/OpenGLDrv/Private/OpenGLVertexDeclaration.cpp
@@ -81,6 +81,10 @@ struct FOpenGLVertexDeclarationKey
 					}
 					break;
 				case VET_Short4N:		SetupGLElement(GLElement, GL_SHORT,			4,			true,	true); break;
+				case VET_UShort2:		SetupGLElement(GLElement, GL_UNSIGNED_SHORT, 2, false, false); break;
+				case VET_UShort4:		SetupGLElement(GLElement, GL_UNSIGNED_SHORT, 4, false, false); break;
+				case VET_UShort2N:		SetupGLElement(GLElement, GL_UNSIGNED_SHORT, 2, true, true); break;
+				case VET_UShort4N:		SetupGLElement(GLElement, GL_UNSIGNED_SHORT, 4, true, true); break;
 				default: UE_LOG(LogRHI, Fatal,TEXT("Unknown RHI vertex element type %u"),(uint8)InElements[ElementIndex].Type);
 			};
 			VertexElements.Add(GLElement);
diff --git a/Engine/Source/Runtime/RHI/Public/RHIDefinitions.h b/Engine/Source/Runtime/RHI/Public/RHIDefinitions.h
--- a/Engine/Source/Runtime/RHI/Public/RHIDefinitions.h
+++ b/Engine/Source/Runtime/RHI/Public/RHIDefinitions.h
@@ -236,6 +236,10 @@ enum EVertexElementType
 	VET_Half2,			// 16 bit float using 1 bit sign, 5 bit exponent, 10 bit mantissa 
 	VET_Half4,
 	VET_Short4N,		// 4 X 16 bit word, normalized 
+	VET_UShort2,
+	VET_UShort4,
+	VET_UShort2N,		// 16 bit word normalized to (value/65535.0,value/65535.0,0,0,1)
+	VET_UShort4N,		// 4 X 16 bit word unsigned, normalized 
 	VET_MAX
 };
 
diff --git a/Engine/Source/Runtime/Windows/D3D11RHI/Private/D3D11VertexDeclaration.cpp b/Engine/Source/Runtime/Windows/D3D11RHI/Private/D3D11VertexDeclaration.cpp
--- a/Engine/Source/Runtime/Windows/D3D11RHI/Private/D3D11VertexDeclaration.cpp
+++ b/Engine/Source/Runtime/Windows/D3D11RHI/Private/D3D11VertexDeclaration.cpp
@@ -27,20 +27,24 @@ struct FD3D11VertexDeclarationKey
 			D3DElement.AlignedByteOffset = Element.Offset;
 			switch(Element.Type)
 			{
-			case VET_Float1:		D3DElement.Format = DXGI_FORMAT_R32_FLOAT; break;
-			case VET_Float2:		D3DElement.Format = DXGI_FORMAT_R32G32_FLOAT; break;
-			case VET_Float3:		D3DElement.Format = DXGI_FORMAT_R32G32B32_FLOAT; break;
-			case VET_Float4:		D3DElement.Format = DXGI_FORMAT_R32G32B32A32_FLOAT; break;
+			case VET_Float1:			D3DElement.Format = DXGI_FORMAT_R32_FLOAT; break;
+			case VET_Float2:			D3DElement.Format = DXGI_FORMAT_R32G32_FLOAT; break;
+			case VET_Float3:			D3DElement.Format = DXGI_FORMAT_R32G32B32_FLOAT; break;
+			case VET_Float4:			D3DElement.Format = DXGI_FORMAT_R32G32B32A32_FLOAT; break;
 			case VET_PackedNormal:	D3DElement.Format = DXGI_FORMAT_R8G8B8A8_UNORM; break; //TODO: uint32 doesn't work because D3D11 squishes it to 0 in the IA-VS conversion
-			case VET_UByte4:		D3DElement.Format = DXGI_FORMAT_R8G8B8A8_UINT; break; //TODO: SINT, blendindices
+			case VET_UByte4:			D3DElement.Format = DXGI_FORMAT_R8G8B8A8_UINT; break; //TODO: SINT, blendindices
 			case VET_UByte4N:		D3DElement.Format = DXGI_FORMAT_R8G8B8A8_UNORM; break;
 			case VET_Color:			D3DElement.Format = DXGI_FORMAT_B8G8R8A8_UNORM; break;
-			case VET_Short2:		D3DElement.Format = DXGI_FORMAT_R16G16_SINT; break;
-			case VET_Short4:		D3DElement.Format = DXGI_FORMAT_R16G16B16A16_SINT; break;
+			case VET_Short2:			D3DElement.Format = DXGI_FORMAT_R16G16_SINT; break;
+			case VET_Short4:			D3DElement.Format = DXGI_FORMAT_R16G16B16A16_SINT; break;
 			case VET_Short2N:		D3DElement.Format = DXGI_FORMAT_R16G16_SNORM; break;
 			case VET_Half2:			D3DElement.Format = DXGI_FORMAT_R16G16_FLOAT; break;
 			case VET_Half4:			D3DElement.Format = DXGI_FORMAT_R16G16B16A16_FLOAT; break;
 			case VET_Short4N:		D3DElement.Format = DXGI_FORMAT_R16G16B16A16_SNORM; break;
+			case VET_UShort2:		D3DElement.Format = DXGI_FORMAT_R16G16_UINT; break;
+			case VET_UShort4:		D3DElement.Format = DXGI_FORMAT_R16G16B16A16_UINT; break;
+			case VET_UShort2N:		D3DElement.Format = DXGI_FORMAT_R16G16_UNORM; break;
+			case VET_UShort4N:		D3DElement.Format = DXGI_FORMAT_R16G16B16A16_UNORM; break;
 			default: UE_LOG(LogD3D11RHI, Fatal,TEXT("Unknown RHI vertex element type %u"),(uint8)InElements[ElementIndex].Type);
 			};
 			D3DElement.SemanticName = "ATTRIBUTE";
