UE-16286 Adding check for Buffer->DecompressionState

- CreateCompressedAudioInfo can return null
- Can't repro crash internally but this fix will prevent a null deref on DecompressionState, which is what the crash stack appears to be doing.

diff --git a/Engine/Source/Runtime/Windows/XAudio2/Private/XAudio2Buffer.cpp b/Engine/Source/Runtime/Windows/XAudio2/Private/XAudio2Buffer.cpp
--- a/Engine/Source/Runtime/Windows/XAudio2/Private/XAudio2Buffer.cpp
+++ b/Engine/Source/Runtime/Windows/XAudio2/Private/XAudio2Buffer.cpp
@@ -353,7 +353,7 @@ FXAudio2SoundBuffer* FXAudio2SoundBuffer::CreateQueuedBuffer( FXAudio2Device* XA
 
 	Wave->InitAudioResource(XAudio2Device->GetRuntimeFormat(Wave));
 
-	if( Buffer->DecompressionState->ReadCompressedInfo( Wave->ResourceData, Wave->ResourceSize, &QualityInfo ) )
+	if (Buffer->DecompressionState && Buffer->DecompressionState->ReadCompressedInfo(Wave->ResourceData, Wave->ResourceSize, &QualityInfo))
 	{
 		// Refresh the wave data
 		Wave->SampleRate = QualityInfo.SampleRate;
