Asynchronous realtime audio decompression/procedural generation (XAudio2, Mac, Android)

diff --git a/Engine/Plugins/Runtime/SoundMod/Source/SoundMod/Private/SoundModWave.cpp b/Engine/Plugins/Runtime/SoundMod/Source/SoundMod/Private/SoundModWave.cpp
--- a/Engine/Plugins/Runtime/SoundMod/Source/SoundMod/Private/SoundModWave.cpp
+++ b/Engine/Plugins/Runtime/SoundMod/Source/SoundMod/Private/SoundModWave.cpp
@@ -6,6 +6,7 @@ USoundModWave::USoundModWave(const FObjectInitializer& ObjectInitializer)
 	: Super(ObjectInitializer)
 	, xmpContext(NULL)
 {
+	bCanProcessAsync = true;
 }
 
 DECLARE_CYCLE_STAT(TEXT("Sound Mod Generate Data"), STAT_SoundModGeneratePCMData, STATGROUP_Audio );
diff --git a/Engine/Source/Developer/AudioFormatOgg/Private/AudioFormatOgg.cpp b/Engine/Source/Developer/AudioFormatOgg/Private/AudioFormatOgg.cpp
--- a/Engine/Source/Developer/AudioFormatOgg/Private/AudioFormatOgg.cpp
+++ b/Engine/Source/Developer/AudioFormatOgg/Private/AudioFormatOgg.cpp
@@ -1,6 +1,7 @@
 // Copyright 1998-2015 Epic Games, Inc. All Rights Reserved.
 
 #include "Core.h"
+#include "CoreUObject.h"
 #include "ModuleInterface.h"
 #include "ModuleManager.h"
 #include "TargetPlatform.h"
diff --git a/Engine/Source/Runtime/Android/AndroidAudio/Private/AndroidAudioBuffer.cpp b/Engine/Source/Runtime/Android/AndroidAudio/Private/AndroidAudioBuffer.cpp
--- a/Engine/Source/Runtime/Android/AndroidAudio/Private/AndroidAudioBuffer.cpp
+++ b/Engine/Source/Runtime/Android/AndroidAudio/Private/AndroidAudioBuffer.cpp
@@ -41,6 +41,16 @@ FSLESSoundBuffer::~FSLESSoundBuffer( void )
 
 FSLESSoundBuffer* FSLESSoundBuffer::CreateQueuedBuffer( FSLESAudioDevice* AudioDevice, USoundWave* InWave )
 {
+	// Check to see if thread has finished decompressing on the other thread
+	if (InWave->AudioDecompressor != nullptr)
+	{
+		InWave->AudioDecompressor->EnsureCompletion();
+
+		// Remove the decompressor
+		delete InWave->AudioDecompressor;
+		InWave->AudioDecompressor = nullptr;
+	}
+
 	// Always create a new buffer for real time decompressed sounds
 	FSLESSoundBuffer* Buffer = new FSLESSoundBuffer( AudioDevice);
 	
@@ -49,16 +59,8 @@ FSLESSoundBuffer* FSLESSoundBuffer::CreateQueuedBuffer( FSLESAudioDevice* AudioD
 	
 	Buffer->DecompressionState = AudioDevice->CreateCompressedAudioInfo(InWave);
 
-	InWave->InitAudioResource( AudioDevice->GetRuntimeFormat(InWave) );
-	
-	if( Buffer->DecompressionState->ReadCompressedInfo( InWave->ResourceData, InWave->ResourceSize, &QualityInfo ) )
-	{
-		// Refresh the wave data
-		InWave->SampleRate = QualityInfo.SampleRate;
-		InWave->NumChannels = QualityInfo.NumChannels;
-		InWave->RawPCMDataSize = QualityInfo.SampleDataSize;
-		InWave->Duration = QualityInfo.Duration;
-		
+	if (Buffer->DecompressionState->ReadCompressedInfo( InWave->ResourceData, InWave->ResourceSize, &QualityInfo ))
+	{	
 		// Clear out any dangling pointers
 		Buffer->AudioData = NULL;
 		Buffer->BufferSize = 0;
@@ -80,8 +82,7 @@ FSLESSoundBuffer* FSLESSoundBuffer::CreateQueuedBuffer( FSLESAudioDevice* AudioD
 			InWave->RawPCMDataSize = SampleCount * QualityInfo.NumChannels * sizeof(uint16);;
 		}
 
-		Buffer->Format = SoundFormat_PCMRT;
-		
+		Buffer->Format = SoundFormat_PCMRT;	
 	}
 	else
 	{
@@ -91,7 +92,7 @@ FSLESSoundBuffer* FSLESSoundBuffer::CreateQueuedBuffer( FSLESAudioDevice* AudioD
 		InWave->RemoveAudioResource();
 	}
 	
-	return( Buffer );
+	return Buffer;
 }
 
 /**
diff --git a/Engine/Source/Runtime/Android/AndroidAudio/Private/AndroidAudioDevice.cpp b/Engine/Source/Runtime/Android/AndroidAudio/Private/AndroidAudioDevice.cpp
--- a/Engine/Source/Runtime/Android/AndroidAudio/Private/AndroidAudioDevice.cpp
+++ b/Engine/Source/Runtime/Android/AndroidAudio/Private/AndroidAudioDevice.cpp
@@ -8,8 +8,6 @@
 #include "AndroidAudioDevice.h"
 #include "VorbisAudioInfo.h"
 #include "ADPCMAudioInfo.h"
-#include "AudioEffect.h"
-#include "Engine.h"
 
 DEFINE_LOG_CATEGORY(LogAndroidAudio);
 
diff --git a/Engine/Source/Runtime/Android/AndroidAudio/Private/AndroidAudioSource.cpp b/Engine/Source/Runtime/Android/AndroidAudio/Private/AndroidAudioSource.cpp
--- a/Engine/Source/Runtime/Android/AndroidAudio/Private/AndroidAudioSource.cpp
+++ b/Engine/Source/Runtime/Android/AndroidAudio/Private/AndroidAudioSource.cpp
@@ -8,7 +8,6 @@
 #include "AudioDecompress.h"
 #include "Engine.h"
 
-
 // Callback that is registered if the source needs to loop
 void OpenSLBufferQueueCallback( SLAndroidSimpleBufferQueueItf InQueueInterface, void* pContext ) 
 {
@@ -44,15 +43,40 @@ void FSLESSoundSource::OnRequeueBufferCallback( SLAndroidSimpleBufferQueueItf In
 		}
 
 		// Enqueue the previously decoded buffer
+		if (RealtimeAsyncTask)
+		{
+			RealtimeAsyncTask->EnsureCompletion();
+			switch(RealtimeAsyncTask->GetTask().GetTaskType())
+			{
+			case ERealtimeAudioTaskType::Decompress:
+				bHasLooped = RealtimeAsyncTask->GetTask().GetBufferLooped();
+				break;
+
+			case ERealtimeAudioTaskType::Procedural:
+				AudioBuffers[BufferInUse].AudioDataSize = RealtimeAsyncTask->GetTask().GetBytesWritten();
+				break;
+			}
+
+			delete RealtimeAsyncTask;
+			RealtimeAsyncTask = nullptr;
+		}
+
 		SLresult result = (*SL_PlayerBufferQueue)->Enqueue(SL_PlayerBufferQueue, AudioBuffers[BufferInUse].AudioData, AudioBuffers[BufferInUse].AudioDataSize );
 		if(result != SL_RESULT_SUCCESS) 
 		{ 
 			UE_LOG( LogAndroidAudio, Warning, TEXT("FAILED OPENSL BUFFER Enqueue SL_PlayerBufferQueue (Requeing)"));  
 		}
 
-		// Switch to the next buffer and decode for the next time the callback fires
+		// Switch to the next buffer and decode for the next time the callback fires if we didn't just get the last buffer
 		BufferInUse = !BufferInUse;
-		bHasLooped = ReadMorePCMData(BufferInUse);
+		if (bHasLooped == false || WaveInstance->LoopingMode != LOOP_Never)
+		{
+			if (ReadMorePCMData(BufferInUse, EDataReadMode::Asynchronous))
+			{
+				// If this is a synchronous source we may get notified immediately that we have looped
+				bHasLooped = true;
+			}
+		}
 	}
 }
 
@@ -136,28 +160,38 @@ bool FSLESSoundSource::EnqueuePCMBuffer( bool bLoop)
 	return true;
 }
 
-bool FSLESSoundSource::ReadProceduralData(const int32 BufferIndex)
-{
-	ensure(WaveInstance);
-	const int32 MaxSamples = BufferSize / sizeof(int16);
-	const int32 BytesWritten = WaveInstance->WaveData->GeneratePCMData(AudioBuffers[BufferIndex].AudioData, MaxSamples);
-	
-	AudioBuffers[BufferIndex].AudioDataSize = BytesWritten;
-
-	// convenience return value: we're never actually "looping" here.
-	return false;
-}
-
-bool FSLESSoundSource::ReadMorePCMData(const int32 BufferIndex)
+bool FSLESSoundSource::ReadMorePCMData(const int32 BufferIndex, EDataReadMode DataReadMode)
 {
-	USoundWave *WaveData = WaveInstance->WaveData;
+	USoundWave* WaveData = WaveInstance->WaveData;
 	if (WaveData && WaveData->bProcedural)
 	{
-		return ReadProceduralData(BufferIndex);
+		const int32 MaxSamples = BufferSize / sizeof(int16);
+		if (DataReadMode == EDataReadMode::Synchronous || WaveData->bCanProcessAsync == false)
+		{
+			const int32 BytesWritten = WaveData->GeneratePCMData(AudioBuffers[BufferIndex].AudioData, MaxSamples);
+			AudioBuffers[BufferIndex].AudioDataSize = BytesWritten;
+		}
+		else
+		{
+			RealtimeAsyncTask = new FAsyncRealtimeAudioTask(WaveData, AudioBuffers[BufferIndex].AudioData, MaxSamples);
+			RealtimeAsyncTask->StartBackgroundTask();
+		}
+
+		// we're never actually "looping" here.
+		return false;
 	}
 	else
 	{
-		return Buffer->ReadCompressedData(AudioBuffers[BufferIndex].AudioData, WaveInstance->LoopingMode != LOOP_Never);
+		if (DataReadMode == EDataReadMode::Synchronous)
+		{
+			return Buffer->ReadCompressedData(AudioBuffers[BufferIndex].AudioData, WaveInstance->LoopingMode != LOOP_Never);
+		}
+		else
+		{
+			RealtimeAsyncTask = new FAsyncRealtimeAudioTask(Buffer, AudioBuffers[BufferIndex].AudioData, WaveInstance->LoopingMode != LOOP_Never, DataReadMode == EDataReadMode::AsynchronousSkipFirstFrame);
+			RealtimeAsyncTask->StartBackgroundTask();
+			return false;
+		}
 	}
 }
 
@@ -179,9 +213,20 @@ bool FSLESSoundSource::EnqueuePCMRTBuffer( bool bLoop )
 	AudioBuffers[1].AudioData = (uint8*)FMemory::Malloc(BufferSize);
 	AudioBuffers[1].AudioDataSize = BufferSize;
 
-	// Decompress two buffers worth of data to fill up the queue
-	ReadMorePCMData(0);
-	ReadMorePCMData(1);
+	// Only use the cached data if we're starting from the beginning, otherwise we'll have to take a synchronous hit
+	bool bSkipFirstBuffer = false;;
+	if (WaveInstance->WaveData && WaveInstance->WaveData->CachedRealtimeFirstBuffer && WaveInstance->StartTime == 0.f)
+	{
+		FMemory::Memcpy((uint8*)AudioBuffers[0].AudioData, WaveInstance->WaveData->CachedRealtimeFirstBuffer, BufferSize);
+		bSkipFirstBuffer = true;
+	}
+	else
+	{
+		ReadMorePCMData(0, EDataReadMode::Synchronous);
+	}
+
+	// Start the async population of the next buffer
+	ReadMorePCMData(1, (bSkipFirstBuffer ? EDataReadMode::AsynchronousSkipFirstFrame : EDataReadMode::Asynchronous));
 
 	SLresult result;
 
@@ -315,7 +360,6 @@ FSLESSoundSource::FSLESSoundSource( class FAudioDevice* InAudioDevice )
  */
 FSLESSoundSource::~FSLESSoundSource( void )
 {
-
 	DestroyPlayer();
 
 	ReleaseResources();
@@ -323,6 +367,13 @@ FSLESSoundSource::~FSLESSoundSource( void )
 
 void FSLESSoundSource::ReleaseResources()
 {
+	if (RealtimeAsyncTask)
+	{
+		RealtimeAsyncTask->EnsureCompletion();
+		delete RealtimeAsyncTask;
+		RealtimeAsyncTask = nullptr;
+	}
+
 	FMemory::Free( AudioBuffers[0].AudioData);
 	FMemory::Free( AudioBuffers[1].AudioData);
 
diff --git a/Engine/Source/Runtime/Android/AndroidAudio/Public/AndroidAudioDevice.h b/Engine/Source/Runtime/Android/AndroidAudio/Public/AndroidAudioDevice.h
--- a/Engine/Source/Runtime/Android/AndroidAudio/Public/AndroidAudioDevice.h
+++ b/Engine/Source/Runtime/Android/AndroidAudio/Public/AndroidAudioDevice.h
@@ -8,8 +8,10 @@
 
 class FSLESAudioDevice;
 
+#include "Core.h"
 #include "Engine.h"
 #include "SoundDefinitions.h"
+#include "AudioDecompress.h"
 #include "AudioEffect.h"
 
 #include <SLES/OpenSLES.h>
@@ -129,6 +131,8 @@ public:
 	
 };
 
+typedef FAsyncTask<FAsyncRealtimeAudioTaskWorker<FSLESSoundBuffer>> FAsyncRealtimeAudioTask;
+
 /**
  * OpenSLES implementation of FSoundSource, the interface used to play, stop and update sources
  */
@@ -198,6 +202,13 @@ public:
 
 protected:
 
+	enum class EDataReadMode : uint8
+	{
+		Synchronous,
+		Asynchronous,
+		AsynchronousSkipFirstFrame
+	};
+
 	friend class FSLESAudioDevice;
 	FSLESSoundBuffer*		Buffer;
 	FSLESAudioDevice*		Device;
@@ -208,6 +219,8 @@ protected:
 	SLAndroidSimpleBufferQueueItf	SL_PlayerBufferQueue;
 	SLVolumeItf						SL_VolumeInterface;
 
+	FAsyncRealtimeAudioTask* RealtimeAsyncTask;
+
 	/** Which sound buffer should be written to next - used for double buffering. */
 	bool						bStreamedSound;
 	/** A pair of sound buffers for real-time decoding */
@@ -228,10 +241,7 @@ protected:
 	bool EnqueuePCMRTBuffer( bool bLoop);
 
 	/** Decompress through FSLESSoundBuffer, or call USoundWave procedure to generate more PCM data. Returns true/false: did audio loop? */
-	bool ReadMorePCMData(const int32 BufferIndex);
-
-	/** Handle obtaining more data for procedural USoundWaves. Always returns false for convenience. */
-	bool ReadProceduralData(const int32 BufferIndex);
+	bool ReadMorePCMData( const int32 BufferIndex, EDataReadMode DataReadMode );
 };
 
 /**
diff --git a/Engine/Source/Runtime/Engine/Classes/Sound/SoundWave.h b/Engine/Source/Runtime/Engine/Classes/Sound/SoundWave.h
--- a/Engine/Source/Runtime/Engine/Classes/Sound/SoundWave.h
+++ b/Engine/Source/Runtime/Engine/Classes/Sound/SoundWave.h
@@ -5,6 +5,7 @@
 /** 
  * Playable sound object for raw wave files
  */
+#include "BulkData.h"
 #include "Sound/SoundBase.h"
 #include "Sound/SoundGroups.h"
 #include "SoundWave.generated.h"
@@ -132,6 +133,9 @@ class ENGINE_API USoundWave : public USoundBase
 	/** Set to true for programmatically-generated, streamed audio. */
 	uint32 bProcedural:1;
 
+	/** Set to true for procedural waves that can be processed asynchronously. */
+	uint32 bCanProcessAsync:1;
+
 	/** Whether to free the resource data after it has been uploaded to the hardware */
 	uint32 bDynamicResource:1;
 
@@ -225,6 +229,9 @@ public:
 	typedef FAsyncTask< class FAsyncAudioDecompressWorker > FAsyncAudioDecompress;	// Forward declare typedef
 	FAsyncAudioDecompress*		AudioDecompressor;
 
+	/** Pointer to 16 bit PCM data - used to avoid synchronous operation to obtain first block of the relatime decompressed buffer */
+	uint8*						CachedRealtimeFirstBuffer;
+
 	/** Pointer to 16 bit PCM data - used to decompress data to and preview sounds */
 	uint8*						RawPCMData;
 
diff --git a/Engine/Source/Runtime/Engine/Private/AudioDecompress.cpp b/Engine/Source/Runtime/Engine/Private/AudioDecompress.cpp
--- a/Engine/Source/Runtime/Engine/Private/AudioDecompress.cpp
+++ b/Engine/Source/Runtime/Engine/Private/AudioDecompress.cpp
@@ -290,11 +290,11 @@ FAsyncAudioDecompressWorker::FAsyncAudioDecompressWorker(USoundWave* InWave)
 	}
 }
 
-void FAsyncAudioDecompressWorker::DoWork( void )
+void FAsyncAudioDecompressWorker::DoWork()
 {
 	if (AudioInfo)
 	{
-		FSoundQualityInfo	QualityInfo = { 0 };
+		FSoundQualityInfo QualityInfo = { 0 };
 
 		// Parse the audio header for the relevant information
 		if (AudioInfo->ReadCompressedInfo(Wave->ResourceData, Wave->ResourceSize, &QualityInfo))
@@ -317,22 +317,44 @@ void FAsyncAudioDecompressWorker::DoWork( void )
 			// Extract the data
 			Wave->SampleRate = QualityInfo.SampleRate;
 			Wave->NumChannels = QualityInfo.NumChannels;
-			if(QualityInfo.Duration > 0.0f) Wave->Duration = QualityInfo.Duration;
+			if (QualityInfo.Duration > 0.0f)
+			{ 
+				Wave->Duration = QualityInfo.Duration;
+			}
+
+			if (Wave->DecompressionType == DTYPE_RealTime)
+			{
+				const uint32 PCMBufferSize = MONO_PCM_BUFFER_SIZE * Wave->NumChannels;
+				Wave->CachedRealtimeFirstBuffer = (uint8*)FMemory::Malloc(PCMBufferSize);
+				AudioInfo->ReadCompressedData( Wave->CachedRealtimeFirstBuffer, false, PCMBufferSize );
+			}
+			else
+			{
+				Wave->RawPCMDataSize = QualityInfo.SampleDataSize;
+				Wave->RawPCMData = ( uint8* )FMemory::Malloc( Wave->RawPCMDataSize );
 
-			Wave->RawPCMDataSize = QualityInfo.SampleDataSize;
-			Wave->RawPCMData = ( uint8* )FMemory::Malloc( Wave->RawPCMDataSize );
+				// Decompress all the sample data into preallocated memory
+				AudioInfo->ExpandFile(Wave->RawPCMData, &QualityInfo);
 
-			// Decompress all the sample data into preallocated memory
-			AudioInfo->ExpandFile(Wave->RawPCMData, &QualityInfo);
+				const SIZE_T ResSize = Wave->GetResourceSize(EResourceSizeMode::Exclusive);
+				Wave->TrackedMemoryUsage += ResSize;
+				INC_DWORD_STAT_BY( STAT_AudioMemorySize, ResSize );
+				INC_DWORD_STAT_BY( STAT_AudioMemory, ResSize );
+			}
+		}
+		else if (Wave->DecompressionType == DTYPE_RealTime)
+		{
+			Wave->DecompressionType = DTYPE_Invalid;
+			Wave->NumChannels = 0;
 
-			const SIZE_T ResSize = Wave->GetResourceSize(EResourceSizeMode::Exclusive);
-			Wave->TrackedMemoryUsage += ResSize;
-			INC_DWORD_STAT_BY( STAT_AudioMemorySize, ResSize );
-			INC_DWORD_STAT_BY( STAT_AudioMemory, ResSize );
+			Wave->RemoveAudioResource();
 		}
 
-		// Delete the compressed data
-		Wave->RemoveAudioResource();
+		if (Wave->DecompressionType == DTYPE_Native)
+		{
+			// Delete the compressed data
+			Wave->RemoveAudioResource();
+		}
 
 		delete AudioInfo;
 	}
diff --git a/Engine/Source/Runtime/Engine/Private/AudioDevice.cpp b/Engine/Source/Runtime/Engine/Private/AudioDevice.cpp
--- a/Engine/Source/Runtime/Engine/Private/AudioDevice.cpp
+++ b/Engine/Source/Runtime/Engine/Private/AudioDevice.cpp
@@ -2317,6 +2317,7 @@ void FAudioDevice::Precache(USoundWave* SoundWave, bool bSynchronous, bool bTrac
 		if (FPlatformProperties::SupportsAudioStreaming() && SoundWave->IsStreaming())
 		{
 			SoundWave->DecompressionType = DTYPE_Streaming;
+			SoundWave->bCanProcessAsync = false;
 		}
 		else if (SupportsRealtimeDecompression() && 
 			(bDisableAudioCaching || (!SoundGroup.bAlwaysDecompressOnLoad && SoundWave->Duration > SoundGroup.DecompressedDuration)))
@@ -2342,9 +2343,8 @@ void FAudioDevice::Precache(USoundWave* SoundWave, bool bSynchronous, bool bTrac
 		// Grab the compressed audio data
 		SoundWave->InitAudioResource(GetRuntimeFormat(SoundWave));
 
-		if (SoundWave->AudioDecompressor == NULL && SoundWave->DecompressionType == DTYPE_Native)
+		if (SoundWave->AudioDecompressor == nullptr && (SoundWave->DecompressionType == DTYPE_Native || SoundWave->DecompressionType == DTYPE_RealTime))
 		{
-			check(!SoundWave->AudioDecompressor); // should not have had a valid pointer at this point
 			// Create a worker to decompress the audio data
 			if (bSynchronous)
 			{
@@ -2354,6 +2354,10 @@ void FAudioDevice::Precache(USoundWave* SoundWave, bool bSynchronous, bool bTrac
 			}
 			else
 			{
+				if (SoundWave->DecompressionType == DTYPE_RealTime)
+				{
+					SoundWave->CachedRealtimeFirstBuffer = (uint8*)FMemory::Malloc( MONO_PCM_BUFFER_SIZE * SoundWave->NumChannels );
+				}
 				SoundWave->AudioDecompressor = new FAsyncAudioDecompress(SoundWave);
 				SoundWave->AudioDecompressor->StartBackgroundTask();
 			}
@@ -2362,7 +2366,10 @@ void FAudioDevice::Precache(USoundWave* SoundWave, bool bSynchronous, bool bTrac
 			SoundWave->bDecompressedFromOgg = GetRuntimeFormat(SoundWave) == NAME_OGG;
 
 			// the audio decompressor will track memory
-			bTrackMemory = false;
+			if (SoundWave->DecompressionType == DTYPE_Native)
+			{
+				bTrackMemory = false;
+			}
 		}
 	}
 	else
diff --git a/Engine/Source/Runtime/Engine/Private/SoundWave.cpp b/Engine/Source/Runtime/Engine/Private/SoundWave.cpp
--- a/Engine/Source/Runtime/Engine/Private/SoundWave.cpp
+++ b/Engine/Source/Runtime/Engine/Private/SoundWave.cpp
@@ -71,11 +71,18 @@ SIZE_T USoundWave::GetResourceSize(EResourceSizeMode::Type Mode)
 {
 	int32 CalculatedResourceSize = 0;
 
-	if( DecompressionType == DTYPE_Native )
+	if (DecompressionType == DTYPE_Native)
 	{
 		// If we've been decompressed, need to account for decompressed and also compressed
 		CalculatedResourceSize += RawPCMDataSize;
 	}
+	else if (DecompressionType == DTYPE_RealTime)
+	{
+		if (CachedRealtimeFirstBuffer)
+		{
+			CalculatedResourceSize += MONO_PCM_BUFFER_SIZE * NumChannels;
+		}
+	}
 
 	if (GEngine && GEngine->GetMainAudioDevice())
 	{
@@ -467,18 +474,24 @@ void USoundWave::FreeResources()
 	{
 		// Notify the audio device to free the bulk data associated with this wave.
 		FAudioDeviceManager* AudioDeviceManager = GEngine->GetAudioDeviceManager();
-		if (AudioDeviceManager != NULL)
+		if (AudioDeviceManager)
 		{
 			AudioDeviceManager->StopSoundsUsingWave(this);
 			AudioDeviceManager->FreeResource(this);
 		}
 	}
 
+	if (CachedRealtimeFirstBuffer)
+	{
+		FMemory::Free(CachedRealtimeFirstBuffer);
+		CachedRealtimeFirstBuffer = nullptr;
+	}
+
 	// Just in case the data was created but never uploaded
-	if (RawPCMData != NULL)
+	if (RawPCMData)
 	{
 		FMemory::Free(RawPCMData);
-		RawPCMData = NULL;
+		RawPCMData = nullptr;
 	}
 
 	// Remove the compressed copy of the data
diff --git a/Engine/Source/Runtime/Engine/Public/AudioDecompress.h b/Engine/Source/Runtime/Engine/Public/AudioDecompress.h
--- a/Engine/Source/Runtime/Engine/Public/AudioDecompress.h
+++ b/Engine/Source/Runtime/Engine/Public/AudioDecompress.h
@@ -6,6 +6,8 @@
 
 #pragma once
 
+#include "Sound/SoundWave.h"
+
 // 186ms of 44.1KHz data
 // 372ms of 22KHz data
 #define MONO_PCM_BUFFER_SAMPLES		8192
@@ -146,3 +148,93 @@ public:
 
 typedef FAsyncTask<FAsyncAudioDecompressWorker> FAsyncAudioDecompress;
 
+enum class ERealtimeAudioTaskType : uint8
+{
+	Decompress,
+	Procedural
+};
+
+template<class T>
+class FAsyncRealtimeAudioTaskWorker : public FNonAbandonableTask
+{
+protected:
+	T* AudioBuffer;
+	USoundWave* WaveData;
+	uint8* AudioData;
+	int32 MaxSamples;
+	int32 BytesWritten;
+	ERealtimeAudioTaskType TaskType;
+	uint32 bSkipFirstBuffer:1;
+	uint32 bLoopingMode:1;
+	uint32 bLooped:1;
+
+public:
+	FAsyncRealtimeAudioTaskWorker(T* InAudioBuffer, uint8* InAudioData, bool bInLoopingMode, bool bInSkipFirstBuffer)
+		: AudioBuffer(InAudioBuffer)
+		, AudioData(InAudioData)
+		, TaskType(ERealtimeAudioTaskType::Decompress)
+		, bSkipFirstBuffer(bInSkipFirstBuffer)
+		, bLoopingMode(bInLoopingMode)
+		, bLooped(false)
+	{
+	}
+
+	FAsyncRealtimeAudioTaskWorker(USoundWave* InWaveData, uint8* InAudioData, int32 InMaxSamples)
+		: WaveData(InWaveData)
+		, AudioData(InAudioData)
+		, MaxSamples(InMaxSamples)
+		, BytesWritten(0)
+		, TaskType(ERealtimeAudioTaskType::Procedural)
+	{
+	}
+
+	void DoWork()
+	{
+		switch(TaskType)
+		{
+		case ERealtimeAudioTaskType::Decompress:
+			if (bSkipFirstBuffer)
+			{
+				AudioBuffer->ReadCompressedData( ( uint8* )AudioData, bLoopingMode );
+			}
+			bLooped = AudioBuffer->ReadCompressedData( ( uint8* )AudioData, bLoopingMode );
+			break;
+
+		case ERealtimeAudioTaskType::Procedural:
+			BytesWritten = WaveData->GeneratePCMData( (uint8*)AudioData, MaxSamples );
+			break;
+
+		default:
+			check(false);
+		}
+	}
+
+	FORCEINLINE TStatId GetStatId() const
+	{
+		if (TaskType == ERealtimeAudioTaskType::Procedural)
+		{
+			RETURN_QUICK_DECLARE_CYCLE_STAT(FAsyncRealtimeAudioProceduralWorker, STATGROUP_ThreadPoolAsyncTasks);
+		}
+		else
+		{
+			RETURN_QUICK_DECLARE_CYCLE_STAT(FAsyncRealtimeAudioDecompressWorker, STATGROUP_ThreadPoolAsyncTasks);
+		}
+	}
+
+	ERealtimeAudioTaskType GetTaskType() const
+	{
+		return TaskType;
+	}
+
+	bool GetBufferLooped() const
+	{ 
+		check(TaskType == ERealtimeAudioTaskType::Decompress);
+		return bLooped;
+	}
+
+	int32 GetBytesWritten() const
+	{ 
+		check(TaskType == ERealtimeAudioTaskType::Procedural);
+		return BytesWritten;
+	}
+};
\ No newline at end of file
diff --git a/Engine/Source/Runtime/Mac/CoreAudio/Private/CoreAudioBuffer.cpp b/Engine/Source/Runtime/Mac/CoreAudio/Private/CoreAudioBuffer.cpp
--- a/Engine/Source/Runtime/Mac/CoreAudio/Private/CoreAudioBuffer.cpp
+++ b/Engine/Source/Runtime/Mac/CoreAudio/Private/CoreAudioBuffer.cpp
@@ -138,6 +138,16 @@ void FCoreAudioSoundBuffer::Seek( const float SeekTime )
  */
 FCoreAudioSoundBuffer* FCoreAudioSoundBuffer::CreateQueuedBuffer( FCoreAudioDevice* CoreAudioDevice, USoundWave* Wave )
 {
+	// Check to see if thread has finished decompressing on the other thread
+	if (Wave->AudioDecompressor != nullptr)
+	{
+		Wave->AudioDecompressor->EnsureCompletion();
+
+		// Remove the decompressor
+		delete Wave->AudioDecompressor;
+		Wave->AudioDecompressor = nullptr;
+	}
+
 	// Always create a new buffer for real time decompressed sounds
 	FCoreAudioSoundBuffer* Buffer = new FCoreAudioSoundBuffer( CoreAudioDevice, SoundFormat_PCMRT );
 	
@@ -146,16 +156,8 @@ FCoreAudioSoundBuffer* FCoreAudioSoundBuffer::CreateQueuedBuffer( FCoreAudioDevi
 
 	Buffer->DecompressionState = CoreAudioDevice->CreateCompressedAudioInfo(Wave);
 	
-	Wave->InitAudioResource( CoreAudioDevice->GetRuntimeFormat(Wave) );
-	
-	if( Buffer->DecompressionState->ReadCompressedInfo( Wave->ResourceData, Wave->ResourceSize, &QualityInfo ) )
+	if (Buffer->DecompressionState->ReadCompressedInfo( Wave->ResourceData, Wave->ResourceSize, &QualityInfo ))
 	{
-		// Refresh the wave data
-		Wave->SampleRate = QualityInfo.SampleRate;
-		Wave->NumChannels = QualityInfo.NumChannels;
-		Wave->RawPCMDataSize = QualityInfo.SampleDataSize;
-		Wave->Duration = QualityInfo.Duration;
-		
 		// Clear out any dangling pointers
 		Buffer->PCMData = NULL;
 		Buffer->PCMDataSize = 0;
@@ -170,7 +172,7 @@ FCoreAudioSoundBuffer* FCoreAudioSoundBuffer::CreateQueuedBuffer( FCoreAudioDevi
 		Wave->RemoveAudioResource();
 	}
 	
-	return( Buffer );
+	return Buffer;
 }
 
 /**
diff --git a/Engine/Source/Runtime/Mac/CoreAudio/Private/CoreAudioSource.cpp b/Engine/Source/Runtime/Mac/CoreAudio/Private/CoreAudioSource.cpp
--- a/Engine/Source/Runtime/Mac/CoreAudio/Private/CoreAudioSource.cpp
+++ b/Engine/Source/Runtime/Mac/CoreAudio/Private/CoreAudioSource.cpp
@@ -51,6 +51,7 @@ FCoreAudioSoundSource::FCoreAudioSoundSource( FAudioDevice* InAudioDevice )
 :	FSoundSource( InAudioDevice ),
 	Buffer( NULL ),
 	CoreAudioConverter( NULL ),
+	RealtimeAsyncTask(nullptr),
 	bStreamedSound( false ),
 	bBuffersToFlush( false ),
 	SourceNode( 0 ),
@@ -94,6 +95,13 @@ FCoreAudioSoundSource::~FCoreAudioSoundSource()
  */
 void FCoreAudioSoundSource::FreeResources( void )
 {
+	if (RealtimeAsyncTask)
+	{
+		RealtimeAsyncTask->EnsureCompletion();
+		delete RealtimeAsyncTask;
+		RealtimeAsyncTask = nullptr;
+	}
+
 	// If we're a streaming buffer...
 	if( bStreamedSound )
 	{
@@ -106,7 +114,7 @@ void FCoreAudioSoundSource::FreeResources( void )
 		{
 			check( Buffer->ResourceID == 0 );
 			delete Buffer;
-			Buffer = NULL;
+			Buffer = nullptr;
 		}
 		
 		bStreamedSound = false;
@@ -131,35 +139,47 @@ void FCoreAudioSoundSource::SubmitPCMBuffers( void )
 	CoreAudioBuffers[0].AudioDataSize = Buffer->PCMDataSize;
 }
 
-bool FCoreAudioSoundSource::ReadProceduralData( const int32 BufferIndex )
-{
-	const int32 MaxSamples = ( MONO_PCM_BUFFER_SIZE * Buffer->NumChannels ) / sizeof( int16 );
-	const int32 BytesRead = WaveInstance->WaveData->GeneratePCMData( CoreAudioBuffers[BufferIndex].AudioData, MaxSamples );
-
-	CoreAudioBuffers[BufferIndex].AudioDataSize = BytesRead;
-
-	if (BytesRead > 0)
-	{
-		++NumActiveBuffers;
-	}
-	
-	// convenience return value: we're never actually "looping" here.
-	return false;
-}
-
-bool FCoreAudioSoundSource::ReadMorePCMData( const int32 BufferIndex )
+bool FCoreAudioSoundSource::ReadMorePCMData( const int32 BufferIndex, EDataReadMode DataReadMode )
 {
 	CoreAudioBuffers[BufferIndex].ReadCursor = 0;
 
-	USoundWave *WaveData = WaveInstance->WaveData;
+	USoundWave* WaveData = WaveInstance->WaveData;
 	if( WaveData && WaveData->bProcedural )
 	{
-		return ReadProceduralData( BufferIndex );
+		const int32 MaxSamples = ( MONO_PCM_BUFFER_SIZE * Buffer->NumChannels ) / sizeof( int16 );
+
+		if (DataReadMode == EDataReadMode::Synchronous || WaveData->bCanProcessAsync == false)
+		{
+			const int32 BytesRead = WaveData->GeneratePCMData(CoreAudioBuffers[BufferIndex].AudioData, MaxSamples);
+			CoreAudioBuffers[BufferIndex].AudioDataSize = BytesRead;
+
+			if (BytesRead > 0)
+			{
+				++NumActiveBuffers;
+			}
+		}
+		else
+		{
+			RealtimeAsyncTask = new FAsyncRealtimeAudioTask(WaveData, (uint8*)CoreAudioBuffers[BufferIndex].AudioData, MaxSamples);
+			RealtimeAsyncTask->StartBackgroundTask();
+		}
+	
+		// we're never actually "looping" here.
+		return false;
 	}
 	else
 	{
-		NumActiveBuffers++;
-		return Buffer->ReadCompressedData( CoreAudioBuffers[BufferIndex].AudioData, WaveInstance->LoopingMode != LOOP_Never );
+		if (DataReadMode == EDataReadMode::Synchronous)
+		{
+			++NumActiveBuffers;
+			return Buffer->ReadCompressedData( CoreAudioBuffers[BufferIndex].AudioData, WaveInstance->LoopingMode != LOOP_Never );
+		}
+		else
+		{
+			RealtimeAsyncTask = new FAsyncRealtimeAudioTask(Buffer, (uint8*)CoreAudioBuffers[BufferIndex].AudioData, WaveInstance->LoopingMode != LOOP_Never, DataReadMode == EDataReadMode::AsynchronousSkipFirstFrame);
+			RealtimeAsyncTask->StartBackgroundTask();
+			return false;
+		}
 	}
 }
 
@@ -174,18 +194,33 @@ void FCoreAudioSoundSource::SubmitPCMRTBuffers( void )
 
 	bStreamedSound = true;
 
+	const uint32 BufferSize = MONO_PCM_BUFFER_SIZE * Buffer->NumChannels;
+
 	// Set up double buffer area to decompress to
-	CoreAudioBuffers[0].AudioData = ( uint8* )FMemory::Malloc( MONO_PCM_BUFFER_SIZE * Buffer->NumChannels );
-	CoreAudioBuffers[0].AudioDataSize = MONO_PCM_BUFFER_SIZE * Buffer->NumChannels;
+	CoreAudioBuffers[0].AudioData = (uint8*)FMemory::Malloc(BufferSize);
+	CoreAudioBuffers[0].AudioDataSize = BufferSize;
 
-	CoreAudioBuffers[1].AudioData = ( uint8* )FMemory::Malloc( MONO_PCM_BUFFER_SIZE * Buffer->NumChannels );
-	CoreAudioBuffers[1].AudioDataSize = MONO_PCM_BUFFER_SIZE * Buffer->NumChannels;
+	CoreAudioBuffers[1].AudioData = (uint8*)FMemory::Malloc(BufferSize);
+	CoreAudioBuffers[1].AudioDataSize = BufferSize;
 
 	NumActiveBuffers = 0;
 	BufferInUse = 0;
 	
-	ReadMorePCMData(0);
-	ReadMorePCMData(1);
+	// Only use the cached data if we're starting from the beginning, otherwise we'll have to take a synchronous hit
+	bool bSkipFirstBuffer = false;;
+	if (WaveInstance->WaveData && WaveInstance->WaveData->CachedRealtimeFirstBuffer && WaveInstance->StartTime == 0.f)
+	{
+		FMemory::Memcpy((uint8*)CoreAudioBuffers[0].AudioData, WaveInstance->WaveData->CachedRealtimeFirstBuffer, BufferSize);
+		bSkipFirstBuffer = true;
+		NumActiveBuffers = 1;
+	}
+	else
+	{
+		ReadMorePCMData(0, EDataReadMode::Synchronous);
+	}
+
+	// Start the async population of the next buffer
+	ReadMorePCMData(1, (bSkipFirstBuffer ? EDataReadMode::AsynchronousSkipFirstFrame : EDataReadMode::Asynchronous));
 }
 
 /**
@@ -469,29 +504,59 @@ void FCoreAudioSoundSource::Pause( void )
 /**
  * Handles feeding new data to a real time decompressed sound
  */
-void FCoreAudioSoundSource::HandleRealTimeSource( void )
+void FCoreAudioSoundSource::HandleRealTimeSource()
 {
-	// Get the next bit of streaming data
-	const bool bLooped = ReadMorePCMData(1 - BufferInUse);
+	bool bLooped = false;
+	const int32 BufferIndex = 1 - BufferInUse;
+	if (RealtimeAsyncTask && RealtimeAsyncTask->IsDone())
+	{
+		switch(RealtimeAsyncTask->GetTask().GetTaskType())
+		{
+		case ERealtimeAudioTaskType::Decompress:
+			bLooped = RealtimeAsyncTask->GetTask().GetBufferLooped();
+			++NumActiveBuffers;
+			break;
+
+		case ERealtimeAudioTaskType::Procedural:
+			const int32 BytesWritten = RealtimeAsyncTask->GetTask().GetBytesWritten();
+			CoreAudioBuffers[BufferIndex].AudioDataSize = BytesWritten;
+			if (BytesWritten > 0)
+			{
+				++NumActiveBuffers;
+			}
+			break;
+		}
 
-	// Have we reached the end of the compressed sound?
-	if( bLooped )
+		delete RealtimeAsyncTask;
+		RealtimeAsyncTask = nullptr;
+	}
+	else
+	{
+		// Get the next bit of streaming data
+		bLooped = ReadMorePCMData(BufferIndex, EDataReadMode::Asynchronous);
+	}
+
+	if (RealtimeAsyncTask == nullptr)
 	{
-		switch( WaveInstance->LoopingMode )
+		// Have we reached the end of the compressed sound?
+		if( bLooped )
 		{
-			case LOOP_Never:
-				// Play out any queued buffers - once there are no buffers left, the state check at the beginning of IsFinished will fire
-				bBuffersToFlush = true;
-				break;
+			switch( WaveInstance->LoopingMode )
+			{
+				case LOOP_Never:
+					// Play out any queued buffers - once there are no buffers left, the state check at the beginning of IsFinished will fire
+					bBuffersToFlush = true;
+					break;
 				
-			case LOOP_WithNotification:
-				// If we have just looped, and we are programmatically looping, send notification
-				WaveInstance->NotifyFinished();
-				break;
+				case LOOP_WithNotification:
+					// If we have just looped, and we are programmatically looping, send notification
+					WaveInstance->NotifyFinished();
+					break;
 				
-			case LOOP_Forever:
-				// Let the sound loop indefinitely
-				break;
+				case LOOP_Forever:
+					// Let the sound loop indefinitely
+					break;
+			}
 		}
 	}
 }
@@ -513,30 +578,19 @@ bool FCoreAudioSoundSource::IsFinished( void )
 	if( WaveInstance )
 	{
 		// If not rendering, we're either at the end of a sound, or starved
-		if( NumActiveBuffers == 0 )
+		// and we are expecting the sound to be finishing
+		if (NumActiveBuffers == 0 && (bBuffersToFlush || !bStreamedSound))
 		{
-			// ... are we expecting the sound to be finishing?
-			if( bBuffersToFlush || !bStreamedSound )
-			{
-				// ... notify the wave instance that it has finished playing.
-				WaveInstance->NotifyFinished();
-				return( true );
-			}
+			// ... notify the wave instance that it has finished playing.
+			WaveInstance->NotifyFinished();
+			return true;
 		}
 
 		// Service any real time sounds
-		if( bStreamedSound )
+		if( bStreamedSound && !bBuffersToFlush && NumActiveBuffers < 2)
 		{
-			if( NumActiveBuffers < 2 )
-			{
-				// Continue feeding new sound data (unless we are waiting for the sound to finish)
-				if( !bBuffersToFlush )
-				{
-					HandleRealTimeSource();
-				}
-
-				return( false );
-			}
+			// Continue feeding new sound data (unless we are waiting for the sound to finish)
+			HandleRealTimeSource();
 		}
 
 		return( false );
diff --git a/Engine/Source/Runtime/Mac/CoreAudio/Public/CoreAudioDevice.h b/Engine/Source/Runtime/Mac/CoreAudio/Public/CoreAudioDevice.h
--- a/Engine/Source/Runtime/Mac/CoreAudio/Public/CoreAudioDevice.h
+++ b/Engine/Source/Runtime/Mac/CoreAudio/Public/CoreAudioDevice.h
@@ -161,6 +161,9 @@ public:
 	bool						bDynamicResource;
 };
 
+typedef FAsyncTask<class FAsyncRealtimeAudioTaskWorker<FCoreAudioSoundBuffer>> FAsyncRealtimeAudioTask;
+
+
 /**
  * CoreAudio implementation of FSoundSource, the interface used to play, stop and update sources
  */
@@ -238,11 +241,16 @@ public:
 											 AudioStreamPacketDescription **outPacketDescription, void *inUserData );
 
 protected:
-	/** Decompress USoundWave procedure to generate more PCM data. Returns true/false: did audio loop? */
-	bool ReadMorePCMData( const int32 BufferIndex );
 
-	/** Handle obtaining more data for procedural USoundWaves. Always returns false for convenience. */
-	bool ReadProceduralData( const int32 BufferIndex );
+	enum class EDataReadMode : uint8
+	{
+		Synchronous,
+		Asynchronous,
+		AsynchronousSkipFirstFrame
+	};
+
+	/** Decompress USoundWave procedure to generate more PCM data. Returns true/false: did audio loop? */
+	bool ReadMorePCMData( const int32 BufferIndex, EDataReadMode DataReadMode );
 
 	OSStatus CreateAudioUnit( OSType Type, OSType SubType, OSType Manufacturer, AudioStreamBasicDescription* InputFormat, AudioStreamBasicDescription* OutputFormat, AUNode* OutNode, AudioUnit* OutUnit );
 	OSStatus ConnectAudioUnit( AUNode DestNode, uint32 DestInputNumber, AUNode OutNode, AudioUnit OutUnit );
@@ -259,6 +267,9 @@ protected:
 
 	AudioConverterRef			CoreAudioConverter;
 
+	/** Asynchronous task for real time audio sources */
+	FAsyncRealtimeAudioTask* RealtimeAsyncTask;
+
 	/** Which sound buffer should be written to next - used for double buffering. */
 	bool						bStreamedSound;
 	/** A pair of sound buffers to allow notification when a sound loops. */
diff --git a/Engine/Source/Runtime/Windows/XAudio2/Private/XAudio2Buffer.cpp b/Engine/Source/Runtime/Windows/XAudio2/Private/XAudio2Buffer.cpp
--- a/Engine/Source/Runtime/Windows/XAudio2/Private/XAudio2Buffer.cpp
+++ b/Engine/Source/Runtime/Windows/XAudio2/Private/XAudio2Buffer.cpp
@@ -343,6 +343,16 @@ void FXAudio2SoundBuffer::Seek( const float SeekTime )
  */
 FXAudio2SoundBuffer* FXAudio2SoundBuffer::CreateQueuedBuffer( FXAudio2Device* XAudio2Device, USoundWave* Wave )
 {
+	// Check to see if thread has finished decompressing on the other thread
+	if( Wave->AudioDecompressor != nullptr )
+	{
+		Wave->AudioDecompressor->EnsureCompletion();
+
+		// Remove the decompressor
+		delete Wave->AudioDecompressor;
+		Wave->AudioDecompressor = nullptr;
+	}
+
 	// Always create a new buffer for real time decompressed sounds
 	FXAudio2SoundBuffer* Buffer = new FXAudio2SoundBuffer( XAudio2Device, SoundFormat_PCMRT );
 
@@ -351,16 +361,8 @@ FXAudio2SoundBuffer* FXAudio2SoundBuffer::CreateQueuedBuffer( FXAudio2Device* XA
 
 	Buffer->DecompressionState = XAudio2Device->CreateCompressedAudioInfo(Wave);
 
-	Wave->InitAudioResource(XAudio2Device->GetRuntimeFormat(Wave));
-
 	if (Buffer->DecompressionState && Buffer->DecompressionState->ReadCompressedInfo(Wave->ResourceData, Wave->ResourceSize, &QualityInfo))
 	{
-		// Refresh the wave data
-		Wave->SampleRate = QualityInfo.SampleRate;
-		Wave->NumChannels = QualityInfo.NumChannels;
-		Wave->RawPCMDataSize = QualityInfo.SampleDataSize;
-		Wave->Duration = QualityInfo.Duration;
-
 		// Clear out any dangling pointers
 		Buffer->PCM.PCMData = NULL;
 		Buffer->PCM.PCMDataSize = 0;
@@ -375,7 +377,7 @@ FXAudio2SoundBuffer* FXAudio2SoundBuffer::CreateQueuedBuffer( FXAudio2Device* XA
 		Wave->RemoveAudioResource();
 	}
 
-	return( Buffer );
+	return Buffer;
 }
 
 /**
diff --git a/Engine/Source/Runtime/Windows/XAudio2/Private/XAudio2Source.cpp b/Engine/Source/Runtime/Windows/XAudio2/Private/XAudio2Source.cpp
--- a/Engine/Source/Runtime/Windows/XAudio2/Private/XAudio2Source.cpp
+++ b/Engine/Source/Runtime/Windows/XAudio2/Private/XAudio2Source.cpp
@@ -36,6 +36,7 @@ FXMPHelper* FXMPHelper::GetXMPHelper( void )
 FXAudio2SoundSource::FXAudio2SoundSource(FAudioDevice* InAudioDevice)
 :	FSoundSource( InAudioDevice ),
 	Source( NULL ),
+	RealtimeAsyncTask( nullptr ),
 	CurrentBuffer( 0 ),
 	bBuffersToFlush( false ),
 	bLoopCallback( false ),
@@ -81,11 +82,18 @@ void FXAudio2SoundSource::InitializeSourceEffects(uint32 InVoiceId)
  */
 void FXAudio2SoundSource::FreeResources( void )
 {
+	if (RealtimeAsyncTask)
+	{
+		RealtimeAsyncTask->EnsureCompletion();
+		delete RealtimeAsyncTask;
+		RealtimeAsyncTask = nullptr;
+	}
+
 	// Release voice.
 	if( Source )
 	{
 		Source->DestroyVoice();
-		Source = NULL;
+		Source = nullptr;
 	}
 
 	// If we're a streaming buffer...
@@ -100,7 +108,7 @@ void FXAudio2SoundSource::FreeResources( void )
 		{
 			check( Buffer->ResourceID == 0 );
 			delete Buffer;
-			Buffer = NULL;
+			Buffer = nullptr;
 		}
 
 		CurrentBuffer = 0;
@@ -138,28 +146,39 @@ void FXAudio2SoundSource::SubmitPCMBuffers( void )
 	}
 }
 
-
-bool FXAudio2SoundSource::ReadProceduralData( const int32 BufferIndex )
-{
-	const int32 MaxSamples = ( MONO_PCM_BUFFER_SIZE * Buffer->NumChannels ) / sizeof( int16 );
-	const int32 BytesWritten = WaveInstance->WaveData->GeneratePCMData( (uint8*)XAudio2Buffers[BufferIndex].pAudioData, MaxSamples );
-
-	XAudio2Buffers[BufferIndex].AudioBytes = BytesWritten;
-
-	// convenience return value: we're never actually "looping" here.
-	return false;  
-}
-
-bool FXAudio2SoundSource::ReadMorePCMData( const int32 BufferIndex )
+bool FXAudio2SoundSource::ReadMorePCMData( const int32 BufferIndex, EDataReadMode DataReadMode )
 {
-	USoundWave *WaveData = WaveInstance->WaveData;
+	USoundWave* WaveData = WaveInstance->WaveData;
 	if( WaveData && WaveData->bProcedural )
 	{
-		return ReadProceduralData( BufferIndex );
+		const int32 MaxSamples = ( MONO_PCM_BUFFER_SIZE * Buffer->NumChannels ) / sizeof( int16 );
+
+		if (DataReadMode == EDataReadMode::Synchronous || WaveData->bCanProcessAsync == false)
+		{
+			const int32 BytesWritten = WaveData->GeneratePCMData( (uint8*)XAudio2Buffers[BufferIndex].pAudioData, MaxSamples );
+			XAudio2Buffers[BufferIndex].AudioBytes = BytesWritten;
+		}
+		else
+		{
+			RealtimeAsyncTask = new FAsyncRealtimeAudioTask(WaveData, (uint8*)XAudio2Buffers[BufferIndex].pAudioData, MaxSamples);
+			RealtimeAsyncTask->StartBackgroundTask();
+		}
+
+		// we're never actually "looping" here.
+		return false;
 	}
 	else
 	{
-		return XAudio2Buffer->ReadCompressedData( ( uint8* )XAudio2Buffers[BufferIndex].pAudioData, WaveInstance->LoopingMode != LOOP_Never );
+		if (DataReadMode == EDataReadMode::Synchronous)
+		{
+			return XAudio2Buffer->ReadCompressedData( ( uint8* )XAudio2Buffers[BufferIndex].pAudioData, WaveInstance->LoopingMode != LOOP_Never );
+		}
+		else
+		{
+			RealtimeAsyncTask = new FAsyncRealtimeAudioTask(XAudio2Buffer, ( uint8* )XAudio2Buffers[BufferIndex].pAudioData, WaveInstance->LoopingMode != LOOP_Never, DataReadMode == EDataReadMode::AsynchronousSkipFirstFrame);
+			RealtimeAsyncTask->StartBackgroundTask();
+			return false;
+		}
 	}
 }
 
@@ -175,28 +194,45 @@ void FXAudio2SoundSource::SubmitPCMRTBuffers( void )
 	// Set the buffer to be in real time mode
 	CurrentBuffer = 0;
 
+	const uint32 BufferSize = MONO_PCM_BUFFER_SIZE * Buffer->NumChannels;
+	
 	// Set up double buffer area to decompress to
-	XAudio2Buffers[0].pAudioData = ( uint8* )FMemory::Malloc( MONO_PCM_BUFFER_SIZE * Buffer->NumChannels );
-	XAudio2Buffers[0].AudioBytes = MONO_PCM_BUFFER_SIZE * Buffer->NumChannels;
+	XAudio2Buffers[0].pAudioData = (uint8*)FMemory::Malloc(BufferSize);
+	XAudio2Buffers[0].AudioBytes = BufferSize;
 
-	ReadMorePCMData(0);
+	// Only use the cached data if we're starting from the beginning, otherwise we'll have to take a synchronous hit
+	bool bSkipFirstBuffer = false;;
+	if (WaveInstance->WaveData && WaveInstance->WaveData->CachedRealtimeFirstBuffer && WaveInstance->StartTime == 0.f)
+	{
+		FMemory::Memcpy((uint8*)XAudio2Buffers[0].pAudioData, WaveInstance->WaveData->CachedRealtimeFirstBuffer, BufferSize);
+		bSkipFirstBuffer = true;
+	}
+	else
+	{
+		ReadMorePCMData(0, EDataReadMode::Synchronous);
+	}
 
 	AudioDevice->ValidateAPICall( TEXT( "SubmitSourceBuffer - PCMRT" ), 
 		Source->SubmitSourceBuffer( XAudio2Buffers ) );
 
-	XAudio2Buffers[1].pAudioData = ( uint8* )FMemory::Malloc( MONO_PCM_BUFFER_SIZE * Buffer->NumChannels );
-	XAudio2Buffers[1].AudioBytes = MONO_PCM_BUFFER_SIZE * Buffer->NumChannels;
+	XAudio2Buffers[1].pAudioData = (uint8*)FMemory::Malloc(BufferSize);
+	XAudio2Buffers[1].AudioBytes = BufferSize;
 
-	ReadMorePCMData(1);
+	CurrentBuffer = 1;
 
-	if (XAudio2Buffers[1].AudioBytes > 0)
+	// Start the async population of the next buffer
+	EDataReadMode DataReadMode = EDataReadMode::Asynchronous;
+	if (XAudio2Buffer->SoundFormat == ESoundFormat::SoundFormat_Streaming)
 	{
-		CurrentBuffer = 1;
-
-		AudioDevice->ValidateAPICall( TEXT( "SubmitSourceBuffer - PCMRT" ), 
-			Source->SubmitSourceBuffer( XAudio2Buffers + 1 ) );
+		DataReadMode =  EDataReadMode::Synchronous;
+	}
+	else if (bSkipFirstBuffer)
+	{
+		DataReadMode =  EDataReadMode::AsynchronousSkipFirstFrame;
 	}
 
+	ReadMorePCMData(1, DataReadMode);
+
 	bResourcesNeedFreeing = true;
 }
 
@@ -1175,33 +1211,66 @@ void FXAudio2SoundSource::Pause( void )
 /**
  * Handles feeding new data to a real time decompressed sound
  */
-void FXAudio2SoundSource::HandleRealTimeSource( void )
+void FXAudio2SoundSource::HandleRealTimeSource()
 {
-	// Update the double buffer toggle
-	++CurrentBuffer;
-
-	// Get the next bit of streaming data
-	const bool bLooped = ReadMorePCMData(CurrentBuffer & 1);
-
-	// Have we reached the end of the compressed sound?
-	if( bLooped )
+	bool bLooped = false;
+	if (RealtimeAsyncTask && RealtimeAsyncTask->IsDone())
 	{
-		switch( WaveInstance->LoopingMode )
+		switch(RealtimeAsyncTask->GetTask().GetTaskType())
 		{
-		case LOOP_Never:
-			// Play out any queued buffers - once there are no buffers left, the state check at the beginning of IsFinished will fire
-			bBuffersToFlush = true;
-			XAudio2Buffers[CurrentBuffer & 1].Flags |= XAUDIO2_END_OF_STREAM;
+		case ERealtimeAudioTaskType::Decompress:
+			bLooped = RealtimeAsyncTask->GetTask().GetBufferLooped();
 			break;
 
-		case LOOP_WithNotification:
-			// If we have just looped, and we are programmatically looping, send notification
-			WaveInstance->NotifyFinished();
+		case ERealtimeAudioTaskType::Procedural:
+			XAudio2Buffers[CurrentBuffer & 1].AudioBytes = RealtimeAsyncTask->GetTask().GetBytesWritten();
 			break;
+		}
 
-		case LOOP_Forever:
-			// Let the sound loop indefinitely
-			break;
+		delete RealtimeAsyncTask;
+		RealtimeAsyncTask = nullptr;
+	}
+	else
+	{
+		// Update the double buffer toggle
+		++CurrentBuffer;
+
+		// Get the next bit of streaming data
+		bLooped = ReadMorePCMData(CurrentBuffer & 1, (XAudio2Buffer->SoundFormat == ESoundFormat::SoundFormat_Streaming ? EDataReadMode::Synchronous : EDataReadMode::Asynchronous));
+	}
+
+	if (RealtimeAsyncTask == nullptr)
+	{
+		// Have we reached the end of the compressed sound?
+		if( bLooped )
+		{
+			switch( WaveInstance->LoopingMode )
+			{
+			case LOOP_Never:
+				// Play out any queued buffers - once there are no buffers left, the state check at the beginning of IsFinished will fire
+				bBuffersToFlush = true;
+				XAudio2Buffers[CurrentBuffer & 1].Flags |= XAUDIO2_END_OF_STREAM;
+				break;
+
+			case LOOP_WithNotification:
+				// If we have just looped, and we are programmatically looping, send notification
+				WaveInstance->NotifyFinished();
+				break;
+
+			case LOOP_Forever:
+				// Let the sound loop indefinitely
+				break;
+			}
+		}
+
+		if (XAudio2Buffers[CurrentBuffer & 1].AudioBytes > 0)
+		{
+			AudioDevice->ValidateAPICall( TEXT( "SubmitSourceBuffer - IsFinished" ), 
+				Source->SubmitSourceBuffer( &XAudio2Buffers[CurrentBuffer & 1] ) );
+		}
+		else
+		{
+			--CurrentBuffer;
 		}
 	}
 }
@@ -1229,39 +1298,19 @@ bool FXAudio2SoundSource::IsFinished( void )
 		const bool bIsRealTimeSource = XAudio2Buffer->SoundFormat == SoundFormat_PCMRT || XAudio2Buffer->SoundFormat == SoundFormat_Streaming;
 
 		// If we have no queued buffers, we're either at the end of a sound, or starved
-		if( SourceState.BuffersQueued == 0 )
+		// and we are expecting the sound to be finishing
+		if (SourceState.BuffersQueued == 0 && (bBuffersToFlush || !bIsRealTimeSource))
 		{
-			// ... are we expecting the sound to be finishing?
-			if( bBuffersToFlush || !bIsRealTimeSource )
-			{
-				// ... notify the wave instance that it has finished playing.
-				WaveInstance->NotifyFinished();
-				return( true );
-			}
+			// ... notify the wave instance that it has finished playing.
+			WaveInstance->NotifyFinished();
+			return true;
 		}
 
 		// Service any real time sounds
-		if( bIsRealTimeSource )
+		if (bIsRealTimeSource && !bBuffersToFlush && SourceState.BuffersQueued <= 1)
 		{
-			if( SourceState.BuffersQueued <= 1 )
-			{
-				// Continue feeding new sound data (unless we are waiting for the sound to finish)
-				if( !bBuffersToFlush )
-				{
-					HandleRealTimeSource();
-
-					if (XAudio2Buffers[CurrentBuffer & 1].AudioBytes > 0)
-					{
-						AudioDevice->ValidateAPICall( TEXT( "SubmitSourceBuffer - IsFinished" ), 
-							Source->SubmitSourceBuffer( &XAudio2Buffers[CurrentBuffer & 1] ) );
-					}
-					else
-					{
-						--CurrentBuffer;
-					}
-
-				}
-			}
+			// Continue feeding new sound data (unless we are waiting for the sound to finish)
+			HandleRealTimeSource();
 
 			return( false );
 		}
diff --git a/Engine/Source/Runtime/Windows/XAudio2/Private/XAudio2Support.h b/Engine/Source/Runtime/Windows/XAudio2/Private/XAudio2Support.h
--- a/Engine/Source/Runtime/Windows/XAudio2/Private/XAudio2Support.h
+++ b/Engine/Source/Runtime/Windows/XAudio2/Private/XAudio2Support.h
@@ -260,6 +260,8 @@ public:
 	friend class FXAudio2SoundSource;
 };
 
+typedef FAsyncTask<class FAsyncRealtimeAudioTaskWorker<FXAudio2SoundBuffer>> FAsyncRealtimeAudioTask;
+
 /**
  * XAudio2 implementation of FSoundSource, the interface used to play, stop and update sources
  */
@@ -320,7 +322,7 @@ public:
 	/**
 	 * Handles feeding new data to a real time decompressed sound
 	 */
-	void HandleRealTimeSource( void );
+	void HandleRealTimeSource();
 
 	/**
 	 * Queries the status of the currently associated wave instance.
@@ -389,11 +391,16 @@ public:
 	void RouteToRadio( float ChannelVolumes[CHANNELOUT_COUNT] );
 
 protected:
-	/** Decompress through XAudio2Buffer, or call USoundWave procedure to generate more PCM data. Returns true/false: did audio loop? */
-	bool ReadMorePCMData(const int32 BufferIndex);
 
-	/** Handle obtaining more data for procedural USoundWaves. Always returns false for convenience. */
-	bool ReadProceduralData(const int32 BufferIndex);
+	enum class EDataReadMode : uint8
+	{
+		Synchronous,
+		Asynchronous,
+		AsynchronousSkipFirstFrame
+	};
+
+	/** Decompress through XAudio2Buffer, or call USoundWave procedure to generate more PCM data. Returns true/false: did audio loop? */
+	bool ReadMorePCMData(const int32 BufferIndex, EDataReadMode DataReadMode);
 
 	/** Returns if the source is using the default 3d spatialization. */
 	bool IsUsingDefaultSpatializer();
@@ -439,6 +446,9 @@ protected:
 	IXAudio2SourceVoice*		Source;
 	/** Structure to handle looping sound callbacks */
 	FXAudio2SoundSourceCallback	SourceCallback;
+
+	/** Asynchronous task for real time audio sources */
+	FAsyncRealtimeAudioTask* RealtimeAsyncTask;
 	/** Destination voices */
 	XAUDIO2_SEND_DESCRIPTOR		Destinations[DEST_COUNT];
 	/** Which sound buffer should be written to next - used for double buffering. */
