Fix crash in async realtime audio decompression

diff --git a/Engine/Source/Runtime/Mac/CoreAudio/Private/CoreAudioSource.cpp b/Engine/Source/Runtime/Mac/CoreAudio/Private/CoreAudioSource.cpp
--- a/Engine/Source/Runtime/Mac/CoreAudio/Private/CoreAudioSource.cpp
+++ b/Engine/Source/Runtime/Mac/CoreAudio/Private/CoreAudioSource.cpp
@@ -508,27 +508,30 @@ void FCoreAudioSoundSource::HandleRealTimeSource()
 {
 	bool bLooped = false;
 	const int32 BufferIndex = 1 - BufferInUse;
-	if (RealtimeAsyncTask && RealtimeAsyncTask->IsDone())
+	if (RealtimeAsyncTask)
 	{
-		switch(RealtimeAsyncTask->GetTask().GetTaskType())
+		if (RealtimeAsyncTask->IsDone())
 		{
-		case ERealtimeAudioTaskType::Decompress:
-			bLooped = RealtimeAsyncTask->GetTask().GetBufferLooped();
-			++NumActiveBuffers;
-			break;
-
-		case ERealtimeAudioTaskType::Procedural:
-			const int32 BytesWritten = RealtimeAsyncTask->GetTask().GetBytesWritten();
-			CoreAudioBuffers[BufferIndex].AudioDataSize = BytesWritten;
-			if (BytesWritten > 0)
+			switch(RealtimeAsyncTask->GetTask().GetTaskType())
 			{
+			case ERealtimeAudioTaskType::Decompress:
+				bLooped = RealtimeAsyncTask->GetTask().GetBufferLooped();
 				++NumActiveBuffers;
+				break;
+
+			case ERealtimeAudioTaskType::Procedural:
+				const int32 BytesWritten = RealtimeAsyncTask->GetTask().GetBytesWritten();
+				CoreAudioBuffers[BufferIndex].AudioDataSize = BytesWritten;
+				if (BytesWritten > 0)
+				{
+					++NumActiveBuffers;
+				}
+				break;
 			}
-			break;
-		}
 
-		delete RealtimeAsyncTask;
-		RealtimeAsyncTask = nullptr;
+			delete RealtimeAsyncTask;
+			RealtimeAsyncTask = nullptr;
+		}
 	}
 	else
 	{
diff --git a/Engine/Source/Runtime/Windows/XAudio2/Private/XAudio2Source.cpp b/Engine/Source/Runtime/Windows/XAudio2/Private/XAudio2Source.cpp
--- a/Engine/Source/Runtime/Windows/XAudio2/Private/XAudio2Source.cpp
+++ b/Engine/Source/Runtime/Windows/XAudio2/Private/XAudio2Source.cpp
@@ -108,7 +108,7 @@ void FXAudio2SoundSource::FreeResources( void )
 		{
 			check( Buffer->ResourceID == 0 );
 			delete Buffer;
-			Buffer = nullptr;
+			Buffer = XAudio2Buffer = nullptr;
 		}
 
 		CurrentBuffer = 0;
@@ -1175,14 +1175,14 @@ void FXAudio2SoundSource::Stop( void )
 		{
 			AudioDevice->ValidateAPICall( TEXT( "Stop" ), 
 				Source->Stop( XAUDIO2_PLAY_TAILS ) );
-
-			// Free resources
-			FreeResources();
 		}
 
+		// Free resources
+		FreeResources();
+
 		Paused = false;
 		Playing = false;
-		Buffer = NULL;
+		Buffer = XAudio2Buffer = nullptr;
 		bBuffersToFlush = false;
 		bLoopCallback = false;
 		bResourcesNeedFreeing = false;
@@ -1214,21 +1214,24 @@ void FXAudio2SoundSource::Pause( void )
 void FXAudio2SoundSource::HandleRealTimeSource()
 {
 	bool bLooped = false;
-	if (RealtimeAsyncTask && RealtimeAsyncTask->IsDone())
+	if (RealtimeAsyncTask)
 	{
-		switch(RealtimeAsyncTask->GetTask().GetTaskType())
+		if (RealtimeAsyncTask->IsDone())
 		{
-		case ERealtimeAudioTaskType::Decompress:
-			bLooped = RealtimeAsyncTask->GetTask().GetBufferLooped();
-			break;
+			switch(RealtimeAsyncTask->GetTask().GetTaskType())
+			{
+			case ERealtimeAudioTaskType::Decompress:
+				bLooped = RealtimeAsyncTask->GetTask().GetBufferLooped();
+				break;
 
-		case ERealtimeAudioTaskType::Procedural:
-			XAudio2Buffers[CurrentBuffer & 1].AudioBytes = RealtimeAsyncTask->GetTask().GetBytesWritten();
-			break;
-		}
+			case ERealtimeAudioTaskType::Procedural:
+				XAudio2Buffers[CurrentBuffer & 1].AudioBytes = RealtimeAsyncTask->GetTask().GetBytesWritten();
+				break;
+			}
 
-		delete RealtimeAsyncTask;
-		RealtimeAsyncTask = nullptr;
+			delete RealtimeAsyncTask;
+			RealtimeAsyncTask = nullptr;
+		}
 	}
 	else
 	{
