UE4 - Get more info on D3D OOM as it gets eaten by the crash reporter

diff --git a/Engine/Source/Runtime/Windows/D3D11RHI/Private/D3D11Util.cpp b/Engine/Source/Runtime/Windows/D3D11RHI/Private/D3D11Util.cpp
--- a/Engine/Source/Runtime/Windows/D3D11RHI/Private/D3D11Util.cpp
+++ b/Engine/Source/Runtime/Windows/D3D11RHI/Private/D3D11Util.cpp
@@ -196,7 +196,19 @@ void VerifyD3D11CreateTextureResult(HRESULT D3DResult,const ANSICHAR* Code,const
 	const FString ErrorString = GetD3D11ErrorString(D3DResult, 0);
 	const TCHAR* D3DFormatString = GetD3D11TextureFormatString((DXGI_FORMAT)Format);
 
-	UE_LOG(LogD3D11RHI, Error,TEXT("%s failed \n at %s:%u \n with error %s"),ANSI_TO_TCHAR(Code),ANSI_TO_TCHAR(Filename),Line,*ErrorString);
+	UE_LOG(LogD3D11RHI, Error,
+		TEXT("%s failed \n at %s:%u \n with error %s, \n Size=%ix%ix%i Format=%s(0x%08X), NumMips=%i, Flags=%s"),
+		ANSI_TO_TCHAR(Code),
+		ANSI_TO_TCHAR(Filename),
+		Line,
+		*ErrorString,
+		SizeX,
+		SizeY,
+		SizeZ,
+		D3DFormatString,
+		Format,
+		NumMips,
+		*GetD3D11TextureFlagString(Flags));
 
 	TerminateOnDeviceRemoved(D3DResult);
 
