UE4 - Terminate on out of video memory instead of crashing and report to the user. (UE-16752)

diff --git a/Engine/Source/Runtime/Windows/D3D11RHI/Private/D3D11Util.cpp b/Engine/Source/Runtime/Windows/D3D11RHI/Private/D3D11Util.cpp
--- a/Engine/Source/Runtime/Windows/D3D11RHI/Private/D3D11Util.cpp
+++ b/Engine/Source/Runtime/Windows/D3D11RHI/Private/D3D11Util.cpp
@@ -161,6 +161,26 @@ static void TerminateOnDeviceRemoved(HRESULT D3DResult)
 	}
 }
 
+static void TerminateOnOutOfMemory(HRESULT D3DResult, bool bCreatingTextures)
+{
+	if (D3DResult == E_OUTOFMEMORY)
+	{
+		if (bCreatingTextures)
+		{
+			FPlatformMisc::MessageBoxExt(EAppMsgType::Ok, *LOCTEXT("OutOfVideoMemoryTextures", "Out of video memory trying to allocate a texture! Make sure your video card has the minimum required memory, try lowering the resolution and/or closing other applications that are running. Exiting...").ToString(), TEXT("Error"));
+		}
+		else
+		{
+			FPlatformMisc::MessageBoxExt(EAppMsgType::Ok, *NSLOCTEXT("D3D11RHI", "OutOfMemory", "Out of video memory trying to allocate a rendering resource. Make sure your video card has the minimum required memory, try lowering the resolution and/or closing other applications that are running. Exiting...").ToString(), TEXT("Error"));
+		}
+#if STATS
+		GetRendererModule().DebugLogOnCrash();
+#endif
+		FPlatformMisc::RequestExit(true);
+	}
+}
+
+
 #ifndef MAKE_D3DHRESULT
 	#define _FACD3D						0x876
 	#define MAKE_D3DHRESULT( code)		MAKE_HRESULT( 1, _FACD3D, code )
@@ -175,16 +195,7 @@ void VerifyD3D11Result(HRESULT D3DResult,const ANSICHAR* Code,const ANSICHAR* Fi
 	UE_LOG(LogD3D11RHI, Error,TEXT("%s failed \n at %s:%u \n with error %s"),ANSI_TO_TCHAR(Code),ANSI_TO_TCHAR(Filename),Line,*ErrorString);
 
 	TerminateOnDeviceRemoved(D3DResult);
-
-	// this is to track down a rarely happening crash
-	if(D3DResult == E_OUTOFMEMORY)
-	{
-		FPlatformMisc::MessageBoxExt(EAppMsgType::Ok, *NSLOCTEXT("D3D11RHI", "OutOfMemory", "Out of video memory trying to allocate a rendering resource.").ToString(), TEXT("Error"));
-#if STATS
-		GetRendererModule().DebugLogOnCrash();
-#endif
-		FPlatformMisc::RequestExit(/*bForce=*/ true);
-	}
+	TerminateOnOutOfMemory(D3DResult, false);
 
 	UE_LOG(LogD3D11RHI, Fatal,TEXT("%s failed \n at %s:%u \n with error %s"),ANSI_TO_TCHAR(Code),ANSI_TO_TCHAR(Filename),Line,*ErrorString);
 }
@@ -211,14 +222,7 @@ void VerifyD3D11CreateTextureResult(HRESULT D3DResult,const ANSICHAR* Code,const
 		*GetD3D11TextureFlagString(Flags));
 
 	TerminateOnDeviceRemoved(D3DResult);
-
-	// this is to track down a rarely happening crash
-	if (D3DResult == E_OUTOFMEMORY)
-	{
-#if STATS
-		GetRendererModule().DebugLogOnCrash();
-#endif // STATS
-	}
+	TerminateOnOutOfMemory(D3DResult, true);
 
 	UE_LOG(LogD3D11RHI, Fatal,
 		TEXT("%s failed \n at %s:%u \n with error %s, \n Size=%ix%ix%i Format=%s(0x%08X), NumMips=%i, Flags=%s"),
