Use a triple buffer XAudio/Core for realtime audio decompression to ensure that if async task is getting starved audio does not hitch; if final buffer is active then ensure completion on the main thread of the pending async task so starvation will not occur.
Properly initialize audio resource when playing a native sound with a specified start time.

diff --git a/Engine/Source/Runtime/Android/AndroidAudio/Private/AndroidAudioBuffer.cpp b/Engine/Source/Runtime/Android/AndroidAudio/Private/AndroidAudioBuffer.cpp
--- a/Engine/Source/Runtime/Android/AndroidAudio/Private/AndroidAudioBuffer.cpp
+++ b/Engine/Source/Runtime/Android/AndroidAudio/Private/AndroidAudioBuffer.cpp
@@ -59,6 +59,12 @@ FSLESSoundBuffer* FSLESSoundBuffer::CreateQueuedBuffer( FSLESAudioDevice* AudioD
 	
 	Buffer->DecompressionState = AudioDevice->CreateCompressedAudioInfo(InWave);
 
+	// If the buffer was precached as native, the resource data will have been lost and we need to re-initialize it
+	if (InWave->ResourceData == nullptr)
+	{
+		InWave->InitAudioResource(AudioDevice->GetRuntimeFormat(InWave));
+	}
+
 	if (Buffer->DecompressionState->ReadCompressedInfo( InWave->ResourceData, InWave->ResourceSize, &QualityInfo ))
 	{	
 		// Clear out any dangling pointers
diff --git a/Engine/Source/Runtime/Android/AndroidAudio/Private/AndroidAudioSource.cpp b/Engine/Source/Runtime/Android/AndroidAudio/Private/AndroidAudioSource.cpp
--- a/Engine/Source/Runtime/Android/AndroidAudio/Private/AndroidAudioSource.cpp
+++ b/Engine/Source/Runtime/Android/AndroidAudio/Private/AndroidAudioSource.cpp
@@ -214,20 +214,17 @@ bool FSLESSoundSource::EnqueuePCMRTBuffer( bool bLoop )
 	AudioBuffers[1].AudioDataSize = BufferSize;
 
 	// Only use the cached data if we're starting from the beginning, otherwise we'll have to take a synchronous hit
-	bool bSkipFirstBuffer = false;;
 	if (WaveInstance->WaveData && WaveInstance->WaveData->CachedRealtimeFirstBuffer && WaveInstance->StartTime == 0.f)
 	{
 		FMemory::Memcpy((uint8*)AudioBuffers[0].AudioData, WaveInstance->WaveData->CachedRealtimeFirstBuffer, BufferSize);
-		bSkipFirstBuffer = true;
+		FMemory::Memcpy((uint8*)AudioBuffers[1].AudioData, WaveInstance->WaveData->CachedRealtimeFirstBuffer, BufferSize);
 	}
 	else
 	{
 		ReadMorePCMData(0, EDataReadMode::Synchronous);
+		ReadMorePCMData(1, EDataReadMode::Asynchronous);
 	}
 
-	// Start the async population of the next buffer
-	ReadMorePCMData(1, (bSkipFirstBuffer ? EDataReadMode::AsynchronousSkipFirstFrame : EDataReadMode::Asynchronous));
-
 	SLresult result;
 
 	// callback is used to submit and decompress next buffer
diff --git a/Engine/Source/Runtime/Engine/Private/AudioDecompress.cpp b/Engine/Source/Runtime/Engine/Private/AudioDecompress.cpp
--- a/Engine/Source/Runtime/Engine/Private/AudioDecompress.cpp
+++ b/Engine/Source/Runtime/Engine/Private/AudioDecompress.cpp
@@ -325,8 +325,8 @@ void FAsyncAudioDecompressWorker::DoWork()
 			if (Wave->DecompressionType == DTYPE_RealTime)
 			{
 				const uint32 PCMBufferSize = MONO_PCM_BUFFER_SIZE * Wave->NumChannels;
-				Wave->CachedRealtimeFirstBuffer = (uint8*)FMemory::Malloc(PCMBufferSize);
-				AudioInfo->ReadCompressedData( Wave->CachedRealtimeFirstBuffer, false, PCMBufferSize );
+				Wave->CachedRealtimeFirstBuffer = (uint8*)FMemory::Malloc(PCMBufferSize * 2);
+				AudioInfo->ReadCompressedData( Wave->CachedRealtimeFirstBuffer, false, PCMBufferSize * 2 );
 			}
 			else
 			{
diff --git a/Engine/Source/Runtime/Engine/Private/AudioDevice.cpp b/Engine/Source/Runtime/Engine/Private/AudioDevice.cpp
--- a/Engine/Source/Runtime/Engine/Private/AudioDevice.cpp
+++ b/Engine/Source/Runtime/Engine/Private/AudioDevice.cpp
@@ -2365,7 +2365,7 @@ void FAudioDevice::Precache(USoundWave* SoundWave, bool bSynchronous, bool bTrac
 			{
 				if (SoundWave->DecompressionType == DTYPE_RealTime)
 				{
-					SoundWave->CachedRealtimeFirstBuffer = (uint8*)FMemory::Malloc( MONO_PCM_BUFFER_SIZE * SoundWave->NumChannels );
+					SoundWave->CachedRealtimeFirstBuffer = (uint8*)FMemory::Malloc( MONO_PCM_BUFFER_SIZE * SoundWave->NumChannels * 2 );
 				}
 				SoundWave->AudioDecompressor = new FAsyncAudioDecompress(SoundWave);
 				SoundWave->AudioDecompressor->StartBackgroundTask();
diff --git a/Engine/Source/Runtime/Engine/Public/AudioDecompress.h b/Engine/Source/Runtime/Engine/Public/AudioDecompress.h
--- a/Engine/Source/Runtime/Engine/Public/AudioDecompress.h
+++ b/Engine/Source/Runtime/Engine/Public/AudioDecompress.h
@@ -195,6 +195,8 @@ public:
 		case ERealtimeAudioTaskType::Decompress:
 			if (bSkipFirstBuffer)
 			{
+				// If we're using cached data we need to skip the first two reads from the data
+				AudioBuffer->ReadCompressedData( ( uint8* )AudioData, bLoopingMode );
 				AudioBuffer->ReadCompressedData( ( uint8* )AudioData, bLoopingMode );
 			}
 			bLooped = AudioBuffer->ReadCompressedData( ( uint8* )AudioData, bLoopingMode );
diff --git a/Engine/Source/Runtime/Mac/CoreAudio/Private/CoreAudioBuffer.cpp b/Engine/Source/Runtime/Mac/CoreAudio/Private/CoreAudioBuffer.cpp
--- a/Engine/Source/Runtime/Mac/CoreAudio/Private/CoreAudioBuffer.cpp
+++ b/Engine/Source/Runtime/Mac/CoreAudio/Private/CoreAudioBuffer.cpp
@@ -155,7 +155,13 @@ FCoreAudioSoundBuffer* FCoreAudioSoundBuffer::CreateQueuedBuffer( FCoreAudioDevi
 	FSoundQualityInfo QualityInfo = { 0 };
 
 	Buffer->DecompressionState = CoreAudioDevice->CreateCompressedAudioInfo(Wave);
-	
+
+	// If the buffer was precached as native, the resource data will have been lost and we need to re-initialize it
+	if (Wave->ResourceData == nullptr)
+	{
+		Wave->InitAudioResource(CoreAudioDevice->GetRuntimeFormat(Wave));
+	}
+
 	if (Buffer->DecompressionState->ReadCompressedInfo( Wave->ResourceData, Wave->ResourceSize, &QualityInfo ))
 	{
 		// Clear out any dangling pointers
diff --git a/Engine/Source/Runtime/Mac/CoreAudio/Private/CoreAudioSource.cpp b/Engine/Source/Runtime/Mac/CoreAudio/Private/CoreAudioSource.cpp
--- a/Engine/Source/Runtime/Mac/CoreAudio/Private/CoreAudioSource.cpp
+++ b/Engine/Source/Runtime/Mac/CoreAudio/Private/CoreAudioSource.cpp
@@ -108,6 +108,7 @@ void FCoreAudioSoundSource::FreeResources( void )
 		// ... free the buffers
 		FMemory::Free( ( void* )CoreAudioBuffers[0].AudioData );
 		FMemory::Free( ( void* )CoreAudioBuffers[1].AudioData );
+		FMemory::Free( ( void* )CoreAudioBuffers[2].AudioData );
 		
 		// Buffers without a valid resource ID are transient and need to be deleted.
 		if( Buffer )
@@ -190,7 +191,7 @@ void FCoreAudioSoundSource::SubmitPCMRTBuffers( void )
 {
 	SCOPE_CYCLE_COUNTER( STAT_AudioSubmitBuffersTime );
 
-	FMemory::Memzero( CoreAudioBuffers, sizeof( CoreAudioBuffer ) * 2 );
+	FMemory::Memzero( CoreAudioBuffers, sizeof( CoreAudioBuffer ) * 3 );
 
 	bStreamedSound = true;
 
@@ -203,6 +204,9 @@ void FCoreAudioSoundSource::SubmitPCMRTBuffers( void )
 	CoreAudioBuffers[1].AudioData = (uint8*)FMemory::Malloc(BufferSize);
 	CoreAudioBuffers[1].AudioDataSize = BufferSize;
 
+	CoreAudioBuffers[2].AudioData = (uint8*)FMemory::Malloc(BufferSize);
+	CoreAudioBuffers[2].AudioDataSize = BufferSize;
+
 	NumActiveBuffers = 0;
 	BufferInUse = 0;
 	
@@ -211,16 +215,18 @@ void FCoreAudioSoundSource::SubmitPCMRTBuffers( void )
 	if (WaveInstance->WaveData && WaveInstance->WaveData->CachedRealtimeFirstBuffer && WaveInstance->StartTime == 0.f)
 	{
 		FMemory::Memcpy((uint8*)CoreAudioBuffers[0].AudioData, WaveInstance->WaveData->CachedRealtimeFirstBuffer, BufferSize);
+		FMemory::Memcpy((uint8*)CoreAudioBuffers[1].AudioData, WaveInstance->WaveData->CachedRealtimeFirstBuffer + BufferSize, BufferSize);
 		bSkipFirstBuffer = true;
-		NumActiveBuffers = 1;
+		NumActiveBuffers = 2;
 	}
 	else
 	{
 		ReadMorePCMData(0, EDataReadMode::Synchronous);
+		ReadMorePCMData(1, EDataReadMode::Synchronous);
 	}
 
 	// Start the async population of the next buffer
-	ReadMorePCMData(1, (bSkipFirstBuffer ? EDataReadMode::AsynchronousSkipFirstFrame : EDataReadMode::Asynchronous));
+	ReadMorePCMData(2, (bSkipFirstBuffer ? EDataReadMode::AsynchronousSkipFirstFrame : EDataReadMode::Asynchronous));
 }
 
 /**
@@ -501,17 +507,49 @@ void FCoreAudioSoundSource::Pause( void )
 	}
 }
 
+void FCoreAudioSoundSource::HandleRealTimeSourceData(bool bLooped)
+{
+	// Have we reached the end of the compressed sound?
+	if( bLooped )
+	{
+		switch( WaveInstance->LoopingMode )
+		{
+			case LOOP_Never:
+				// Play out any queued buffers - once there are no buffers left, the state check at the beginning of IsFinished will fire
+				bBuffersToFlush = true;
+				break;
+				
+			case LOOP_WithNotification:
+				// If we have just looped, and we are programmatically looping, send notification
+				WaveInstance->NotifyFinished();
+				break;
+				
+			case LOOP_Forever:
+				// Let the sound loop indefinitely
+				break;
+		}
+	}
+}
+
 /**
  * Handles feeding new data to a real time decompressed sound
  */
-void FCoreAudioSoundSource::HandleRealTimeSource()
+void FCoreAudioSoundSource::HandleRealTimeSource(bool bBlockForData)
 {
-	bool bLooped = false;
-	const int32 BufferIndex = 1 - BufferInUse;
+	const bool bGetMoreData = bBlockForData || (RealtimeAsyncTask == nullptr);
+	int32 BufferIndex = (BufferInUse + NumActiveBuffers) % 3;
 	if (RealtimeAsyncTask)
 	{
-		if (RealtimeAsyncTask->IsDone())
+		const bool bTaskDone = RealtimeAsyncTask->IsDone();
+		if (bTaskDone || bBlockForData)
 		{
+			bool bLooped = false;
+
+			if (!bTaskDone)
+			{
+				RealtimeAsyncTask->EnsureCompletion();
+			}
+
 			switch(RealtimeAsyncTask->GetTask().GetTaskType())
 			{
 			case ERealtimeAudioTaskType::Decompress:
@@ -531,35 +569,24 @@ void FCoreAudioSoundSource::HandleRealTimeSource()
 
 			delete RealtimeAsyncTask;
 			RealtimeAsyncTask = nullptr;
+
+			HandleRealTimeSourceData(bLooped);
+
+			if (++BufferIndex > 2)
+			{
+				BufferIndex = 0;
+			}
 		}
 	}
-	else
+
+	if (bGetMoreData)
 	{
 		// Get the next bit of streaming data
-		bLooped = ReadMorePCMData(BufferIndex, EDataReadMode::Asynchronous);
-	}
+		const bool bLooped = ReadMorePCMData(BufferIndex, EDataReadMode::Asynchronous);
 
-	if (RealtimeAsyncTask == nullptr)
-	{
-		// Have we reached the end of the compressed sound?
-		if( bLooped )
+		if (RealtimeAsyncTask == nullptr)
 		{
-			switch( WaveInstance->LoopingMode )
-			{
-				case LOOP_Never:
-					// Play out any queued buffers - once there are no buffers left, the state check at the beginning of IsFinished will fire
-					bBuffersToFlush = true;
-					break;
-				
-				case LOOP_WithNotification:
-					// If we have just looped, and we are programmatically looping, send notification
-					WaveInstance->NotifyFinished();
-					break;
-				
-				case LOOP_Forever:
-					// Let the sound loop indefinitely
-					break;
-			}
+			HandleRealTimeSourceData(bLooped);
 		}
 	}
 }
@@ -590,10 +617,10 @@ bool FCoreAudioSoundSource::IsFinished( void )
 		}
 
 		// Service any real time sounds
-		if( bStreamedSound && !bBuffersToFlush && NumActiveBuffers < 2)
+		if( bStreamedSound && !bBuffersToFlush && NumActiveBuffers < 3)
 		{
 			// Continue feeding new sound data (unless we are waiting for the sound to finish)
-			HandleRealTimeSource();
+			HandleRealTimeSource(NumActiveBuffers < 2);
 		}
 
 		return( false );
@@ -1086,7 +1113,10 @@ OSStatus FCoreAudioSoundSource::CoreAudioConvertCallback( AudioConverterRef Conv
 		if( Source->bStreamedSound )
 		{
 			Source->NumActiveBuffers--;
-			Source->BufferInUse = 1 - Source->BufferInUse;
+			if (++Source->BufferInUse > 2)
+			{
+				Source->BufferInUse = 0;
+			}
 		}
 		else if( Source->WaveInstance )
 		{
diff --git a/Engine/Source/Runtime/Mac/CoreAudio/Public/CoreAudioDevice.h b/Engine/Source/Runtime/Mac/CoreAudio/Public/CoreAudioDevice.h
--- a/Engine/Source/Runtime/Mac/CoreAudio/Public/CoreAudioDevice.h
+++ b/Engine/Source/Runtime/Mac/CoreAudio/Public/CoreAudioDevice.h
@@ -211,11 +211,6 @@ public:
 	 */
 	virtual void Pause( void );
 	
-	/**
-	 * Handles feeding new data to a real time decompressed sound
-	 */
-	void HandleRealTimeSource( void );
-	
 	/**
 	 * Queries the status of the currently associated wave instance.
 	 *
@@ -249,6 +244,16 @@ protected:
 		AsynchronousSkipFirstFrame
 	};
 
+	/**
+	 * Handles feeding new data to a real time decompressed sound
+	 */
+	void HandleRealTimeSourceData(bool bLooped);
+
+	/**
+	 * Handles feeding new data to a real time decompressed sound
+	 */
+	void HandleRealTimeSource(bool bBlockForData);
+
 	/** Decompress USoundWave procedure to generate more PCM data. Returns true/false: did audio loop? */
 	bool ReadMorePCMData( const int32 BufferIndex, EDataReadMode DataReadMode );
 
@@ -273,7 +278,7 @@ protected:
 	/** Which sound buffer should be written to next - used for double buffering. */
 	bool						bStreamedSound;
 	/** A pair of sound buffers to allow notification when a sound loops. */
-	CoreAudioBuffer				CoreAudioBuffers[2];
+	CoreAudioBuffer				CoreAudioBuffers[3];
 	/** Set when we wish to let the buffers play themselves out */
 	bool						bBuffersToFlush;
 
diff --git a/Engine/Source/Runtime/Windows/XAudio2/Private/XAudio2Buffer.cpp b/Engine/Source/Runtime/Windows/XAudio2/Private/XAudio2Buffer.cpp
--- a/Engine/Source/Runtime/Windows/XAudio2/Private/XAudio2Buffer.cpp
+++ b/Engine/Source/Runtime/Windows/XAudio2/Private/XAudio2Buffer.cpp
@@ -356,6 +356,12 @@ FXAudio2SoundBuffer* FXAudio2SoundBuffer::CreateQueuedBuffer( FXAudio2Device* XA
 	// Always create a new buffer for real time decompressed sounds
 	FXAudio2SoundBuffer* Buffer = new FXAudio2SoundBuffer( XAudio2Device, SoundFormat_PCMRT );
 
+	// If the buffer was precached as native, the resource data will have been lost and we need to re-initialize it
+	if (Wave->ResourceData == nullptr)
+	{
+		Wave->InitAudioResource(XAudio2Device->GetRuntimeFormat(Wave));
+	}
+
 	// Prime the first two buffers and prepare the decompression
 	FSoundQualityInfo QualityInfo = { 0 };
 
diff --git a/Engine/Source/Runtime/Windows/XAudio2/Private/XAudio2Source.cpp b/Engine/Source/Runtime/Windows/XAudio2/Private/XAudio2Source.cpp
--- a/Engine/Source/Runtime/Windows/XAudio2/Private/XAudio2Source.cpp
+++ b/Engine/Source/Runtime/Windows/XAudio2/Private/XAudio2Source.cpp
@@ -102,6 +102,7 @@ void FXAudio2SoundSource::FreeResources( void )
 		// ... free the buffers
 		FMemory::Free( ( void* )XAudio2Buffers[0].pAudioData );
 		FMemory::Free( ( void* )XAudio2Buffers[1].pAudioData );
+		FMemory::Free( ( void* )XAudio2Buffers[2].pAudioData );
 
 		// Buffers without a valid resource ID are transient and need to be deleted.
 		if( Buffer )
@@ -189,36 +190,44 @@ void FXAudio2SoundSource::SubmitPCMRTBuffers( void )
 {
 	SCOPE_CYCLE_COUNTER( STAT_AudioSubmitBuffersTime );
 
-	FMemory::Memzero( XAudio2Buffers, sizeof( XAUDIO2_BUFFER ) * 2 );
+	FMemory::Memzero( XAudio2Buffers, sizeof( XAUDIO2_BUFFER ) * 3 );
 
 	// Set the buffer to be in real time mode
 	CurrentBuffer = 0;
 
 	const uint32 BufferSize = MONO_PCM_BUFFER_SIZE * Buffer->NumChannels;
 	
-	// Set up double buffer area to decompress to
+	// Set up buffer areas to decompress to
 	XAudio2Buffers[0].pAudioData = (uint8*)FMemory::Malloc(BufferSize);
 	XAudio2Buffers[0].AudioBytes = BufferSize;
 
+	XAudio2Buffers[1].pAudioData = (uint8*)FMemory::Malloc(BufferSize);
+	XAudio2Buffers[1].AudioBytes = BufferSize;
+
 	// Only use the cached data if we're starting from the beginning, otherwise we'll have to take a synchronous hit
 	bool bSkipFirstBuffer = false;;
 	if (WaveInstance->WaveData && WaveInstance->WaveData->CachedRealtimeFirstBuffer && WaveInstance->StartTime == 0.f)
 	{
 		FMemory::Memcpy((uint8*)XAudio2Buffers[0].pAudioData, WaveInstance->WaveData->CachedRealtimeFirstBuffer, BufferSize);
+		FMemory::Memcpy((uint8*)XAudio2Buffers[1].pAudioData, WaveInstance->WaveData->CachedRealtimeFirstBuffer + BufferSize, BufferSize);
 		bSkipFirstBuffer = true;
 	}
 	else
 	{
 		ReadMorePCMData(0, EDataReadMode::Synchronous);
+		ReadMorePCMData(1, EDataReadMode::Synchronous);
 	}
 
 	AudioDevice->ValidateAPICall( TEXT( "SubmitSourceBuffer - PCMRT" ), 
-		Source->SubmitSourceBuffer( XAudio2Buffers ) );
+		Source->SubmitSourceBuffer( &XAudio2Buffers[0] ) );
 
-	XAudio2Buffers[1].pAudioData = (uint8*)FMemory::Malloc(BufferSize);
-	XAudio2Buffers[1].AudioBytes = BufferSize;
+	AudioDevice->ValidateAPICall( TEXT( "SubmitSourceBuffer - PCMRT" ), 
+		Source->SubmitSourceBuffer( &XAudio2Buffers[1] ) );
+
+	XAudio2Buffers[2].pAudioData = (uint8*)FMemory::Malloc(BufferSize);
+	XAudio2Buffers[2].AudioBytes = BufferSize;
 
-	CurrentBuffer = 1;
+	CurrentBuffer = 2;
 
 	// Start the async population of the next buffer
 	EDataReadMode DataReadMode = EDataReadMode::Asynchronous;
@@ -231,7 +240,7 @@ void FXAudio2SoundSource::SubmitPCMRTBuffers( void )
 		DataReadMode =  EDataReadMode::AsynchronousSkipFirstFrame;
 	}
 
-	ReadMorePCMData(1, DataReadMode);
+	ReadMorePCMData(2, DataReadMode);
 
 	bResourcesNeedFreeing = true;
 }
@@ -1208,16 +1217,62 @@ void FXAudio2SoundSource::Pause( void )
 	}
 }
 
+void FXAudio2SoundSource::HandleRealTimeSourceData(bool bLooped)
+{
+	// Have we reached the end of the compressed sound?
+	if( bLooped )
+	{
+		switch( WaveInstance->LoopingMode )
+		{
+		case LOOP_Never:
+			// Play out any queued buffers - once there are no buffers left, the state check at the beginning of IsFinished will fire
+			bBuffersToFlush = true;
+			XAudio2Buffers[CurrentBuffer].Flags |= XAUDIO2_END_OF_STREAM;
+			break;
+
+		case LOOP_WithNotification:
+			// If we have just looped, and we are programmatically looping, send notification
+			WaveInstance->NotifyFinished();
+			break;
+
+		case LOOP_Forever:
+			// Let the sound loop indefinitely
+			break;
+		}
+	}
+
+	if (XAudio2Buffers[CurrentBuffer].AudioBytes > 0)
+	{
+		AudioDevice->ValidateAPICall( TEXT( "SubmitSourceBuffer - IsFinished" ), 
+			Source->SubmitSourceBuffer( &XAudio2Buffers[CurrentBuffer] ) );
+	}
+	else
+	{
+		if (--CurrentBuffer < 0)
+		{
+			CurrentBuffer = 2;
+		}
+	}
+}
+
 /**
  * Handles feeding new data to a real time decompressed sound
  */
-void FXAudio2SoundSource::HandleRealTimeSource()
+void FXAudio2SoundSource::HandleRealTimeSource(bool bBlockForData)
 {
-	bool bLooped = false;
+	const bool bGetMoreData = bBlockForData || (RealtimeAsyncTask == nullptr);
 	if (RealtimeAsyncTask)
 	{
-		if (RealtimeAsyncTask->IsDone())
+		const bool bTaskDone = RealtimeAsyncTask->IsDone();
+		if (bTaskDone || bBlockForData)
 		{
+			bool bLooped = false;
+
+			if (!bTaskDone)
+			{
+				RealtimeAsyncTask->EnsureCompletion();
+			}
+
 			switch(RealtimeAsyncTask->GetTask().GetTaskType())
 			{
 			case ERealtimeAudioTaskType::Decompress:
@@ -1225,55 +1280,31 @@ void FXAudio2SoundSource::HandleRealTimeSource()
 				break;
 
 			case ERealtimeAudioTaskType::Procedural:
-				XAudio2Buffers[CurrentBuffer & 1].AudioBytes = RealtimeAsyncTask->GetTask().GetBytesWritten();
+				XAudio2Buffers[CurrentBuffer].AudioBytes = RealtimeAsyncTask->GetTask().GetBytesWritten();
 				break;
 			}
 
 			delete RealtimeAsyncTask;
 			RealtimeAsyncTask = nullptr;
-		}
-	}
-	else
-	{
-		// Update the double buffer toggle
-		++CurrentBuffer;
 
-		// Get the next bit of streaming data
-		bLooped = ReadMorePCMData(CurrentBuffer & 1, (XAudio2Buffer->SoundFormat == ESoundFormat::SoundFormat_Streaming ? EDataReadMode::Synchronous : EDataReadMode::Asynchronous));
+			HandleRealTimeSourceData(bLooped);
+		}
 	}
-
-	if (RealtimeAsyncTask == nullptr)
+	
+	if (bGetMoreData)
 	{
-		// Have we reached the end of the compressed sound?
-		if( bLooped )
+		// Update the buffer index
+		if (++CurrentBuffer > 2)
 		{
-			switch( WaveInstance->LoopingMode )
-			{
-			case LOOP_Never:
-				// Play out any queued buffers - once there are no buffers left, the state check at the beginning of IsFinished will fire
-				bBuffersToFlush = true;
-				XAudio2Buffers[CurrentBuffer & 1].Flags |= XAUDIO2_END_OF_STREAM;
-				break;
-
-			case LOOP_WithNotification:
-				// If we have just looped, and we are programmatically looping, send notification
-				WaveInstance->NotifyFinished();
-				break;
-
-			case LOOP_Forever:
-				// Let the sound loop indefinitely
-				break;
-			}
+			CurrentBuffer = 0;
 		}
 
-		if (XAudio2Buffers[CurrentBuffer & 1].AudioBytes > 0)
-		{
-			AudioDevice->ValidateAPICall( TEXT( "SubmitSourceBuffer - IsFinished" ), 
-				Source->SubmitSourceBuffer( &XAudio2Buffers[CurrentBuffer & 1] ) );
-		}
-		else
+		// Get the next bit of streaming data
+		const bool bLooped = ReadMorePCMData(CurrentBuffer, (XAudio2Buffer->SoundFormat == ESoundFormat::SoundFormat_Streaming ? EDataReadMode::Synchronous : EDataReadMode::Asynchronous));
+
+		if (RealtimeAsyncTask == nullptr)
 		{
-			--CurrentBuffer;
+			HandleRealTimeSourceData(bLooped);
 		}
 	}
 }
@@ -1310,10 +1341,10 @@ bool FXAudio2SoundSource::IsFinished( void )
 		}
 
 		// Service any real time sounds
-		if (bIsRealTimeSource && !bBuffersToFlush && SourceState.BuffersQueued <= 1)
+		if (bIsRealTimeSource && !bBuffersToFlush && SourceState.BuffersQueued <= 2)
 		{
 			// Continue feeding new sound data (unless we are waiting for the sound to finish)
-			HandleRealTimeSource();
+			HandleRealTimeSource(SourceState.BuffersQueued < 2);
 
 			return( false );
 		}
diff --git a/Engine/Source/Runtime/Windows/XAudio2/Private/XAudio2Support.h b/Engine/Source/Runtime/Windows/XAudio2/Private/XAudio2Support.h
--- a/Engine/Source/Runtime/Windows/XAudio2/Private/XAudio2Support.h
+++ b/Engine/Source/Runtime/Windows/XAudio2/Private/XAudio2Support.h
@@ -322,7 +322,12 @@ public:
 	/**
 	 * Handles feeding new data to a real time decompressed sound
 	 */
-	void HandleRealTimeSource();
+	void HandleRealTimeSource(bool bBlockForData);
+
+	/**
+	 * Handles pushing fetched real time source data to the hardware
+	 */
+	void HandleRealTimeSourceData(bool bLooped);
 
 	/**
 	 * Queries the status of the currently associated wave instance.
@@ -454,7 +459,7 @@ protected:
 	/** Which sound buffer should be written to next - used for double buffering. */
 	int32							CurrentBuffer;
 	/** A pair of sound buffers to allow notification when a sound loops. */
-	XAUDIO2_BUFFER				XAudio2Buffers[2];
+	XAUDIO2_BUFFER				XAudio2Buffers[3];
 	/** Additional buffer info for XWMA sounds */
 	XAUDIO2_BUFFER_WMA			XAudio2BufferXWMA[1];
 	/** Set when we wish to let the buffers play themselves out */
