UE4 - Try to track down crashes exiting D3D; clear a dangling reference (UE-17216, UE-17163)

diff --git a/Engine/Source/Runtime/Windows/D3D11RHI/Private/D3D11Device.cpp b/Engine/Source/Runtime/Windows/D3D11RHI/Private/D3D11Device.cpp
--- a/Engine/Source/Runtime/Windows/D3D11RHI/Private/D3D11Device.cpp
+++ b/Engine/Source/Runtime/Windows/D3D11RHI/Private/D3D11Device.cpp
@@ -199,8 +199,18 @@ FD3D11DynamicRHI::FD3D11DynamicRHI(IDXGIFactory1* InDXGIFactory1,D3D_FEATURE_LEV
 	}
 }
 
+FD3D11DynamicRHI::~FD3D11DynamicRHI()
+{
+	UE_LOG(LogD3D11RHI, Log, TEXT("~FD3D11DynamicRHI"));
+
+	check(Direct3DDeviceIMContext == nullptr);
+	check(Direct3DDevice == nullptr);
+	check(DXGIFactory1 == nullptr || DXGIFactory1.GetRefCount() == 1);
+}
+
 void FD3D11DynamicRHI::Shutdown()
 {
+	UE_LOG(LogD3D11RHI, Log, TEXT("Shutdown"));
 	check(IsInGameThread() && IsInRenderingThread());  // require that the render thread has been shut down
 
 	// Cleanup the D3D device.
@@ -217,8 +227,9 @@ void FD3D11DynamicRHI::Shutdown()
 
 void FD3D11DynamicRHI::RHIPushEvent(const TCHAR* Name)
 { 
-	GPUProfilingData.PushEvent(Name); 
+	GPUProfilingData.PushEvent(Name);
 }
+
 void FD3D11DynamicRHI::RHIPopEvent()
 { 
 	GPUProfilingData.PopEvent(); 
@@ -261,9 +272,13 @@ void FD3D11DynamicRHI::RHIGetSupportedResolution( uint32 &Width, uint32 &Height
 			TRefCountPtr<IDXGIOutput> Output;
 			HResult = Adapter->EnumOutputs(o,Output.GetInitReference());
 			if(DXGI_ERROR_NOT_FOUND == HResult)
+			{
 				break;
+			}
 			if(FAILED(HResult))
+			{
 				return;
+			}
 
 			// TODO: GetDisplayModeList is a terribly SLOW call.  It can take up to a second per invocation.
 			//  We might want to work around some DXGI badness here.
@@ -393,6 +408,8 @@ void FD3D11DynamicRHI::UpdateMSAASettings()
 
 void FD3D11DynamicRHI::CleanupD3DDevice()
 {
+	UE_LOG(LogD3D11RHI, Log, TEXT("CleanupD3DDevice"));
+
 	if(GIsRHIInitialized)
 	{
 		check(Direct3DDevice);
diff --git a/Engine/Source/Runtime/Windows/D3D11RHI/Private/D3D11RHI.cpp b/Engine/Source/Runtime/Windows/D3D11RHI/Private/D3D11RHI.cpp
--- a/Engine/Source/Runtime/Windows/D3D11RHI/Private/D3D11RHI.cpp
+++ b/Engine/Source/Runtime/Windows/D3D11RHI/Private/D3D11RHI.cpp
@@ -379,6 +379,7 @@ void FD3DGPUProfiler::BeginFrame(FD3D11DynamicRHI* InRHI)
 void FD3D11DynamicRHI::RHIEndFrame()
 {
 	GPUProfilingData.EndFrame();
+	CurrentComputeShader = nullptr;
 }
 
 void FD3DGPUProfiler::EndFrame()
diff --git a/Engine/Source/Runtime/Windows/D3D11RHI/Private/D3D11RHIPrivate.h b/Engine/Source/Runtime/Windows/D3D11RHI/Private/D3D11RHIPrivate.h
--- a/Engine/Source/Runtime/Windows/D3D11RHI/Private/D3D11RHIPrivate.h
+++ b/Engine/Source/Runtime/Windows/D3D11RHI/Private/D3D11RHIPrivate.h
@@ -276,7 +276,7 @@ public:
 	FD3D11DynamicRHI(IDXGIFactory1* InDXGIFactory1,D3D_FEATURE_LEVEL InFeatureLevel,int32 InChosenAdapter);
 
 	/** Destructor */
-	virtual ~FD3D11DynamicRHI() {}
+	virtual ~FD3D11DynamicRHI();
 
 	/** If it hasn't been initialized yet, initializes the D3D device. */
 	virtual void InitD3DDevice();
diff --git a/Engine/Source/Runtime/Windows/D3D11RHI/Private/Windows/WindowsD3D11Device.cpp b/Engine/Source/Runtime/Windows/D3D11RHI/Private/Windows/WindowsD3D11Device.cpp
--- a/Engine/Source/Runtime/Windows/D3D11RHI/Private/Windows/WindowsD3D11Device.cpp
+++ b/Engine/Source/Runtime/Windows/D3D11RHI/Private/Windows/WindowsD3D11Device.cpp
@@ -368,6 +368,7 @@ void FD3D11DynamicRHI::Init()
 
 void FD3D11DynamicRHI::InitD3DDevice()
 {
+	UE_LOG(LogD3D11RHI, Log, TEXT("InitD3DDevice"));
 	check( IsInGameThread() );
 
 	// Wait for the rendering thread to go idle.
@@ -376,6 +377,7 @@ void FD3D11DynamicRHI::InitD3DDevice()
 	// If the device we were using has been removed, release it and the resources we created for it.
 	if(bDeviceRemoved)
 	{
+		UE_LOG(LogD3D11RHI, Log, TEXT("bDeviceRemoved"));
 		check(Direct3DDevice);
 
 		HRESULT hRes = Direct3DDevice->GetDeviceRemovedReason();
@@ -403,6 +405,7 @@ void FD3D11DynamicRHI::InitD3DDevice()
 	// If we don't have a device yet, either because this is the first viewport, or the old device was removed, create a device.
 	if(!Direct3DDevice)
 	{
+		UE_LOG(LogD3D11RHI, Log, TEXT("!Direct3DDevice"));
 		check(!GIsRHIInitialized);
 
 		// Clear shadowed shader resources.
@@ -472,11 +475,13 @@ void FD3D11DynamicRHI::InitD3DDevice()
 					{
 						// Use dedicated video memory, if it's more than 200 MB
 						GTotalGraphicsMemory = GDedicatedVideoMemory;
-					} else if ( GDedicatedSystemMemory >= 200*1024*1024 )
+					}
+					else if ( GDedicatedSystemMemory >= 200*1024*1024 )
 					{
 						// Use dedicated system memory, if it's more than 200 MB
 						GTotalGraphicsMemory = GDedicatedSystemMemory;
-					} else if ( GSharedSystemMemory >= 400*1024*1024 )
+					}
+					else if ( GSharedSystemMemory >= 400*1024*1024 )
 					{
 						// Use some shared system memory, if it's more than 400 MB
 						GTotalGraphicsMemory = ConsideredSharedSystemMemory;
