UE4 - Fix temporarily crash

diff --git a/Engine/Source/Runtime/Windows/D3D11RHI/Private/D3D11Device.cpp b/Engine/Source/Runtime/Windows/D3D11RHI/Private/D3D11Device.cpp
--- a/Engine/Source/Runtime/Windows/D3D11RHI/Private/D3D11Device.cpp
+++ b/Engine/Source/Runtime/Windows/D3D11RHI/Private/D3D11Device.cpp
@@ -205,7 +205,6 @@ FD3D11DynamicRHI::~FD3D11DynamicRHI()
 
 	check(Direct3DDeviceIMContext == nullptr);
 	check(Direct3DDevice == nullptr);
-	check(DXGIFactory1 == nullptr || DXGIFactory1.GetRefCount() == 1);
 }
 
 void FD3D11DynamicRHI::Shutdown()
