UE4 - Remove spam msg for D3D

diff --git a/Engine/Source/Runtime/Windows/D3D11RHI/Private/Windows/WindowsD3D11Device.cpp b/Engine/Source/Runtime/Windows/D3D11RHI/Private/Windows/WindowsD3D11Device.cpp
--- a/Engine/Source/Runtime/Windows/D3D11RHI/Private/Windows/WindowsD3D11Device.cpp
+++ b/Engine/Source/Runtime/Windows/D3D11RHI/Private/Windows/WindowsD3D11Device.cpp
@@ -368,7 +368,6 @@ void FD3D11DynamicRHI::Init()
 
 void FD3D11DynamicRHI::InitD3DDevice()
 {
-	UE_LOG(LogD3D11RHI, Log, TEXT("InitD3DDevice"));
 	check( IsInGameThread() );
 
 	// Wait for the rendering thread to go idle.
