Fix SRV format

diff --git a/Engine/Source/Runtime/Windows/D3D11RHI/Private/D3D11Texture.cpp b/Engine/Source/Runtime/Windows/D3D11RHI/Private/D3D11Texture.cpp
--- a/Engine/Source/Runtime/Windows/D3D11RHI/Private/D3D11Texture.cpp
+++ b/Engine/Source/Runtime/Windows/D3D11RHI/Private/D3D11Texture.cpp
@@ -1139,13 +1139,16 @@ FShaderResourceViewRHIRef FD3D11DynamicRHI::RHICreateShaderResourceView(FTexture
 	D3D11_TEXTURE2D_DESC TextureDesc;
 	Texture2D->GetResource()->GetDesc(&TextureDesc);
 
+	bool bSRGB = (Texture2D->GetFlags() & TexCreate_SRGB) != 0;
+	const DXGI_FORMAT PlatformShaderResourceFormat = FindShaderResourceDXGIFormat(TextureDesc.Format,bSRGB);
+
 	// Create a Shader Resource View
 	D3D11_SHADER_RESOURCE_VIEW_DESC SRVDesc;
 	SRVDesc.ViewDimension = D3D11_SRV_DIMENSION_TEXTURE2D;
 	SRVDesc.Texture2D.MostDetailedMip = MipLevel;
 	SRVDesc.Texture2D.MipLevels = 1;
 
-	SRVDesc.Format = TextureDesc.Format;
+	SRVDesc.Format = PlatformShaderResourceFormat;
 	TRefCountPtr<ID3D11ShaderResourceView> ShaderResourceView;
 	VERIFYD3D11RESULT(Direct3DDevice->CreateShaderResourceView(Texture2D->GetResource(), &SRVDesc, (ID3D11ShaderResourceView**)ShaderResourceView.GetInitReference()));
 
