Add ability to create SRVs for additional texture types (3d, 2dArray and Cube)

diff --git a/Engine/Source/Runtime/EmptyRHI/Private/EmptyUAV.cpp b/Engine/Source/Runtime/EmptyRHI/Private/EmptyUAV.cpp
--- a/Engine/Source/Runtime/EmptyRHI/Private/EmptyUAV.cpp
+++ b/Engine/Source/Runtime/EmptyRHI/Private/EmptyUAV.cpp
@@ -75,6 +75,27 @@ FShaderResourceViewRHIRef FEmptyDynamicRHI::RHICreateShaderResourceView(FTexture
 	return SRV;
 }
 
+FShaderResourceViewRHIRef FEmptyDynamicRHI::RHICreateShaderResourceView(FTexture3DRHIParamRef Texture3DRHI, uint8 MipLevel)
+{
+	FEmptyShaderResourceView* SRV = new FEmptyShaderResourceView;
+	SRV->SourceTexture = (FRHITexture*)Texture3DRHI;
+	return SRV;
+}
+
+FShaderResourceViewRHIRef FEmptyDynamicRHI::RHICreateShaderResourceView(FTexture2DArrayRHIParamRef Texture2DArrayRHI, uint8 MipLevel)
+{
+	FEmptyShaderResourceView* SRV = new FEmptyShaderResourceView;
+	SRV->SourceTexture = (FRHITexture*)Texture2DArrayRHI;
+	return SRV;
+}
+
+FShaderResourceViewRHIRef FEmptyDynamicRHI::RHICreateShaderResourceView(FTextureCubeRHIParamRef TextureCubeRHI, uint8 MipLevel)
+{
+	FEmptyShaderResourceView* SRV = new FEmptyShaderResourceView;
+	SRV->SourceTexture = (FRHITexture*)TextureCubeRHI;
+	return SRV;
+}
+
 void FEmptyDynamicRHI::RHIClearUAV(FUnorderedAccessViewRHIParamRef UnorderedAccessViewRHI, const uint32* Values)
 {
 	FEmptyUnorderedAccessView* UnorderedAccessView = ResourceCast(UnorderedAccessViewRHI);
diff --git a/Engine/Source/Runtime/EmptyRHI/Public/EmptyRHI.h b/Engine/Source/Runtime/EmptyRHI/Public/EmptyRHI.h
--- a/Engine/Source/Runtime/EmptyRHI/Public/EmptyRHI.h
+++ b/Engine/Source/Runtime/EmptyRHI/Public/EmptyRHI.h
@@ -81,6 +81,9 @@ public:
 	virtual void RHIGetResourceInfo(FTextureRHIParamRef Ref, FRHIResourceInfo& OutInfo) final override;
 	virtual FShaderResourceViewRHIRef RHICreateShaderResourceView(FTexture2DRHIParamRef Texture2DRHI, uint8 MipLevel) final override;
 	virtual FShaderResourceViewRHIRef RHICreateShaderResourceView(FTexture2DRHIParamRef Texture2DRHI, uint8 MipLevel, uint8 NumMipLevels, uint8 Format) final override;
+	virtual FShaderResourceViewRHIRef RHICreateShaderResourceView(FTexture3DRHIParamRef Texture3DRHI, uint8 MipLevel) final override;
+	virtual FShaderResourceViewRHIRef RHICreateShaderResourceView(FTexture2DArrayRHIParamRef Texture2DArrayRHI, uint8 MipLevel) final override;
+	virtual FShaderResourceViewRHIRef RHICreateShaderResourceView(FTextureCubeRHIParamRef TextureCubeRHI, uint8 MipLevel) final override;
 	virtual void RHIGenerateMips(FTextureRHIParamRef Texture) final override;
 	virtual uint32 RHIComputeMemorySize(FTextureRHIParamRef TextureRHI) final override;
 	virtual FTexture2DRHIRef RHIAsyncReallocateTexture2D(FTexture2DRHIParamRef Texture2D, int32 NewMipCount, int32 NewSizeX, int32 NewSizeY, FThreadSafeCounter* RequestStatus) final override;
diff --git a/Engine/Source/Runtime/IOS/MetalRHI/Private/MetalUAV.cpp b/Engine/Source/Runtime/IOS/MetalRHI/Private/MetalUAV.cpp
--- a/Engine/Source/Runtime/IOS/MetalRHI/Private/MetalUAV.cpp
+++ b/Engine/Source/Runtime/IOS/MetalRHI/Private/MetalUAV.cpp
@@ -103,6 +103,27 @@ FShaderResourceViewRHIRef FMetalDynamicRHI::RHICreateShaderResourceView(FTexture
 	return SRV;
 }
 
+FShaderResourceViewRHIRef FMetalDynamicRHI::RHICreateShaderResourceView(FTexture3DRHIParamRef Texture3DRHI, uint8 MipLevel)
+{
+	FEmptyShaderResourceView* SRV = new FMetalShaderResourceView;
+	SRV->SourceTexture = (FRHITexture*)Texture3DRHI;
+	return SRV;
+}
+
+FShaderResourceViewRHIRef FMetalDynamicRHI::RHICreateShaderResourceView(FTexture2DArrayRHIParamRef Texture2DArrayRHI, uint8 MipLevel)
+{
+	FEmptyShaderResourceView* SRV = new FMetalShaderResourceView;
+	SRV->SourceTexture = (FRHITexture*)Texture2DArrayRHI;
+	return SRV;
+}
+
+FShaderResourceViewRHIRef FMetalDynamicRHI::RHICreateShaderResourceView(FTextureCubeRHIParamRef TextureCubeRHI, uint8 MipLevel)
+{
+	FEmptyShaderResourceView* SRV = new FMetalShaderResourceView;
+	SRV->SourceTexture = (FRHITexture*)TextureCubeRHI;
+	return SRV;
+}
+
 void FMetalDynamicRHI::RHIClearUAV(FUnorderedAccessViewRHIParamRef UnorderedAccessViewRHI, const uint32* Values)
 {
 	FMetalUnorderedAccessView* UnorderedAccessView = ResourceCast(UnorderedAccessViewRHI);
diff --git a/Engine/Source/Runtime/IOS/MetalRHI/Public/MetalRHI.h b/Engine/Source/Runtime/IOS/MetalRHI/Public/MetalRHI.h
--- a/Engine/Source/Runtime/IOS/MetalRHI/Public/MetalRHI.h
+++ b/Engine/Source/Runtime/IOS/MetalRHI/Public/MetalRHI.h
@@ -84,6 +84,9 @@ public:
 	virtual void RHIGetResourceInfo(FTextureRHIParamRef Ref, FRHIResourceInfo& OutInfo) final override;
 	virtual FShaderResourceViewRHIRef RHICreateShaderResourceView(FTexture2DRHIParamRef Texture2DRHI, uint8 MipLevel) final override;
 	virtual FShaderResourceViewRHIRef RHICreateShaderResourceView(FTexture2DRHIParamRef Texture2DRHI, uint8 MipLevel, uint8 NumMipLevels, uint8 Format) final override;
+	virtual FShaderResourceViewRHIRef RHICreateShaderResourceView(FTexture3DRHIParamRef Texture3DRHI, uint8 MipLevel) final override;
+	virtual FShaderResourceViewRHIRef RHICreateShaderResourceView(FTexture2DArrayRHIParamRef Texture2DArrayRHI, uint8 MipLevel) final override;
+	virtual FShaderResourceViewRHIRef RHICreateShaderResourceView(FTextureCubeRHIParamRef TextureCubeRHI, uint8 MipLevel) final override;
 	virtual void RHIGenerateMips(FTextureRHIParamRef Texture) final override;
 	virtual uint32 RHIComputeMemorySize(FTextureRHIParamRef TextureRHI) final override;
 	virtual FTexture2DRHIRef RHIAsyncReallocateTexture2D(FTexture2DRHIParamRef Texture2D, int32 NewMipCount, int32 NewSizeX, int32 NewSizeY, FThreadSafeCounter* RequestStatus) final override;
diff --git a/Engine/Source/Runtime/NullDrv/Public/NullRHI.h b/Engine/Source/Runtime/NullDrv/Public/NullRHI.h
--- a/Engine/Source/Runtime/NullDrv/Public/NullRHI.h
+++ b/Engine/Source/Runtime/NullDrv/Public/NullRHI.h
@@ -294,6 +294,22 @@ public:
 	{ 
 		return new FRHIShaderResourceView(); 
 	}
+
+	virtual FShaderResourceViewRHIRef RHICreateShaderResourceView(FTexture3DRHIParamRef Texture3DRHI, uint8 MipLevel) final override
+	{
+		return new FRHIShaderResourceView(); 
+	}
+
+	virtual FShaderResourceViewRHIRef RHICreateShaderResourceView(FTexture2DArrayRHIParamRef Texture2DArrayRHI, uint8 MipLevel) final override
+	{
+		return new FRHIShaderResourceView(); 
+	}
+
+	virtual FShaderResourceViewRHIRef RHICreateShaderResourceView(FTextureCubeRHIParamRef TextureCubeRHI, uint8 MipLevel) final override
+	{
+		return new FRHIShaderResourceView(); 
+	}
+
 	virtual void RHIGenerateMips(FTextureRHIParamRef Texture) final override
 	{
 
diff --git a/Engine/Source/Runtime/OpenGLDrv/Private/OpenGLTexture.cpp b/Engine/Source/Runtime/OpenGLDrv/Private/OpenGLTexture.cpp
--- a/Engine/Source/Runtime/OpenGLDrv/Private/OpenGLTexture.cpp
+++ b/Engine/Source/Runtime/OpenGLDrv/Private/OpenGLTexture.cpp
@@ -1732,6 +1732,96 @@ FShaderResourceViewRHIRef FOpenGLDynamicRHI::RHICreateShaderResourceView(FTextur
 	return View;
 }
 
+FShaderResourceViewRHIRef FOpenGLDynamicRHI::RHICreateShaderResourceView(FTexture3DRHIParamRef Texture3DRHI, uint8 MipLevel)
+{
+	FOpenGLTexture3D* Texture3D = ResourceCast(Texture3DRHI);
+
+	FOpenGLShaderResourceView *View = 0;
+
+	if (FOpenGL::SupportsTextureView())
+	{
+		VERIFY_GL_SCOPE();
+
+		GLuint Resource = 0;
+
+		FOpenGL::GenTextures( 1, &Resource);
+		const FOpenGLTextureFormat& GLFormat = GOpenGLTextureFormats[Texture3D->GetFormat()];
+		const bool bSRGB = (Texture3D->GetFlags()&TexCreate_SRGB) != 0;
+		
+		FOpenGL::TextureView( Resource, Texture3D->Target, Texture3D->Resource, GLFormat.InternalFormat[bSRGB], MipLevel, 1, 0, 1);
+		
+		View = new FOpenGLShaderResourceView(this, Resource, Texture3D->Target, MipLevel, true);
+	}
+	else
+	{
+		View = new FOpenGLShaderResourceView(this, Texture3D->Resource, Texture3D->Target, MipLevel, false);
+	}
+	
+	FShaderCache::LogSRV(View, Texture3DRHI, MipLevel, Texture3DRHI->GetNumMips(), Texture3DRHI->GetFormat());
+
+	return View;
+}
+
+FShaderResourceViewRHIRef FOpenGLDynamicRHI::RHICreateShaderResourceView(FTexture2DArrayRHIParamRef Texture2DArrayRHI, uint8 MipLevel)
+{
+	FOpenGLTexture2DArray* Texture2DArray = ResourceCast(Texture2DArrayRHI);
+
+	FOpenGLShaderResourceView *View = 0;
+
+	if (FOpenGL::SupportsTextureView())
+	{
+		VERIFY_GL_SCOPE();
+
+		GLuint Resource = 0;
+
+		FOpenGL::GenTextures( 1, &Resource);
+		const FOpenGLTextureFormat& GLFormat = GOpenGLTextureFormats[Texture2DArray->GetFormat()];
+		const bool bSRGB = (Texture2DArray->GetFlags()&TexCreate_SRGB) != 0;
+		
+		FOpenGL::TextureView( Resource, Texture2DArray->Target, Texture2DArray->Resource, GLFormat.InternalFormat[bSRGB], MipLevel, 1, 0, 1);
+		
+		View = new FOpenGLShaderResourceView(this, Resource, Texture2DArray->Target, MipLevel, true);
+	}
+	else
+	{
+		View = new FOpenGLShaderResourceView(this, Texture2DArray->Resource, Texture2DArray->Target, MipLevel, false);
+	}
+	
+	FShaderCache::LogSRV(View, Texture2DArrayRHI, MipLevel, Texture2DArrayRHI->GetNumMips(), Texture2DArrayRHI->GetFormat());
+
+	return View;
+}
+
+FShaderResourceViewRHIRef FOpenGLDynamicRHI::RHICreateShaderResourceView(FTextureCubeRHIParamRef TextureCubeRHI, uint8 MipLevel)
+{
+	FOpenGLTextureCube* TextureCube = ResourceCast(TextureCubeRHI);
+
+	FOpenGLShaderResourceView *View = 0;
+
+	if (FOpenGL::SupportsTextureView())
+	{
+		VERIFY_GL_SCOPE();
+
+		GLuint Resource = 0;
+
+		FOpenGL::GenTextures( 1, &Resource);
+		const FOpenGLTextureFormat& GLFormat = GOpenGLTextureFormats[TextureCube->GetFormat()];
+		const bool bSRGB = (TextureCube->GetFlags()&TexCreate_SRGB) != 0;
+		
+		FOpenGL::TextureView( Resource, TextureCube->Target, TextureCube->Resource, GLFormat.InternalFormat[bSRGB], MipLevel, 1, 0, 1);
+		
+		View = new FOpenGLShaderResourceView(this, Resource, TextureCube->Target, MipLevel, true);
+	}
+	else
+	{
+		View = new FOpenGLShaderResourceView(this, TextureCube->Resource, TextureCube->Target, MipLevel, false);
+	}
+	
+	FShaderCache::LogSRV(View, TextureCube, MipLevel, TextureCubeRHI->GetNumMips(), TextureCubeRHI->GetFormat());
+
+	return View;
+}
+
 
 /** Generates mip maps for the surface. */
 void FOpenGLDynamicRHI::RHIGenerateMips(FTextureRHIParamRef SurfaceRHI)
diff --git a/Engine/Source/Runtime/OpenGLDrv/Public/OpenGLDrv.h b/Engine/Source/Runtime/OpenGLDrv/Public/OpenGLDrv.h
--- a/Engine/Source/Runtime/OpenGLDrv/Public/OpenGLDrv.h
+++ b/Engine/Source/Runtime/OpenGLDrv/Public/OpenGLDrv.h
@@ -353,6 +353,9 @@ public:
 	virtual void RHIGetResourceInfo(FTextureRHIParamRef Ref, FRHIResourceInfo& OutInfo) final override;
 	virtual FShaderResourceViewRHIRef RHICreateShaderResourceView(FTexture2DRHIParamRef Texture2DRHI, uint8 MipLevel) final override;
 	virtual FShaderResourceViewRHIRef RHICreateShaderResourceView(FTexture2DRHIParamRef Texture2DRHI, uint8 MipLevel, uint8 NumMipLevels, uint8 Format) final override;
+	virtual FShaderResourceViewRHIRef RHICreateShaderResourceView(FTexture3DRHIParamRef Texture3DRHI, uint8 MipLevel) final override;
+	virtual FShaderResourceViewRHIRef RHICreateShaderResourceView(FTexture2DArrayRHIParamRef Texture2DArrayRHI, uint8 MipLevel) final override;
+	virtual FShaderResourceViewRHIRef RHICreateShaderResourceView(FTextureCubeRHIParamRef TextureCubeRHI, uint8 MipLevel) final override;
 	virtual void RHIGenerateMips(FTextureRHIParamRef Texture) final override;
 	virtual uint32 RHIComputeMemorySize(FTextureRHIParamRef TextureRHI) final override;
 	virtual FTexture2DRHIRef RHIAsyncReallocateTexture2D(FTexture2DRHIParamRef Texture2D, int32 NewMipCount, int32 NewSizeX, int32 NewSizeY, FThreadSafeCounter* RequestStatus) final override;
diff --git a/Engine/Source/Runtime/RHI/Public/DynamicRHI.h b/Engine/Source/Runtime/RHI/Public/DynamicRHI.h
--- a/Engine/Source/Runtime/RHI/Public/DynamicRHI.h
+++ b/Engine/Source/Runtime/RHI/Public/DynamicRHI.h
@@ -118,7 +118,7 @@ public:
 	virtual void* RHILockVertexBuffer(FVertexBufferRHIParamRef VertexBuffer, uint32 Offset, uint32 SizeRHI, EResourceLockMode LockMode) = 0;
 
 	// FlushType: Flush RHI Thread
-	virtual void RHIUnlockVertexBuffer(FVertexBufferRHIParamRef VertexBuffer) = 0;	
+	virtual void RHIUnlockVertexBuffer(FVertexBufferRHIParamRef VertexBuffer) = 0;
 
 	/** Copies the contents of one vertex buffer to another vertex buffer.  They must have identical sizes. */
 	// FlushType: Flush Immediate (seems dangerous)
@@ -300,6 +300,27 @@ public:
 	// FlushType: Wait RHI Thread
 	virtual FShaderResourceViewRHIRef RHICreateShaderResourceView(FTexture2DRHIParamRef Texture2DRHI, uint8 MipLevel, uint8 NumMipLevels, uint8 Format) = 0;
 
+	/**
+	* Creates a shader resource view for a 3d texture, viewing only a single
+	* mip level.
+	*/
+	// FlushType: Wait RHI Thread
+	virtual FShaderResourceViewRHIRef RHICreateShaderResourceView(FTexture3DRHIParamRef Texture3DRHI, uint8 MipLevel) = 0;
+
+	/**
+	* Creates a shader resource view for a 2d texture array, viewing only a single
+	* mip level.
+	*/
+	// FlushType: Wait RHI Thread
+	virtual FShaderResourceViewRHIRef RHICreateShaderResourceView(FTexture2DArrayRHIParamRef Texture2DArrayRHI, uint8 MipLevel) = 0;
+
+	/**
+	* Creates a shader resource view for a cube texture, viewing only a single
+	* mip level.
+	*/
+	// FlushType: Wait RHI Thread
+	virtual FShaderResourceViewRHIRef RHICreateShaderResourceView(FTextureCubeRHIParamRef TextureCubeRHI, uint8 MipLevel) = 0;
+
 	/**
 	* Generates mip maps for a texture.
 	*/
diff --git a/Engine/Source/Runtime/RHI/Public/RHICommandList.h b/Engine/Source/Runtime/RHI/Public/RHICommandList.h
--- a/Engine/Source/Runtime/RHI/Public/RHICommandList.h
+++ b/Engine/Source/Runtime/RHI/Public/RHICommandList.h
@@ -2042,7 +2042,7 @@ public:
 	{
 		return RHICreateUniformBuffer(Contents, Layout, Usage);
 	}
-
+	
 	FORCEINLINE FIndexBufferRHIRef CreateAndLockIndexBuffer(uint32 Stride, uint32 Size, uint32 InUsage, FRHIResourceCreateInfo& CreateInfo, void*& OutDataBuffer)
 	{
 		return GDynamicRHI->CreateAndLockIndexBuffer_RenderThread(*this, Stride, Size, InUsage, CreateInfo, OutDataBuffer);
@@ -2271,6 +2271,36 @@ public:
 		return GDynamicRHI->RHICreateShaderResourceView(Texture2DRHI, MipLevel, NumMipLevels, Format);
 	}
 	
+	FORCEINLINE FShaderResourceViewRHIRef CreateShaderResourceView(FTexture3DRHIParamRef Texture3DRHI, uint8 MipLevel)
+	{
+		if(GRHIThread)
+		{
+			QUICK_SCOPE_CYCLE_COUNTER(STAT_RHIMETHOD_CreateShaderResourceView_WaitRHI);
+			ImmediateFlush(EImmediateFlushType::WaitForRHIThread);
+		}
+		return GDynamicRHI->RHICreateShaderResourceView(Texture3DRHI, MipLevel);
+	}
+
+	FORCEINLINE FShaderResourceViewRHIRef CreateShaderResourceView(FTexture2DArrayRHIParamRef Texture2DArrayRHI, uint8 MipLevel)
+	{
+		if(GRHIThread)
+		{
+			QUICK_SCOPE_CYCLE_COUNTER(STAT_RHIMETHOD_CreateShaderResourceView_WaitRHI);
+			ImmediateFlush(EImmediateFlushType::WaitForRHIThread);
+		}
+		return GDynamicRHI->RHICreateShaderResourceView(Texture2DArrayRHI, MipLevel);
+	}
+
+	FORCEINLINE FShaderResourceViewRHIRef CreateShaderResourceView(FTextureCubeRHIParamRef TextureCubeRHI, uint8 MipLevel)
+	{
+		if(GRHIThread)
+		{
+			QUICK_SCOPE_CYCLE_COUNTER(STAT_RHIMETHOD_CreateShaderResourceView_WaitRHI);
+			ImmediateFlush(EImmediateFlushType::WaitForRHIThread);
+		}
+		return GDynamicRHI->RHICreateShaderResourceView(TextureCubeRHI, MipLevel);
+	}
+
 	FORCEINLINE void GenerateMips(FTextureRHIParamRef Texture)
 	{
 		QUICK_SCOPE_CYCLE_COUNTER(STAT_RHIMETHOD_GenerateMips_Flush);
@@ -2846,6 +2876,21 @@ FORCEINLINE FShaderResourceViewRHIRef RHICreateShaderResourceView(FTexture2DRHIP
 	return FRHICommandListExecutor::GetImmediateCommandList().CreateShaderResourceView(Texture2DRHI, MipLevel, NumMipLevels, Format);
 }
 
+FORCEINLINE FShaderResourceViewRHIRef RHICreateShaderResourceView(FTexture3DRHIParamRef Texture3DRHI, uint8 MipLevel)
+{
+	return FRHICommandListExecutor::GetImmediateCommandList().CreateShaderResourceView(Texture3DRHI, MipLevel);
+}
+
+FORCEINLINE FShaderResourceViewRHIRef RHICreateShaderResourceView(FTexture2DArrayRHIParamRef Texture2DArrayRHI, uint8 MipLevel)
+{
+	return FRHICommandListExecutor::GetImmediateCommandList().CreateShaderResourceView(Texture2DArrayRHI, MipLevel);
+}
+
+FORCEINLINE FShaderResourceViewRHIRef RHICreateShaderResourceView(FTextureCubeRHIParamRef TextureCubeRHI, uint8 MipLevel)
+{
+	return FRHICommandListExecutor::GetImmediateCommandList().CreateShaderResourceView(TextureCubeRHI, MipLevel);
+}
+
 FORCEINLINE FTexture2DRHIRef RHIAsyncReallocateTexture2D(FTexture2DRHIParamRef Texture2D, int32 NewMipCount, int32 NewSizeX, int32 NewSizeY, FThreadSafeCounter* RequestStatus)
 {
 	return FRHICommandListExecutor::GetImmediateCommandList().AsyncReallocateTexture2D(Texture2D, NewMipCount, NewSizeX, NewSizeY, RequestStatus);
diff --git a/Engine/Source/Runtime/Windows/D3D11RHI/Private/D3D11RHIPrivate.h b/Engine/Source/Runtime/Windows/D3D11RHI/Private/D3D11RHIPrivate.h
--- a/Engine/Source/Runtime/Windows/D3D11RHI/Private/D3D11RHIPrivate.h
+++ b/Engine/Source/Runtime/Windows/D3D11RHI/Private/D3D11RHIPrivate.h
@@ -348,6 +348,9 @@ public:
 	virtual void RHIGetResourceInfo(FTextureRHIParamRef Ref, FRHIResourceInfo& OutInfo) final override;
 	virtual FShaderResourceViewRHIRef RHICreateShaderResourceView(FTexture2DRHIParamRef Texture2DRHI, uint8 MipLevel) final override;
 	virtual FShaderResourceViewRHIRef RHICreateShaderResourceView(FTexture2DRHIParamRef Texture2DRHI, uint8 MipLevel, uint8 NumMipLevels, uint8 Format) final override;
+	virtual FShaderResourceViewRHIRef RHICreateShaderResourceView(FTexture3DRHIParamRef Texture3DRHI, uint8 MipLevel) final override;
+	virtual FShaderResourceViewRHIRef RHICreateShaderResourceView(FTexture2DArrayRHIParamRef Texture2DArrayRHI, uint8 MipLevel) final override;
+	virtual FShaderResourceViewRHIRef RHICreateShaderResourceView(FTextureCubeRHIParamRef TextureCubeRHI, uint8 MipLevel) final override;
 	virtual void RHIGenerateMips(FTextureRHIParamRef Texture) final override;
 	virtual uint32 RHIComputeMemorySize(FTextureRHIParamRef TextureRHI) final override;
 	virtual FTexture2DRHIRef RHIAsyncReallocateTexture2D(FTexture2DRHIParamRef Texture2D, int32 NewMipCount, int32 NewSizeX, int32 NewSizeY, FThreadSafeCounter* RequestStatus) final override;
diff --git a/Engine/Source/Runtime/Windows/D3D11RHI/Private/D3D11Texture.cpp b/Engine/Source/Runtime/Windows/D3D11RHI/Private/D3D11Texture.cpp
--- a/Engine/Source/Runtime/Windows/D3D11RHI/Private/D3D11Texture.cpp
+++ b/Engine/Source/Runtime/Windows/D3D11RHI/Private/D3D11Texture.cpp
@@ -1199,6 +1199,77 @@ FShaderResourceViewRHIRef FD3D11DynamicRHI::RHICreateShaderResourceView(FTexture
 	return new FD3D11ShaderResourceView(ShaderResourceView,Texture2D);
 }
 
+FShaderResourceViewRHIRef FD3D11DynamicRHI::RHICreateShaderResourceView(FTexture3DRHIParamRef Texture3DRHI, uint8 MipLevel)
+{
+	FD3D11Texture3D* Texture3D = ResourceCast(Texture3DRHI);
+
+	D3D11_TEXTURE3D_DESC TextureDesc;
+	Texture3D->GetResource()->GetDesc(&TextureDesc);
+
+	bool bSRGB = (Texture3D->GetFlags() & TexCreate_SRGB) != 0;
+	const DXGI_FORMAT PlatformShaderResourceFormat = FindShaderResourceDXGIFormat(TextureDesc.Format,bSRGB);
+
+	// Create a Shader Resource View
+	D3D11_SHADER_RESOURCE_VIEW_DESC SRVDesc;
+	SRVDesc.ViewDimension = D3D11_SRV_DIMENSION_TEXTURE3D;
+	SRVDesc.Texture3D.MostDetailedMip = MipLevel;
+	SRVDesc.Texture3D.MipLevels = 1;
+
+	SRVDesc.Format = PlatformShaderResourceFormat;
+	TRefCountPtr<ID3D11ShaderResourceView> ShaderResourceView;
+	VERIFYD3D11RESULT(Direct3DDevice->CreateShaderResourceView(Texture3D->GetResource(), &SRVDesc, (ID3D11ShaderResourceView**)ShaderResourceView.GetInitReference()));
+
+	return new FD3D11ShaderResourceView(ShaderResourceView,Texture3D);
+}
+
+FShaderResourceViewRHIRef FD3D11DynamicRHI::RHICreateShaderResourceView(FTexture2DArrayRHIParamRef Texture2DArrayRHI, uint8 MipLevel)
+{
+	FD3D11Texture2DArray* Texture2DArray = ResourceCast(Texture2DArrayRHI);
+
+	D3D11_TEXTURE2D_DESC TextureDesc;
+	Texture2DArray->GetResource()->GetDesc(&TextureDesc);
+
+	bool bSRGB = (Texture2DArray->GetFlags() & TexCreate_SRGB) != 0;
+	const DXGI_FORMAT PlatformShaderResourceFormat = FindShaderResourceDXGIFormat(TextureDesc.Format,bSRGB);
+
+	// Create a Shader Resource View
+	D3D11_SHADER_RESOURCE_VIEW_DESC SRVDesc;
+	SRVDesc.ViewDimension = D3D11_SRV_DIMENSION_TEXTURE2DARRAY;
+	SRVDesc.Texture2DArray.MostDetailedMip = MipLevel;
+	SRVDesc.Texture2DArray.MipLevels = 1;
+	SRVDesc.Texture2DArray.FirstArraySlice = 0;
+	SRVDesc.Texture2DArray.ArraySize = TextureDesc.ArraySize;
+
+	SRVDesc.Format = PlatformShaderResourceFormat;
+	TRefCountPtr<ID3D11ShaderResourceView> ShaderResourceView;
+	VERIFYD3D11RESULT(Direct3DDevice->CreateShaderResourceView(Texture2DArray->GetResource(), &SRVDesc, (ID3D11ShaderResourceView**)ShaderResourceView.GetInitReference()));
+
+	return new FD3D11ShaderResourceView(ShaderResourceView,Texture2DArray);
+}
+
+FShaderResourceViewRHIRef FD3D11DynamicRHI::RHICreateShaderResourceView(FTextureCubeRHIParamRef TextureCubeRHI, uint8 MipLevel)
+{
+	FD3D11TextureCube* TextureCube = ResourceCast(TextureCubeRHI);
+
+	D3D11_TEXTURE2D_DESC TextureDesc;
+	TextureCube->GetResource()->GetDesc(&TextureDesc);
+
+	bool bSRGB = (TextureCube->GetFlags() & TexCreate_SRGB) != 0;
+	const DXGI_FORMAT PlatformShaderResourceFormat = FindShaderResourceDXGIFormat(TextureDesc.Format,bSRGB);
+
+	// Create a Shader Resource View
+	D3D11_SHADER_RESOURCE_VIEW_DESC SRVDesc;
+	SRVDesc.ViewDimension = D3D11_SRV_DIMENSION_TEXTURECUBE;
+	SRVDesc.TextureCube.MostDetailedMip = MipLevel;
+	SRVDesc.TextureCube.MipLevels = 1;
+
+	SRVDesc.Format = PlatformShaderResourceFormat;
+	TRefCountPtr<ID3D11ShaderResourceView> ShaderResourceView;
+	VERIFYD3D11RESULT(Direct3DDevice->CreateShaderResourceView(TextureCube->GetResource(), &SRVDesc, (ID3D11ShaderResourceView**)ShaderResourceView.GetInitReference()));
+
+	return new FD3D11ShaderResourceView(ShaderResourceView,TextureCube);
+}
+
 /** Generates mip maps for the surface. */
 void FD3D11DynamicRHI::RHIGenerateMips(FTextureRHIParamRef TextureRHI)
 {
