PR #794: D3D11 depth-stencil drawn-clear fix (Contributed by [REDACTED_PERSON]). UE-8970

diff --git a/Engine/Source/Runtime/Windows/D3D11RHI/Private/D3D11Commands.cpp b/Engine/Source/Runtime/Windows/D3D11RHI/Private/D3D11Commands.cpp
--- a/Engine/Source/Runtime/Windows/D3D11RHI/Private/D3D11Commands.cpp
+++ b/Engine/Source/Runtime/Windows/D3D11RHI/Private/D3D11Commands.cpp
@@ -1640,6 +1640,20 @@ void FD3D11DynamicRHI::RHIClearMRT(bool bClearColor,int32 NumClearColors,const F
 			Width = Desc.Width;
 			Height = Desc.Height;
 			BaseTexture->Release();
+
+			// Adjust dimensions for the mip level we're clearing.
+			D3D11_DEPTH_STENCIL_VIEW_DESC DSVDesc;
+			DepthStencilView->GetDesc(&DSVDesc);
+			if (DSVDesc.ViewDimension == D3D11_DSV_DIMENSION_TEXTURE1D ||
+				DSVDesc.ViewDimension == D3D11_DSV_DIMENSION_TEXTURE1DARRAY ||
+				DSVDesc.ViewDimension == D3D11_DSV_DIMENSION_TEXTURE2D ||
+				DSVDesc.ViewDimension == D3D11_DSV_DIMENSION_TEXTURE2DARRAY)
+			{
+				// All the non-multisampled texture types have their mip-slice in the same position.
+				uint32 MipIndex = DSVDesc.Texture2D.MipSlice;
+				Width >>= MipIndex;
+				Height >>= MipIndex;
+			}
 		}
 
 		if ((Viewport.Width < Width || Viewport.Height < Height) 
