UE4 - Prereq integration for D3D12 support
(Not sure why these are locked...)

diff --git a/Engine/Source/Developer/Windows/ShaderFormatD3D/Private/D3D11ShaderCompiler.cpp b/Engine/Source/Developer/Windows/ShaderFormatD3D/Private/D3D11ShaderCompiler.cpp
--- a/Engine/Source/Developer/Windows/ShaderFormatD3D/Private/D3D11ShaderCompiler.cpp
+++ b/Engine/Source/Developer/Windows/ShaderFormatD3D/Private/D3D11ShaderCompiler.cpp
@@ -482,6 +482,9 @@ void CompileD3D11Shader(const FShaderCompilerInput& Input,FShaderCompilerOutput&
 
 			bool bGlobalUniformBufferUsed = false;
 			uint32 NumSamplers = 0;
+			uint32 NumSRVs = 0;
+			uint32 NumCBs = 0;
+			uint32 NumUAVs = 0;
 
 			TBitArray<> UsedUniformBufferSlots;
 			UsedUniformBufferSlots.Init(false,32);
@@ -547,6 +550,8 @@ void CompileD3D11Shader(const FShaderCompilerInput& Input,FShaderCompilerOutput&
 							);
 						UsedUniformBufferSlots[CBIndex] = true;
 					}
+
+					NumCBs = FMath::Max(NumCBs, BindDesc.BindPoint + BindDesc.BindCount);
 				}
 				else if (BindDesc.Type == D3D10_SIT_TEXTURE || BindDesc.Type == D3D10_SIT_SAMPLER)
 				{
@@ -592,7 +597,11 @@ void CompileD3D11Shader(const FShaderCompilerInput& Input,FShaderCompilerOutput&
 
 					if (BindDesc.Type == D3D10_SIT_SAMPLER)
 					{
-						NumSamplers += BindCount;
+						NumSamplers = FMath::Max(NumSamplers, BindDesc.BindPoint + BindDesc.BindCount);
+					}
+					else if (BindDesc.Type == D3D10_SIT_TEXTURE)
+					{
+						NumSRVs = FMath::Max(NumSRVs, BindDesc.BindPoint + BindDesc.BindCount);
 					}
 
 					// Add a parameter for the texture only, the sampler index will be invalid
@@ -616,6 +625,8 @@ void CompileD3D11Shader(const FShaderCompilerInput& Input,FShaderCompilerOutput&
 							BindDesc.BindPoint,
 							1
 							);
+
+						NumUAVs = FMath::Max(NumUAVs, BindDesc.BindPoint + BindDesc.BindCount);
 					}
 				else if (BindDesc.Type == D3D11_SIT_STRUCTURED || BindDesc.Type == D3D11_SIT_BYTEADDRESS)
 				{
@@ -628,6 +639,8 @@ void CompileD3D11Shader(const FShaderCompilerInput& Input,FShaderCompilerOutput&
 						BindDesc.BindPoint,
 						1
 						);
+
+					NumSRVs = FMath::Max(NumSRVs, BindDesc.BindPoint + BindDesc.BindCount);
 				}
 			}
 
@@ -680,8 +693,12 @@ void CompileD3D11Shader(const FShaderCompilerInput& Input,FShaderCompilerOutput&
 			Ar << SRT;
 			Ar.Serialize( CompressedData->GetBufferPointer(), CompressedData->GetBufferSize() );
 
-			// Pack bGlobalUniformBufferUsed in the last byte
-			Output.Code.Add( bGlobalUniformBufferUsed );
+			// Pack bGlobalUniformBufferUsed and resource counts in the last few bytes
+			Output.Code.Add(bGlobalUniformBufferUsed);
+			Output.Code.Add(NumSamplers);
+			Output.Code.Add(NumSRVs);
+			Output.Code.Add(NumCBs);
+			Output.Code.Add(NumUAVs);
 
 			// Set the number of instructions.
 			Output.NumInstructions = ShaderDesc.InstructionCount;
diff --git a/Engine/Source/Developer/Windows/ShaderFormatD3D/Private/ShaderFormatD3D.cpp b/Engine/Source/Developer/Windows/ShaderFormatD3D/Private/ShaderFormatD3D.cpp
--- a/Engine/Source/Developer/Windows/ShaderFormatD3D/Private/ShaderFormatD3D.cpp
+++ b/Engine/Source/Developer/Windows/ShaderFormatD3D/Private/ShaderFormatD3D.cpp
@@ -16,10 +16,10 @@ class FShaderFormatD3D : public IShaderFormat
 	enum
 	{
 		/** Version for shader format, this becomes part of the DDC key. */
-		UE_SHADER_PCD3D_SM5_VER = 2,
-		UE_SHADER_PCD3D_SM4_VER = 2,
-		UE_SHADER_PCD3D_ES2_VER = 2,
-		UE_SHADER_PCD3D_ES3_1_VER = 1,
+		UE_SHADER_PCD3D_SM5_VER = 3,
+		UE_SHADER_PCD3D_SM4_VER = 3,
+		UE_SHADER_PCD3D_ES2_VER = 3,
+		UE_SHADER_PCD3D_ES3_1_VER = 2,
 	};
 
 	void CheckFormat(FName Format) const
diff --git a/Engine/Source/Runtime/Windows/D3D11RHI/Private/D3D11Shaders.cpp b/Engine/Source/Runtime/Windows/D3D11RHI/Private/D3D11Shaders.cpp
--- a/Engine/Source/Runtime/Windows/D3D11RHI/Private/D3D11Shaders.cpp
+++ b/Engine/Source/Runtime/Windows/D3D11RHI/Private/D3D11Shaders.cpp
@@ -16,12 +16,12 @@ FVertexShaderRHIRef FD3D11DynamicRHI::RHICreateVertexShader(const TArray<uint8>&
 	Ar << Shader->ShaderResourceTable;
 	int32 Offset = Ar.Tell();
 	const uint8* CodePtr = Code.GetData() + Offset;
-	const size_t CodeSize = Code.Num() - Offset - 1;
+	const size_t CodeSize = Code.Num() - Offset - 5;
 
 	VERIFYD3D11RESULT( Direct3DDevice->CreateVertexShader( (void*)CodePtr, CodeSize, NULL, Shader->Resource.GetInitReference() ) );
 	
-	// bGlobalUniformBufferUsed is in the last byte, see CompileD3D11Shader
-	Shader->bShaderNeedsGlobalConstantBuffer = Code[Code.Num() - 1] != 0;
+	// bGlobalUniformBufferUsed and resource counts are packed in the last few bytes, see CompileD3D11Shader
+	Shader->bShaderNeedsGlobalConstantBuffer = Code[Code.Num() - 5] != 0;
 	
 	// TEMP
 	Shader->Code = Code;
@@ -40,12 +40,12 @@ FGeometryShaderRHIRef FD3D11DynamicRHI::RHICreateGeometryShader(const TArray<uin
 	Ar << Shader->ShaderResourceTable;
 	int32 Offset = Ar.Tell();
 	const uint8* CodePtr = Code.GetData() + Offset;
-	const size_t CodeSize = Code.Num() - Offset - 1;
+	const size_t CodeSize = Code.Num() - Offset - 5;
 
 	VERIFYD3D11RESULT( Direct3DDevice->CreateGeometryShader( (void*)CodePtr, CodeSize, NULL, Shader->Resource.GetInitReference() ) );
 	
-	// bGlobalUniformBufferUsed is in the last byte, see CompileD3D11Shader
-	Shader->bShaderNeedsGlobalConstantBuffer = Code[Code.Num() - 1] != 0;
+	// bGlobalUniformBufferUsed and resource counts are packed in the last few bytes, see CompileD3D11Shader
+	Shader->bShaderNeedsGlobalConstantBuffer = Code[Code.Num() - 5] != 0;
 
 	return Shader;
 }
@@ -60,7 +60,7 @@ FGeometryShaderRHIRef FD3D11DynamicRHI::RHICreateGeometryShaderWithStreamOutput(
 	Ar << Shader->ShaderResourceTable;
 	int32 Offset = Ar.Tell();
 	const uint8* CodePtr = Code.GetData() + Offset;
-	const size_t CodeSize = Code.Num() - Offset - 1;
+	const size_t CodeSize = Code.Num() - Offset - 5;
 
 	uint32 D3DRasterizedStream = RasterizedStream;
 	if (RasterizedStream == -1)
@@ -92,8 +92,8 @@ FGeometryShaderRHIRef FD3D11DynamicRHI::RHICreateGeometryShaderWithStreamOutput(
 		Shader->Resource.GetInitReference()
 	) );
 	
-	// bGlobalUniformBufferUsed is in the last byte, see CompileD3D11Shader
-	Shader->bShaderNeedsGlobalConstantBuffer = Code[Code.Num() - 1] != 0;
+	// bGlobalUniformBufferUsed and resource counts are packed in the last few bytes, see CompileD3D11Shader
+	Shader->bShaderNeedsGlobalConstantBuffer = Code[Code.Num() - 5] != 0;
 
 	return Shader;
 }
@@ -108,12 +108,12 @@ FHullShaderRHIRef FD3D11DynamicRHI::RHICreateHullShader(const TArray<uint8>& Cod
 	Ar << Shader->ShaderResourceTable;
 	int32 Offset = Ar.Tell();
 	const uint8* CodePtr = Code.GetData() + Offset;
-	const size_t CodeSize = Code.Num() - Offset - 1;
+	const size_t CodeSize = Code.Num() - Offset - 5;
 
 	VERIFYD3D11RESULT( Direct3DDevice->CreateHullShader( (void*)CodePtr, CodeSize, NULL, Shader->Resource.GetInitReference() ) );
 	
-	// bGlobalUniformBufferUsed is in the last byte, see CompileD3D11Shader
-	Shader->bShaderNeedsGlobalConstantBuffer = Code[Code.Num() - 1] != 0;
+	// bGlobalUniformBufferUsed and resource counts are packed in the last few bytes, see CompileD3D11Shader
+	Shader->bShaderNeedsGlobalConstantBuffer = Code[Code.Num() - 5] != 0;
 
 	return Shader;
 }
@@ -128,12 +128,12 @@ FDomainShaderRHIRef FD3D11DynamicRHI::RHICreateDomainShader(const TArray<uint8>&
 	Ar << Shader->ShaderResourceTable;
 	int32 Offset = Ar.Tell();
 	const uint8* CodePtr = Code.GetData() + Offset;
-	const size_t CodeSize = Code.Num() - Offset - 1;
+	const size_t CodeSize = Code.Num() - Offset - 5;
 
 	VERIFYD3D11RESULT( Direct3DDevice->CreateDomainShader( (void*)CodePtr, CodeSize, NULL, Shader->Resource.GetInitReference() ) );
 	
-	// bGlobalUniformBufferUsed is in the last byte, see CompileD3D11Shader
-	Shader->bShaderNeedsGlobalConstantBuffer = Code[Code.Num() - 1] != 0;
+	// bGlobalUniformBufferUsed and resource counts are packed in the last few bytes, see CompileD3D11Shader
+	Shader->bShaderNeedsGlobalConstantBuffer = Code[Code.Num() - 5] != 0;
 
 	return Shader;
 }
@@ -148,12 +148,12 @@ FPixelShaderRHIRef FD3D11DynamicRHI::RHICreatePixelShader(const TArray<uint8>& C
 	Ar << Shader->ShaderResourceTable;
 	int32 Offset = Ar.Tell();
 	const uint8* CodePtr = Code.GetData() + Offset;
-	const size_t CodeSize = Code.Num() - Offset - 1;
+	const size_t CodeSize = Code.Num() - Offset - 5;
 
 	VERIFYD3D11RESULT( Direct3DDevice->CreatePixelShader( (void*)CodePtr, CodeSize, NULL, Shader->Resource.GetInitReference() ) );
 	
-	// bGlobalUniformBufferUsed is in the last byte, see CompileD3D11Shader
-	Shader->bShaderNeedsGlobalConstantBuffer = Code[Code.Num() - 1] != 0;
+	// bGlobalUniformBufferUsed and resource counts are packed in the last few bytes, see CompileD3D11Shader
+	Shader->bShaderNeedsGlobalConstantBuffer = Code[Code.Num() - 5] != 0;
 
 	return Shader;
 }
@@ -168,12 +168,12 @@ FComputeShaderRHIRef FD3D11DynamicRHI::RHICreateComputeShader(const TArray<uint8
 	Ar << Shader->ShaderResourceTable;
 	int32 Offset = Ar.Tell();
 	const uint8* CodePtr = Code.GetData() + Offset;
-	const size_t CodeSize = Code.Num() - Offset - 1;
+	const size_t CodeSize = Code.Num() - Offset - 5;
 
 	VERIFYD3D11RESULT( Direct3DDevice->CreateComputeShader( (void*)CodePtr, CodeSize, NULL, Shader->Resource.GetInitReference() ) );
 	
-	// bGlobalUniformBufferUsed is in the last byte, see CompileD3D11Shader
-	Shader->bShaderNeedsGlobalConstantBuffer = Code[Code.Num() - 1] != 0;
+	// bGlobalUniformBufferUsed and resource counts are packed in the last few bytes, see CompileD3D11Shader
+	Shader->bShaderNeedsGlobalConstantBuffer = Code[Code.Num() - 5] != 0;
 
 	return Shader;
 }
@@ -218,7 +218,7 @@ FD3D11BoundShaderState::FD3D11BoundShaderState(
 		InVertexDeclaration ? InVertexDeclaration->VertexElements.GetData() : &NullInputElement,
 		InVertexDeclaration ? InVertexDeclaration->VertexElements.Num() : 0,
 		&InVertexShader->Code[ InVertexShader->Offset ],			// TEMP ugly
-		InVertexShader->Code.Num() - 1 - InVertexShader->Offset,
+		InVertexShader->Code.Num() - 5 - InVertexShader->Offset,
 		InputLayout.GetInitReference()
 		));
 
