UE4 - Compile fix for Win XP

diff --git a/Engine/Source/Runtime/Windows/D3D11RHI/Private/D3D11Device.cpp b/Engine/Source/Runtime/Windows/D3D11RHI/Private/D3D11Device.cpp
--- a/Engine/Source/Runtime/Windows/D3D11RHI/Private/D3D11Device.cpp
+++ b/Engine/Source/Runtime/Windows/D3D11RHI/Private/D3D11Device.cpp
@@ -7,6 +7,9 @@
 #include "D3D11RHIPrivate.h"
 #include "AllowWindowsPlatformTypes.h"
 	#include <delayimp.h>
+	#ifndef DXGI_ERROR_NOT_FOUND
+	#define DXGI_ERROR_NOT_FOUND             _HRESULT_TYPEDEF_(0x887A0002L)
+	#endif
 #include "HideWindowsPlatformTypes.h"
 
 
