UE4 - D3D12RHI (not working)

diff --git a/Engine/Source/Programs/UnrealBuildTool/Windows/UEBuildWindows.cs b/Engine/Source/Programs/UnrealBuildTool/Windows/UEBuildWindows.cs
--- a/Engine/Source/Programs/UnrealBuildTool/Windows/UEBuildWindows.cs
+++ b/Engine/Source/Programs/UnrealBuildTool/Windows/UEBuildWindows.cs
@@ -521,7 +521,13 @@ namespace UnrealBuildTool
                     // To enable platform specific D3D11 RHI Types
                     InModule.AddPrivateIncludePath("Runtime/Windows/D3D11RHI/Private/Windows");
                 }
-            }
+
+				if (InModule.ToString() == "D3D12RHI")
+				{
+					// To enable platform specific D3D12 RHI Types
+					InModule.AddPrivateIncludePath("Runtime/Windows/D3D12RHI/Private/Windows");
+				}
+			}
         }
 
         /**
diff --git a/Engine/Source/Runtime/RHI/Private/Windows/WindowsDynamicRHI.cpp b/Engine/Source/Runtime/RHI/Private/Windows/WindowsDynamicRHI.cpp
--- a/Engine/Source/Runtime/RHI/Private/Windows/WindowsDynamicRHI.cpp
+++ b/Engine/Source/Runtime/RHI/Private/Windows/WindowsDynamicRHI.cpp
@@ -10,12 +10,13 @@ FDynamicRHI* PlatformCreateDynamicRHI()
 	// command line overrides
 	const bool bForceD3D11 = FParse::Param(FCommandLine::Get(),TEXT("d3d11")) || FParse::Param(FCommandLine::Get(),TEXT("sm5")) || FParse::Param(FCommandLine::Get(),TEXT("dx11"));
 	const bool bForceD3D10 = FParse::Param(FCommandLine::Get(),TEXT("d3d10")) || FParse::Param(FCommandLine::Get(),TEXT("sm4")) || FParse::Param(FCommandLine::Get(),TEXT("dx10"));
+	const bool bForceD3D12 = FParse::Param(FCommandLine::Get(), TEXT("d3d12")) || FParse::Param(FCommandLine::Get(), TEXT("dx12"));
 	const bool bForceOpenGL = FWindowsPlatformMisc::VerifyWindowsVersion(6, 0) == false || FParse::Param(FCommandLine::Get(), TEXT("opengl")) || FParse::Param(FCommandLine::Get(), TEXT("opengl3")) || FParse::Param(FCommandLine::Get(), TEXT("opengl4"));
 	const bool bForceVulkan = FParse::Param(FCommandLine::Get(), TEXT("vulkan"));
 
-	if (((bForceD3D11 ? 1 : 0) + (bForceD3D10 ? 1 : 0) + (bForceOpenGL ? 1 : 0) + (bForceVulkan ? 1 : 0)) > 1)
+	if (((bForceD3D12 ? 1 : 0) + (bForceD3D11 ? 1 : 0) + (bForceD3D10 ? 1 : 0) + (bForceOpenGL ? 1 : 0) + (bForceVulkan ? 1 : 0)) > 1)
 	{
-		UE_LOG(LogRHI, Fatal,TEXT("-d3d11, -d3d10, -vulkan and -opengl[3|4] mutually exclusive options, but more than one was specified on the command-line."));
+		UE_LOG(LogRHI, Fatal,TEXT("-d3d12, -d3d11, -d3d10, -vulkan, and -opengl[3|4] mutually exclusive options, but more than one was specified on the command-line."));
 	}
 
 	// Load the dynamic RHI module.
@@ -41,6 +42,21 @@ FDynamicRHI* PlatformCreateDynamicRHI()
 			DynamicRHIModule = NULL;
 		}
 	}
+	else if (bForceD3D12)
+	{
+		DynamicRHIModule = &FModuleManager::LoadModuleChecked<IDynamicRHIModule>(TEXT("D3D12RHI"));
+
+		if (!DynamicRHIModule->IsSupported())
+		{
+			FMessageDialog::Open(EAppMsgType::Ok, NSLOCTEXT("WindowsDynamicRHI", "RequiredDX12", "DX12 is not supported on your system. Try running without the -dx12 or -d3d12 command line argument."));
+			FPlatformMisc::RequestExit(1);
+			DynamicRHIModule = NULL;
+		}
+		else if (FPlatformProcess::IsApplicationRunning(TEXT("fraps.exe")))
+		{
+			FMessageDialog::Open(EAppMsgType::Ok, NSLOCTEXT("WindowsDynamicRHI", "UseExpressionEncoder", "Fraps has been known to crash D3D12. Please use Microsoft Expression Encoder instead for capturing."));
+		}
+	}
 	else
 	{
 		DynamicRHIModule = &FModuleManager::LoadModuleChecked<IDynamicRHIModule>(TEXT("D3D11RHI"));
@@ -51,9 +67,9 @@ FDynamicRHI* PlatformCreateDynamicRHI()
 			FPlatformMisc::RequestExit(1);
 			DynamicRHIModule = NULL;
 		}
-		else if( FPlatformProcess::IsApplicationRunning( TEXT("fraps.exe") ) )
+		else if (FPlatformProcess::IsApplicationRunning(TEXT("fraps.exe")))
 		{
-			FMessageDialog::Open( EAppMsgType::Ok, NSLOCTEXT("WindowsDynamicRHI", "UseExpressionEncoder", "Fraps has been known to crash D3D11. Please use Microsoft Expression Encoder instead for capturing."));
+			FMessageDialog::Open(EAppMsgType::Ok, NSLOCTEXT("WindowsDynamicRHI", "UseExpressionEncoder", "Fraps has been known to crash D3D11. Please use Microsoft Expression Encoder instead for capturing."));
 		}
 	}
 
diff --git a/Engine/Source/Runtime/Windows/D3D12RHI/D3D12RHI.Build.cs b/Engine/Source/Runtime/Windows/D3D12RHI/D3D12RHI.Build.cs
new file mode 100644
--- /dev/null
+++ b/Engine/Source/Runtime/Windows/D3D12RHI/D3D12RHI.Build.cs
@@ -0,0 +1,38 @@
+// Copyright 1998-2014 Epic Games, Inc. All Rights Reserved.
+
+using UnrealBuildTool;
+
+public class D3D12RHI : ModuleRules
+{
+	public D3D12RHI(TargetInfo Target)
+	{
+		BinariesSubFolder = "NotForLicensees";
+		PrivateIncludePaths.Add("Runtime/Windows/D3D12RHI/Private");
+
+		PrivateDependencyModuleNames.AddRange(
+			new string[] {
+				"Core",
+				"Engine",
+				"RHI",
+				"RenderCore",
+				"ShaderCore",
+                "UtilityShaders",
+			}
+			);
+
+        AddThirdPartyPrivateStaticDependencies(Target, "DX12");
+		AddThirdPartyPrivateStaticDependencies(Target, "DX11");
+        AddThirdPartyPrivateStaticDependencies(Target, "NVAPI");
+        AddThirdPartyPrivateStaticDependencies(Target, "AMD");
+        
+		if (Target.Configuration != UnrealTargetConfiguration.Shipping)
+		{
+			PrivateIncludePathModuleNames.AddRange(new string[] { "TaskGraph" });
+		}
+
+        Definitions.Add("SUB_ALLOCATED_DEFAULT_ALLOCATIONS=1");
+
+		// Not fully implemented yet.
+        Definitions.Add("SUPPORTS_MEMORY_RESIDENCY=0");
+	}
+}
diff --git a/Engine/Source/Runtime/Windows/D3D12RHI/Private/D3D12Commands.cpp b/Engine/Source/Runtime/Windows/D3D12RHI/Private/D3D12Commands.cpp
new file mode 100644
--- /dev/null
+++ b/Engine/Source/Runtime/Windows/D3D12RHI/Private/D3D12Commands.cpp
@@ -0,0 +1,2233 @@
+// Copyright 1998-2014 Epic Games, Inc. All Rights Reserved.
+
+/*=============================================================================
+D3D12Commands.cpp: D3D RHI commands implementation.
+=============================================================================*/
+
+#include "D3D12RHIPrivate.h"
+#if WITH_D3DX_LIBS
+#include "AllowWindowsPlatformTypes.h"
+	#include <xnamath.h>
+#include "HideWindowsPlatformTypes.h"
+#endif
+#include "D3D12RHIPrivateUtil.h"
+#include "StaticBoundShaderState.h"
+#include "GlobalShader.h"
+#include "OneColorShader.h"
+#include "RHICommandList.h"
+#include "RHIStaticStates.h"
+#include "ShaderParameterUtils.h"
+#include "ScreenRendering.h"
+
+
+// MSFT: Seb: Fix up these D3D11 names and remove namespace
+namespace D3D12RHI
+{
+	FGlobalBoundShaderState GD3D11ClearMRTBoundShaderState[8];
+	TGlobalResource<FVector4VertexDeclaration> GD3D11Vector4VertexDeclaration;
+
+	// MSFT: Seb: Do we need this here?
+	//FGlobalBoundShaderState ResolveBoundShaderState;
+}
+using namespace D3D12RHI;
+
+#define DECLARE_ISBOUNDSHADER(ShaderType) inline void ValidateBoundShader(FD3D12StateCache& InStateCache, F##ShaderType##RHIParamRef ShaderType##RHI) \
+{ \
+	FD3D12##ShaderType* CachedShader; \
+	InStateCache.Get##ShaderType(&CachedShader); \
+	FD3D12##ShaderType* ShaderType = FD3D12DynamicRHI::ResourceCast(ShaderType##RHI); \
+	ensureMsgf(CachedShader == ShaderType, TEXT("Parameters are being set for a %s which is not currently bound"), TEXT( #ShaderType )); \
+}
+
+DECLARE_ISBOUNDSHADER(VertexShader)
+DECLARE_ISBOUNDSHADER(PixelShader)
+DECLARE_ISBOUNDSHADER(GeometryShader)
+DECLARE_ISBOUNDSHADER(HullShader)
+DECLARE_ISBOUNDSHADER(DomainShader)
+DECLARE_ISBOUNDSHADER(ComputeShader)
+
+#if DO_CHECK
+#define VALIDATE_BOUND_SHADER(s) ValidateBoundShader(StateCache, s)
+#else
+#define VALIDATE_BOUND_SHADER(s)
+#endif
+
+#define WITH_GPA (!PLATFORM_XBOXONE)
+#if WITH_GPA
+#define GPA_WINDOWS 1
+#include <GPUPerfAPI/Gpa.h>
+#endif
+
+void FD3D12DynamicRHI::SetupRecursiveResources()
+{
+	extern ENGINE_API FShaderCompilingManager* GShaderCompilingManager;
+	check (FPlatformProperties::RequiresCookedData() || GShaderCompilingManager);
+
+	FRHICommandList_RecursiveHazardous RHICmdList(RHIGetDefaultContext());
+	auto ShaderMap = GetGlobalShaderMap(GMaxRHIFeatureLevel);
+	TShaderMapRef<TOneColorVS<true> > VertexShader(ShaderMap);
+	GD3D11Vector4VertexDeclaration.InitRHI();
+
+	for (int32 NumBuffers = 1; NumBuffers <= MaxSimultaneousRenderTargets; NumBuffers++)
+	{
+		FOneColorPS* PixelShader = NULL;
+
+		if (NumBuffers <= 1)
+		{
+			TShaderMapRef<TOneColorPixelShaderMRT<1> > MRTPixelShader(ShaderMap);
+			PixelShader = *MRTPixelShader;
+		}
+		else if (NumBuffers == 2)
+		{
+			TShaderMapRef<TOneColorPixelShaderMRT<2> > MRTPixelShader(ShaderMap);
+			PixelShader = *MRTPixelShader;
+		}
+		else if (NumBuffers == 3)
+		{
+			TShaderMapRef<TOneColorPixelShaderMRT<3> > MRTPixelShader(ShaderMap);
+			PixelShader = *MRTPixelShader;
+		}
+		else if (NumBuffers == 4)
+		{
+			TShaderMapRef<TOneColorPixelShaderMRT<4> > MRTPixelShader(ShaderMap);
+			PixelShader = *MRTPixelShader;
+		}
+		else if (NumBuffers == 5)
+		{
+			TShaderMapRef<TOneColorPixelShaderMRT<5> > MRTPixelShader(ShaderMap);
+			PixelShader = *MRTPixelShader;
+		}
+		else if (NumBuffers == 6)
+		{
+			TShaderMapRef<TOneColorPixelShaderMRT<6> > MRTPixelShader(ShaderMap);
+			PixelShader = *MRTPixelShader;
+		}
+		else if (NumBuffers == 7)
+		{
+			TShaderMapRef<TOneColorPixelShaderMRT<7> > MRTPixelShader(ShaderMap);
+			PixelShader = *MRTPixelShader;
+		}
+		else if (NumBuffers == 8)
+		{
+			TShaderMapRef<TOneColorPixelShaderMRT<8> > MRTPixelShader(ShaderMap);
+			PixelShader = *MRTPixelShader;
+		}
+
+		SetGlobalBoundShaderState(RHICmdList, GMaxRHIFeatureLevel, GD3D11ClearMRTBoundShaderState[NumBuffers - 1], GD3D11Vector4VertexDeclaration.VertexDeclarationRHI, *VertexShader, PixelShader);
+	}
+
+	// MSFT: Seb: Is this needed?
+	//extern ENGINE_API TGlobalResource<FScreenVertexDeclaration> GScreenVertexDeclaration;
+
+	//TShaderMapRef<FD3D12RHIResolveVS> ResolveVertexShader(ShaderMap);
+	//TShaderMapRef<FD3D12RHIResolveDepthNonMSPS> ResolvePixelShader(ShaderMap);
+
+	//SetGlobalBoundShaderState(RHICmdList, GMaxRHIFeatureLevel, ResolveBoundShaderState, GScreenVertexDeclaration.VertexDeclarationRHI, *ResolveVertexShader, *ResolvePixelShader);
+}
+
+void FD3D12DynamicRHI::RHIGpuTimeBegin(uint32 Hash, bool bCompute)
+{
+#if WITH_GPA
+	char Str[256];
+	if (GpaBegin(Str, Hash, bCompute, (void*)GetRHIDevice()->GetDevice()))
+	{
+		OutputDebugStringA(Str);
+	}
+#endif
+}
+
+void FD3D12DynamicRHI::RHIGpuTimeEnd(uint32 Hash, bool bCompute)
+{
+#if WITH_GPA
+	GpaEnd(Hash, bCompute);
+#endif
+}
+
+
+// Vertex state.
+void FD3D12CommandContext::RHISetStreamSource(uint32 StreamIndex, FVertexBufferRHIParamRef VertexBufferRHI, uint32 Stride, uint32 Offset)
+{
+	FD3D12VertexBuffer* VertexBuffer = FD3D12DynamicRHI::ResourceCast(VertexBufferRHI);
+
+	StateCache.SetStreamSource(VertexBuffer ? VertexBuffer->ResourceLocation.GetReference() : nullptr, StreamIndex, Stride, Offset);
+}
+
+// Stream-Out state.
+void FD3D12DynamicRHI::RHISetStreamOutTargets(uint32 NumTargets, const FVertexBufferRHIParamRef* VertexBuffers, const uint32* Offsets)
+{
+	FD3D12CommandContext& CmdContext = GetRHIDevice()->GetDefaultCommandContext();
+	FD3D12Resource* D3DVertexBuffers[D3D12_SO_BUFFER_SLOT_COUNT] = { 0 };
+	uint32 D3DOffsets[D3D12_SO_BUFFER_SLOT_COUNT] = { 0 };
+
+	if (VertexBuffers)
+	{
+		for (uint32 BufferIndex = 0; BufferIndex < NumTargets; BufferIndex++)
+		{
+			D3DVertexBuffers[BufferIndex] = VertexBuffers[BufferIndex] ? ((FD3D12VertexBuffer*)VertexBuffers[BufferIndex])->ResourceLocation->GetResource() : NULL;
+			D3DOffsets[BufferIndex] = Offsets[BufferIndex] + (VertexBuffers[BufferIndex] ? ((FD3D12VertexBuffer*)VertexBuffers[BufferIndex])->ResourceLocation->GetOffset() : 0);
+		}
+	}
+
+	CmdContext.StateCache.SetStreamOutTargets(NumTargets, D3DVertexBuffers, D3DOffsets);
+}
+
+// Rasterizer state.
+void FD3D12CommandContext::RHISetRasterizerState(FRasterizerStateRHIParamRef NewStateRHI)
+{
+	FD3D12RasterizerState* NewState = FD3D12DynamicRHI::ResourceCast(NewStateRHI);
+	StateCache.SetRasterizerState(&NewState->Desc);
+}
+
+void FD3D12CommandContext::RHISetComputeShader(FComputeShaderRHIParamRef ComputeShaderRHI)
+{
+	FD3D12ComputeShader* ComputeShader = FD3D12DynamicRHI::ResourceCast(ComputeShaderRHI);
+	SetCurrentComputeShader(ComputeShaderRHI);
+}
+
+void FD3D12CommandContext::RHIDispatchComputeShader(uint32 ThreadGroupCountX, uint32 ThreadGroupCountY, uint32 ThreadGroupCountZ)
+{ 
+	FComputeShaderRHIParamRef ComputeShaderRHI = GetCurrentComputeShader();
+	FD3D12ComputeShader* ComputeShader = FD3D12DynamicRHI::ResourceCast(ComputeShaderRHI);
+
+	StateCache.SetComputeShader(ComputeShader);
+
+	OwningRHI.RegisterGPUWork(1);
+
+	if (ComputeShader->bShaderNeedsGlobalConstantBuffer)
+	{
+		CommitComputeShaderConstants();
+	}
+	CommitComputeResourceTables(ComputeShader);
+	StateCache.ApplyState(true);
+
+	numDispatches++;
+	CommandListHandle->Dispatch(ThreadGroupCountX, ThreadGroupCountY, ThreadGroupCountZ);
+
+	StateCache.FlushComputeShaderCache();
+
+	DEBUG_EXECUTE_COMMAND_LIST(this);
+
+	StateCache.SetComputeShader(nullptr);
+}
+
+void FD3D12CommandContext::RHIDispatchIndirectComputeShader(FVertexBufferRHIParamRef ArgumentBufferRHI, uint32 ArgumentOffset)
+{ 
+	FComputeShaderRHIParamRef ComputeShaderRHI = GetCurrentComputeShader();
+	FD3D12ComputeShader* ComputeShader = FD3D12DynamicRHI::ResourceCast(ComputeShaderRHI);
+	FD3D12VertexBuffer* ArgumentBuffer = FD3D12DynamicRHI::ResourceCast(ArgumentBufferRHI);
+
+	OwningRHI.RegisterGPUWork(1);
+
+	StateCache.SetComputeShader(ComputeShader);
+
+	if (ComputeShader->bShaderNeedsGlobalConstantBuffer)
+	{
+		CommitComputeShaderConstants();
+	}
+	CommitComputeResourceTables(ComputeShader);
+	StateCache.ApplyState(true);
+
+	FD3D12DynamicRHI::TransitionResource(CommandListHandle, ArgumentBuffer->ResourceLocation->GetResource(), D3D12_RESOURCE_STATE_INDIRECT_ARGUMENT, D3D12_RESOURCE_BARRIER_TYPE_TRANSITION);
+
+	numDispatches++;
+	CommandListHandle->ExecuteIndirect(
+		GetParentDevice()->GetDispatchIndirectCommandSignature(),
+		1,
+		ArgumentBuffer->ResourceLocation->GetResource()->GetResource(),
+		ArgumentBuffer->ResourceLocation->GetOffset() + ArgumentOffset * 20,
+		NULL,
+		0
+		);
+
+	StateCache.FlushComputeShaderCache();
+
+	DEBUG_EXECUTE_COMMAND_LIST(this);
+
+	StateCache.SetComputeShader(nullptr);
+}
+
+void FD3D12CommandContext::RHISetViewport(uint32 MinX, uint32 MinY, float MinZ, uint32 MaxX, uint32 MaxY, float MaxZ)
+{
+	// These are the maximum viewport extents for D3D12. Exceeding them leads to badness.
+	check(MinX <= (uint32)D3D12_VIEWPORT_BOUNDS_MAX);
+	check(MinY <= (uint32)D3D12_VIEWPORT_BOUNDS_MAX);
+	check(MaxX <= (uint32)D3D12_VIEWPORT_BOUNDS_MAX);
+	check(MaxY <= (uint32)D3D12_VIEWPORT_BOUNDS_MAX);
+
+	D3D12_VIEWPORT Viewport = { MinX, MinY, MaxX - MinX, MaxY - MinY, MinZ, MaxZ };
+	//avoid setting a 0 extent viewport, which the debug runtime doesn't like
+	if (Viewport.Width > 0 && Viewport.Height > 0)
+	{
+		StateCache.SetViewport(Viewport);
+		SetScissorRectIfRequiredWhenSettingViewport(MinX, MinY, MaxX, MaxY);
+	}
+}
+
+void FD3D12CommandContext::RHISetScissorRect(bool bEnable, uint32 MinX, uint32 MinY, uint32 MaxX, uint32 MaxY)
+{
+	if(bEnable)
+	{
+		D3D12_RECT ScissorRect;
+		ScissorRect.left = MinX;
+		ScissorRect.right = MaxX;
+		ScissorRect.top = MinY;
+		ScissorRect.bottom = MaxY;
+		StateCache.SetScissorRects(1, &ScissorRect);
+	}
+	else
+	{
+		D3D12_RECT ScissorRect;
+		ScissorRect.left = 0;
+		ScissorRect.right = GetMax2DTextureDimension();
+		ScissorRect.top = 0;
+		ScissorRect.bottom = GetMax2DTextureDimension();
+		StateCache.SetScissorRects(1, &ScissorRect);
+	}
+}
+
+/**
+* Set bound shader state. This will set the vertex decl/shader, and pixel shader
+* @param BoundShaderState - state resource
+*/
+void FD3D12CommandContext::RHISetBoundShaderState(FBoundShaderStateRHIParamRef BoundShaderStateRHI)
+{
+	SCOPE_CYCLE_COUNTER(STAT_D3D12SetShaderUniformBuffer);SCOPE_CYCLE_COUNTER(STAT_D3D12SetBoundShaderState);
+	FD3D12BoundShaderState* BoundShaderState = FD3D12DynamicRHI::ResourceCast(BoundShaderStateRHI);
+
+	StateCache.SetBoundShaderState(BoundShaderState);
+
+	if(BoundShaderState->GetHullShader() && BoundShaderState->GetDomainShader())
+	{
+		bUsingTessellation = true;
+	}
+	else
+	{
+		bUsingTessellation = false;
+	}
+
+	// @TODO : really should only discard the constants if the shader state has actually changed.
+	bDiscardSharedConstants = true;
+
+	CurrentBoundShaderState = BoundShaderState;
+
+	// Prevent transient bound shader states from being recreated for each use by keeping a history of the most recently used bound shader states.
+	// The history keeps them alive, and the bound shader state cache allows them to be reused if needed.
+	// MSFT: JSTANARD:  Is this still relevant?
+	OwningRHI.AddBoundShaderState(BoundShaderState);
+
+	// Shader changed so all resource tables are dirty
+	DirtyUniformBuffers[SF_Vertex] = 0xffff;
+	DirtyUniformBuffers[SF_Pixel] = 0xffff;
+	DirtyUniformBuffers[SF_Hull] = 0xffff;
+	DirtyUniformBuffers[SF_Domain] = 0xffff;
+	DirtyUniformBuffers[SF_Geometry] = 0xffff;
+
+	// To avoid putting bad samplers into the descriptor heap
+	// Clear all sampler & SRV bindings here
+	StateCache.ClearSamplers();
+	StateCache.ClearSRVs();
+}
+
+void FD3D12CommandContext::RHISetShaderTexture(FVertexShaderRHIParamRef VertexShaderRHI, uint32 TextureIndex, FTextureRHIParamRef NewTextureRHI)
+{
+	uint32 Start = FPlatformTime::Cycles();
+
+	VALIDATE_BOUND_SHADER(VertexShaderRHI);
+
+	FD3D12TextureBase* NewTexture = GetD3D11TextureFromRHITexture(NewTextureRHI);
+	FD3D12ShaderResourceView* ShaderResourceView = NewTexture ? NewTexture->GetShaderResourceView() : NULL;
+	FD3D12ResourceLocation* ResourceLocation = NewTexture ? NewTexture->ResourceLocation : nullptr;
+
+	if (  ( NewTexture == NULL) || ( NewTexture->GetRenderTargetView( 0, 0 ) !=NULL) || ( NewTexture->HasDepthStencilView()) )
+		SetShaderResourceView<SF_Vertex>(ResourceLocation, ShaderResourceView, TextureIndex, FD3D12StateCache::SRV_Dynamic);
+	else
+		SetShaderResourceView<SF_Vertex>(ResourceLocation, ShaderResourceView, TextureIndex, FD3D12StateCache::SRV_Static);
+
+	OwningRHI.IncrementSetShaderTextureCycles(FPlatformTime::Cycles() - Start);
+	OwningRHI.IncrementSetShaderTextureCalls();
+}
+
+void FD3D12CommandContext::RHISetShaderTexture(FHullShaderRHIParamRef HullShaderRHI, uint32 TextureIndex, FTextureRHIParamRef NewTextureRHI)
+{
+	uint32 Start = FPlatformTime::Cycles();
+	 
+	VALIDATE_BOUND_SHADER(HullShaderRHI);
+
+	FD3D12TextureBase* NewTexture = GetD3D11TextureFromRHITexture(NewTextureRHI);
+	FD3D12ShaderResourceView* ShaderResourceView = NewTexture ? NewTexture->GetShaderResourceView() : NULL;
+	FD3D12ResourceLocation* ResourceLocation = NewTexture ? NewTexture->ResourceLocation : nullptr;
+
+	if (  ( NewTexture == NULL) || ( NewTexture->GetRenderTargetView( 0, 0 ) !=NULL) || ( NewTexture->HasDepthStencilView()) )
+		SetShaderResourceView<SF_Hull>(ResourceLocation, ShaderResourceView, TextureIndex, FD3D12StateCache::SRV_Dynamic);
+	else
+		SetShaderResourceView<SF_Hull>(ResourceLocation, ShaderResourceView, TextureIndex, FD3D12StateCache::SRV_Static);
+
+	OwningRHI.IncrementSetShaderTextureCycles(FPlatformTime::Cycles() - Start);
+	OwningRHI.IncrementSetShaderTextureCalls();
+}
+
+void FD3D12CommandContext::RHISetShaderTexture(FDomainShaderRHIParamRef DomainShaderRHI, uint32 TextureIndex, FTextureRHIParamRef NewTextureRHI)
+{
+	uint32 Start = FPlatformTime::Cycles();
+
+	VALIDATE_BOUND_SHADER(DomainShaderRHI);
+
+	FD3D12TextureBase* NewTexture = GetD3D11TextureFromRHITexture(NewTextureRHI);
+	FD3D12ShaderResourceView* ShaderResourceView = NewTexture ? NewTexture->GetShaderResourceView() : NULL;
+	FD3D12ResourceLocation* ResourceLocation = NewTexture ? NewTexture->ResourceLocation : nullptr;
+
+	if (  ( NewTexture == NULL) || ( NewTexture->GetRenderTargetView( 0, 0 ) !=NULL) || ( NewTexture->HasDepthStencilView()) )
+		SetShaderResourceView<SF_Domain>(ResourceLocation, ShaderResourceView, TextureIndex, FD3D12StateCache::SRV_Dynamic);
+	else
+		SetShaderResourceView<SF_Domain>(ResourceLocation, ShaderResourceView, TextureIndex, FD3D12StateCache::SRV_Static);
+
+	OwningRHI.IncrementSetShaderTextureCycles(FPlatformTime::Cycles() - Start);
+	OwningRHI.IncrementSetShaderTextureCalls();
+}
+
+void FD3D12CommandContext::RHISetShaderTexture(FGeometryShaderRHIParamRef GeometryShaderRHI, uint32 TextureIndex, FTextureRHIParamRef NewTextureRHI)
+{
+	uint32 Start = FPlatformTime::Cycles();
+
+	VALIDATE_BOUND_SHADER(GeometryShaderRHI);
+
+	FD3D12TextureBase* NewTexture = GetD3D11TextureFromRHITexture(NewTextureRHI);
+	FD3D12ShaderResourceView* ShaderResourceView = NewTexture ? NewTexture->GetShaderResourceView() : NULL;
+	FD3D12ResourceLocation* ResourceLocation = NewTexture ? NewTexture->ResourceLocation : nullptr;
+
+	if (  ( NewTexture == NULL) || ( NewTexture->GetRenderTargetView( 0, 0 ) !=NULL) || ( NewTexture->HasDepthStencilView()) )
+		SetShaderResourceView<SF_Geometry>(ResourceLocation, ShaderResourceView, TextureIndex, FD3D12StateCache::SRV_Dynamic);
+	else
+		SetShaderResourceView<SF_Geometry>(ResourceLocation, ShaderResourceView, TextureIndex, FD3D12StateCache::SRV_Static);
+
+	OwningRHI.IncrementSetShaderTextureCycles(FPlatformTime::Cycles() - Start);
+	OwningRHI.IncrementSetShaderTextureCalls();
+}
+
+void FD3D12CommandContext::RHISetShaderTexture(FPixelShaderRHIParamRef PixelShaderRHI, uint32 TextureIndex, FTextureRHIParamRef NewTextureRHI)
+{
+	uint32 Start = FPlatformTime::Cycles();
+
+	VALIDATE_BOUND_SHADER(PixelShaderRHI);
+
+	FD3D12TextureBase* NewTexture = GetD3D11TextureFromRHITexture(NewTextureRHI);
+	FD3D12ShaderResourceView* ShaderResourceView = NewTexture ? NewTexture->GetShaderResourceView() : NULL;
+	FD3D12ResourceLocation* ResourceLocation = NewTexture ? NewTexture->ResourceLocation : nullptr;
+
+	if ( ( NewTexture == NULL) ||  ( NewTexture->GetRenderTargetView( 0, 0 ) !=NULL) || ( NewTexture->HasDepthStencilView()) )
+		SetShaderResourceView<SF_Pixel>(ResourceLocation, ShaderResourceView, TextureIndex, FD3D12StateCache::SRV_Dynamic);
+	else
+		SetShaderResourceView<SF_Pixel>(ResourceLocation, ShaderResourceView, TextureIndex, FD3D12StateCache::SRV_Static);
+
+	OwningRHI.IncrementSetShaderTextureCycles(FPlatformTime::Cycles() - Start);
+	OwningRHI.IncrementSetShaderTextureCalls();
+}
+
+void FD3D12CommandContext::RHISetShaderTexture(FComputeShaderRHIParamRef ComputeShaderRHI, uint32 TextureIndex, FTextureRHIParamRef NewTextureRHI)
+{
+	uint32 Start = FPlatformTime::Cycles();
+
+	//VALIDATE_BOUND_SHADER(ComputeShaderRHI);
+
+	FD3D12TextureBase* NewTexture = GetD3D11TextureFromRHITexture(NewTextureRHI);
+	FD3D12ShaderResourceView* ShaderResourceView = NewTexture ? NewTexture->GetShaderResourceView() : NULL;
+	FD3D12ResourceLocation* ResourceLocation = NewTexture ? NewTexture->ResourceLocation : nullptr;
+
+	if ( ( NewTexture == NULL) || ( NewTexture->GetRenderTargetView( 0, 0 ) !=NULL) || ( NewTexture->HasDepthStencilView()) )
+		SetShaderResourceView<SF_Compute>(ResourceLocation, ShaderResourceView, TextureIndex, FD3D12StateCache::SRV_Dynamic);
+	else
+		SetShaderResourceView<SF_Compute>(ResourceLocation, ShaderResourceView, TextureIndex, FD3D12StateCache::SRV_Static);
+
+	OwningRHI.IncrementSetShaderTextureCycles(FPlatformTime::Cycles() - Start);
+	OwningRHI.IncrementSetShaderTextureCalls();
+}
+
+void FD3D12CommandContext::RHISetUAVParameter(FComputeShaderRHIParamRef ComputeShaderRHI, uint32 UAVIndex, FUnorderedAccessViewRHIParamRef UAVRHI)
+{
+	//VALIDATE_BOUND_SHADER(ComputeShaderRHI);
+
+	FD3D12UnorderedAccessView* UAV = FD3D12DynamicRHI::ResourceCast(UAVRHI);
+
+	if(UAV)
+	{
+		ConditionalClearShaderResource(UAV->GetResourceLocation());
+	}
+
+	uint32 InitialCount = -1;
+
+	// Actually set the UAV
+	StateCache.SetUAVs(SF_Compute, UAVIndex, 1, &UAV, &InitialCount);
+
+}
+
+void FD3D12CommandContext::RHISetUAVParameter(FComputeShaderRHIParamRef ComputeShaderRHI, uint32 UAVIndex, FUnorderedAccessViewRHIParamRef UAVRHI, uint32 InitialCount)
+{
+	//VALIDATE_BOUND_SHADER(ComputeShaderRHI);
+
+	FD3D12UnorderedAccessView* UAV = FD3D12DynamicRHI::ResourceCast(UAVRHI);
+
+	if(UAV)
+	{
+		ConditionalClearShaderResource(UAV->GetResourceLocation());
+	}
+
+	StateCache.SetUAVs(SF_Compute, UAVIndex, 1, &UAV, &InitialCount);
+
+}
+
+void FD3D12CommandContext::RHISetShaderResourceViewParameter(FPixelShaderRHIParamRef PixelShaderRHI, uint32 TextureIndex, FShaderResourceViewRHIParamRef SRVRHI)
+{
+	VALIDATE_BOUND_SHADER(PixelShaderRHI);
+
+	FD3D12ShaderResourceView* SRV = FD3D12DynamicRHI::ResourceCast(SRVRHI);
+
+	FD3D12ResourceLocation* ResourceLocation = nullptr;
+	FD3D12ShaderResourceView* D3D12SRV = nullptr;
+
+	if (SRV)
+	{
+		ResourceLocation = SRV->GetResourceLocation();
+		D3D12SRV = SRV;
+	}
+
+	SetShaderResourceView<SF_Pixel>(ResourceLocation, D3D12SRV, TextureIndex);
+}
+
+void FD3D12CommandContext::RHISetShaderResourceViewParameter(FVertexShaderRHIParamRef VertexShaderRHI, uint32 TextureIndex, FShaderResourceViewRHIParamRef SRVRHI)
+{
+	VALIDATE_BOUND_SHADER(VertexShaderRHI);
+
+	FD3D12ShaderResourceView* SRV = FD3D12DynamicRHI::ResourceCast(SRVRHI);
+
+	FD3D12ResourceLocation* ResourceLocation = nullptr;
+	FD3D12ShaderResourceView* D3D12SRV = nullptr;
+
+	if (SRV)
+	{
+		ResourceLocation = SRV->GetResourceLocation();
+		D3D12SRV = SRV;
+	}
+
+	SetShaderResourceView<SF_Vertex>(ResourceLocation, D3D12SRV, TextureIndex);
+}
+
+void FD3D12CommandContext::RHISetShaderResourceViewParameter(FComputeShaderRHIParamRef ComputeShaderRHI, uint32 TextureIndex, FShaderResourceViewRHIParamRef SRVRHI)
+{
+	//VALIDATE_BOUND_SHADER(ComputeShaderRHI);
+
+	FD3D12ShaderResourceView* SRV = FD3D12DynamicRHI::ResourceCast(SRVRHI);
+
+	FD3D12ResourceLocation* ResourceLocation = nullptr;
+	FD3D12ShaderResourceView* D3D12SRV = nullptr;
+
+	if (SRV)
+	{
+		ResourceLocation = SRV->GetResourceLocation();
+		D3D12SRV = SRV;
+	}
+
+	SetShaderResourceView<SF_Compute>(ResourceLocation, D3D12SRV, TextureIndex);
+}
+
+void FD3D12CommandContext::RHISetShaderResourceViewParameter(FHullShaderRHIParamRef HullShaderRHI, uint32 TextureIndex, FShaderResourceViewRHIParamRef SRVRHI)
+{
+	VALIDATE_BOUND_SHADER(HullShaderRHI);
+
+	FD3D12ShaderResourceView* SRV = FD3D12DynamicRHI::ResourceCast(SRVRHI);
+
+	FD3D12ResourceLocation* ResourceLocation = nullptr;
+	FD3D12ShaderResourceView* D3D12SRV = nullptr;
+
+	if (SRV)
+	{
+		ResourceLocation = SRV->GetResourceLocation();
+		D3D12SRV = SRV;
+	}
+
+	SetShaderResourceView<SF_Hull>(ResourceLocation, D3D12SRV, TextureIndex);
+}
+
+void FD3D12CommandContext::RHISetShaderResourceViewParameter(FDomainShaderRHIParamRef DomainShaderRHI, uint32 TextureIndex, FShaderResourceViewRHIParamRef SRVRHI)
+{
+	VALIDATE_BOUND_SHADER(DomainShaderRHI);
+
+	FD3D12ShaderResourceView* SRV = FD3D12DynamicRHI::ResourceCast(SRVRHI);
+
+	FD3D12ResourceLocation* ResourceLocation = nullptr;
+	FD3D12ShaderResourceView* D3D12SRV = nullptr;
+
+	if (SRV)
+	{
+		ResourceLocation = SRV->GetResourceLocation();
+		D3D12SRV = SRV;
+	}
+
+	SetShaderResourceView<SF_Domain>(ResourceLocation, D3D12SRV, TextureIndex);
+}
+
+void FD3D12CommandContext::RHISetShaderResourceViewParameter(FGeometryShaderRHIParamRef GeometryShaderRHI, uint32 TextureIndex, FShaderResourceViewRHIParamRef SRVRHI)
+{
+	VALIDATE_BOUND_SHADER(GeometryShaderRHI);
+
+	FD3D12ShaderResourceView* SRV = FD3D12DynamicRHI::ResourceCast(SRVRHI);
+
+	FD3D12ResourceLocation* ResourceLocation = nullptr;
+	FD3D12ShaderResourceView* D3D12SRV = nullptr;
+
+	if (SRV)
+	{
+		ResourceLocation = SRV->GetResourceLocation();
+		D3D12SRV = SRV;
+	}
+
+	SetShaderResourceView<SF_Geometry>(ResourceLocation, D3D12SRV, TextureIndex);
+}
+
+void FD3D12CommandContext::RHISetShaderSampler(FVertexShaderRHIParamRef VertexShaderRHI, uint32 SamplerIndex, FSamplerStateRHIParamRef NewStateRHI)
+{
+	VALIDATE_BOUND_SHADER(VertexShaderRHI);
+
+	FD3D12VertexShader* VertexShader = FD3D12DynamicRHI::ResourceCast(VertexShaderRHI);
+	FD3D12SamplerState* NewState = FD3D12DynamicRHI::ResourceCast(NewStateRHI);
+
+	StateCache.SetSamplerState<SF_Vertex>(NewState, SamplerIndex);
+}
+
+void FD3D12CommandContext::RHISetShaderSampler(FHullShaderRHIParamRef HullShaderRHI, uint32 SamplerIndex, FSamplerStateRHIParamRef NewStateRHI)
+{
+	VALIDATE_BOUND_SHADER(HullShaderRHI);
+
+	FD3D12HullShader* HullShader = FD3D12DynamicRHI::ResourceCast(HullShaderRHI);
+	FD3D12SamplerState* NewState = FD3D12DynamicRHI::ResourceCast(NewStateRHI);
+
+	StateCache.SetSamplerState<SF_Hull>(NewState, SamplerIndex);
+}
+
+void FD3D12CommandContext::RHISetShaderSampler(FDomainShaderRHIParamRef DomainShaderRHI, uint32 SamplerIndex, FSamplerStateRHIParamRef NewStateRHI)
+{
+	VALIDATE_BOUND_SHADER(DomainShaderRHI);
+
+	FD3D12DomainShader* DomainShader = FD3D12DynamicRHI::ResourceCast(DomainShaderRHI);
+	FD3D12SamplerState* NewState = FD3D12DynamicRHI::ResourceCast(NewStateRHI);
+
+	StateCache.SetSamplerState<SF_Domain>(NewState, SamplerIndex);
+}
+
+void FD3D12CommandContext::RHISetShaderSampler(FGeometryShaderRHIParamRef GeometryShaderRHI, uint32 SamplerIndex, FSamplerStateRHIParamRef NewStateRHI)
+{
+	VALIDATE_BOUND_SHADER(GeometryShaderRHI);
+
+	FD3D12GeometryShader* GeometryShader = FD3D12DynamicRHI::ResourceCast(GeometryShaderRHI);
+	FD3D12SamplerState* NewState = FD3D12DynamicRHI::ResourceCast(NewStateRHI);
+
+	StateCache.SetSamplerState<SF_Geometry>(NewState, SamplerIndex);
+}
+
+void FD3D12CommandContext::RHISetShaderSampler(FPixelShaderRHIParamRef PixelShaderRHI, uint32 SamplerIndex, FSamplerStateRHIParamRef NewStateRHI)
+{
+	VALIDATE_BOUND_SHADER(PixelShaderRHI);
+
+	FD3D12PixelShader* PixelShader = FD3D12DynamicRHI::ResourceCast(PixelShaderRHI);
+	FD3D12SamplerState* NewState = FD3D12DynamicRHI::ResourceCast(NewStateRHI);
+
+	StateCache.SetSamplerState<SF_Pixel>(NewState, SamplerIndex);
+}
+
+void FD3D12CommandContext::RHISetShaderSampler(FComputeShaderRHIParamRef ComputeShaderRHI, uint32 SamplerIndex, FSamplerStateRHIParamRef NewStateRHI)
+{
+	//VALIDATE_BOUND_SHADER(ComputeShaderRHI);
+	FD3D12ComputeShader* ComputeShader = FD3D12DynamicRHI::ResourceCast(ComputeShaderRHI);
+	FD3D12SamplerState* NewState = FD3D12DynamicRHI::ResourceCast(NewStateRHI);
+
+	StateCache.SetSamplerState<SF_Compute>(NewState, SamplerIndex);
+}
+
+void FD3D12CommandContext::RHISetShaderUniformBuffer(FVertexShaderRHIParamRef VertexShader, uint32 BufferIndex, FUniformBufferRHIParamRef BufferRHI)
+{
+	SCOPE_CYCLE_COUNTER(STAT_D3D12SetShaderUniformBuffer);
+	VALIDATE_BOUND_SHADER(VertexShader);
+	FD3D12UniformBuffer* Buffer = FD3D12DynamicRHI::ResourceCast(BufferRHI);
+#if PLATFORM_XBOXONE
+	if (Buffer && Buffer->RingAllocation.IsValid())
+	{
+		StateCache.SetDynamicConstantBuffer<SF_Vertex>(BufferIndex, Buffer->RingAllocation);
+	}
+	else
+#endif
+	{
+		StateCache.SetConstantBuffer<SF_Vertex>(BufferIndex, nullptr, Buffer);
+	}
+
+	BoundUniformBuffers[SF_Vertex][BufferIndex] = BufferRHI;
+	DirtyUniformBuffers[SF_Vertex] |= (1 << BufferIndex);
+}
+
+void FD3D12CommandContext::RHISetShaderUniformBuffer(FHullShaderRHIParamRef HullShader, uint32 BufferIndex, FUniformBufferRHIParamRef BufferRHI)
+{
+	SCOPE_CYCLE_COUNTER(STAT_D3D12SetShaderUniformBuffer);
+	VALIDATE_BOUND_SHADER(HullShader);
+	FD3D12UniformBuffer* Buffer = FD3D12DynamicRHI::ResourceCast(BufferRHI);
+#if PLATFORM_XBOXONE
+	if (Buffer && Buffer->RingAllocation.IsValid())
+	{
+		StateCache.SetDynamicConstantBuffer<SF_Hull>(BufferIndex, Buffer->RingAllocation);
+	}
+	else
+#endif
+	{
+		StateCache.SetConstantBuffer<SF_Hull>(BufferIndex, nullptr, Buffer);
+	}
+
+	BoundUniformBuffers[SF_Hull][BufferIndex] = BufferRHI;
+	DirtyUniformBuffers[SF_Hull] |= (1 << BufferIndex);
+}
+
+void FD3D12CommandContext::RHISetShaderUniformBuffer(FDomainShaderRHIParamRef DomainShader, uint32 BufferIndex, FUniformBufferRHIParamRef BufferRHI)
+{
+	SCOPE_CYCLE_COUNTER(STAT_D3D12SetShaderUniformBuffer);
+	VALIDATE_BOUND_SHADER(DomainShader);
+	FD3D12UniformBuffer* Buffer = FD3D12DynamicRHI::ResourceCast(BufferRHI);
+#if PLATFORM_XBOXONE
+	if (Buffer && Buffer->RingAllocation.IsValid())
+	{
+		StateCache.SetDynamicConstantBuffer<SF_Domain>(BufferIndex, Buffer->RingAllocation);
+	}
+	else
+#endif
+	{
+		StateCache.SetConstantBuffer<SF_Domain>(BufferIndex, nullptr, Buffer);
+	}
+
+	BoundUniformBuffers[SF_Domain][BufferIndex] = BufferRHI;
+	DirtyUniformBuffers[SF_Domain] |= (1 << BufferIndex);
+}
+
+void FD3D12CommandContext::RHISetShaderUniformBuffer(FGeometryShaderRHIParamRef GeometryShader, uint32 BufferIndex, FUniformBufferRHIParamRef BufferRHI)
+{
+	SCOPE_CYCLE_COUNTER(STAT_D3D12SetShaderUniformBuffer);
+	VALIDATE_BOUND_SHADER(GeometryShader);
+	FD3D12UniformBuffer* Buffer = FD3D12DynamicRHI::ResourceCast(BufferRHI);
+#if PLATFORM_XBOXONE
+	if (Buffer && Buffer->RingAllocation.IsValid())
+	{
+		StateCache.SetDynamicConstantBuffer<SF_Geometry>(BufferIndex, Buffer->RingAllocation);
+	}
+	else
+#endif
+	{
+		StateCache.SetConstantBuffer<SF_Geometry>(BufferIndex, nullptr, Buffer);
+	}
+
+	BoundUniformBuffers[SF_Geometry][BufferIndex] = BufferRHI;
+	DirtyUniformBuffers[SF_Geometry] |= (1 << BufferIndex);
+}
+
+void FD3D12CommandContext::RHISetShaderUniformBuffer(FPixelShaderRHIParamRef PixelShader, uint32 BufferIndex, FUniformBufferRHIParamRef BufferRHI)
+{
+	SCOPE_CYCLE_COUNTER(STAT_D3D12SetShaderUniformBuffer);
+	VALIDATE_BOUND_SHADER(PixelShader);
+	FD3D12UniformBuffer* Buffer = FD3D12DynamicRHI::ResourceCast(BufferRHI);
+#if PLATFORM_XBOXONE
+	if (Buffer && Buffer->RingAllocation.IsValid())
+	{
+		StateCache.SetDynamicConstantBuffer<SF_Pixel>(BufferIndex, Buffer->RingAllocation);
+	}
+	else
+#endif
+	{
+		StateCache.SetConstantBuffer<SF_Pixel>(BufferIndex, nullptr, Buffer);
+	}
+
+	BoundUniformBuffers[SF_Pixel][BufferIndex] = BufferRHI;
+	DirtyUniformBuffers[SF_Pixel] |= (1 << BufferIndex);
+}
+
+void FD3D12CommandContext::RHISetShaderUniformBuffer(FComputeShaderRHIParamRef ComputeShader, uint32 BufferIndex, FUniformBufferRHIParamRef BufferRHI)
+{
+	SCOPE_CYCLE_COUNTER(STAT_D3D12SetShaderUniformBuffer);
+	//VALIDATE_BOUND_SHADER(ComputeShader);
+	FD3D12UniformBuffer* Buffer = FD3D12DynamicRHI::ResourceCast(BufferRHI);
+#if PLATFORM_XBOXONE
+	if (Buffer && Buffer->RingAllocation.IsValid())
+	{
+		StateCache.SetDynamicConstantBuffer<SF_Compute>(BufferIndex, Buffer->RingAllocation);
+	}
+	else
+#endif
+	{
+		StateCache.SetConstantBuffer<SF_Compute>(BufferIndex, nullptr, Buffer);
+	}
+
+	BoundUniformBuffers[SF_Compute][BufferIndex] = BufferRHI;
+	DirtyUniformBuffers[SF_Compute] |= (1 << BufferIndex);
+}
+
+void FD3D12CommandContext::RHISetShaderParameter(FHullShaderRHIParamRef HullShaderRHI, uint32 BufferIndex, uint32 BaseIndex, uint32 NumBytes, const void* NewValue)
+{
+	VALIDATE_BOUND_SHADER(HullShaderRHI);
+	checkSlow(HSConstantBuffers[BufferIndex]);
+	HSConstantBuffers[BufferIndex]->UpdateConstant((const uint8*)NewValue,BaseIndex,NumBytes);
+}
+
+void FD3D12CommandContext::RHISetShaderParameter(FDomainShaderRHIParamRef DomainShaderRHI, uint32 BufferIndex, uint32 BaseIndex, uint32 NumBytes, const void* NewValue)
+{
+	VALIDATE_BOUND_SHADER(DomainShaderRHI);
+	checkSlow(DSConstantBuffers[BufferIndex]);
+	DSConstantBuffers[BufferIndex]->UpdateConstant((const uint8*)NewValue,BaseIndex,NumBytes);
+}
+
+void FD3D12CommandContext::RHISetShaderParameter(FVertexShaderRHIParamRef VertexShaderRHI, uint32 BufferIndex, uint32 BaseIndex, uint32 NumBytes, const void* NewValue)
+{
+	VALIDATE_BOUND_SHADER(VertexShaderRHI);
+	checkSlow(VSConstantBuffers[BufferIndex]);
+	VSConstantBuffers[BufferIndex]->UpdateConstant((const uint8*)NewValue,BaseIndex,NumBytes);
+}
+
+void FD3D12CommandContext::RHISetShaderParameter(FPixelShaderRHIParamRef PixelShaderRHI, uint32 BufferIndex, uint32 BaseIndex, uint32 NumBytes, const void* NewValue)
+{
+	VALIDATE_BOUND_SHADER(PixelShaderRHI);
+	checkSlow(PSConstantBuffers[BufferIndex]);
+	PSConstantBuffers[BufferIndex]->UpdateConstant((const uint8*)NewValue,BaseIndex,NumBytes);
+}
+
+void FD3D12CommandContext::RHISetShaderParameter(FGeometryShaderRHIParamRef GeometryShaderRHI, uint32 BufferIndex, uint32 BaseIndex, uint32 NumBytes, const void* NewValue)
+{
+	VALIDATE_BOUND_SHADER(GeometryShaderRHI);
+	checkSlow(GSConstantBuffers[BufferIndex]);
+	GSConstantBuffers[BufferIndex]->UpdateConstant((const uint8*)NewValue,BaseIndex,NumBytes);
+}
+
+void FD3D12CommandContext::RHISetShaderParameter(FComputeShaderRHIParamRef ComputeShaderRHI, uint32 BufferIndex, uint32 BaseIndex, uint32 NumBytes, const void* NewValue)
+{
+	//VALIDATE_BOUND_SHADER(ComputeShaderRHI);
+	checkSlow(CSConstantBuffers[BufferIndex]);
+	CSConstantBuffers[BufferIndex]->UpdateConstant((const uint8*)NewValue,BaseIndex,NumBytes);
+}
+
+void FD3D12CommandContext::ValidateExclusiveDepthStencilAccess(FExclusiveDepthStencil RequestedAccess) const
+{
+	const bool bSrcDepthWrite = RequestedAccess.IsDepthWrite();
+	const bool bSrcStencilWrite = RequestedAccess.IsStencilWrite();
+
+	if (bSrcDepthWrite || bSrcStencilWrite)
+	{
+		// New Rule: You have to call SetRenderTarget[s]() before
+		ensure(CurrentDepthTexture);
+
+		const bool bDstDepthWrite = CurrentDSVAccessType.IsDepthWrite();
+		const bool bDstStencilWrite = CurrentDSVAccessType.IsStencilWrite();
+
+		// requested access is not possible, fix SetRenderTarget EExclusiveDepthStencil or request a different one
+		check(!bSrcDepthWrite || bDstDepthWrite);
+		check(!bSrcStencilWrite || bDstStencilWrite);
+	}
+	}
+
+void FD3D12CommandContext::RHISetDepthStencilState(FDepthStencilStateRHIParamRef NewStateRHI, uint32 StencilRef)
+{
+	FD3D12DepthStencilState* NewState = FD3D12DynamicRHI::ResourceCast(NewStateRHI);
+
+	ValidateExclusiveDepthStencilAccess(NewState->AccessType);
+
+	StateCache.SetDepthStencilState(&NewState->Desc, StencilRef);
+}
+
+void FD3D12CommandContext::RHISetBlendState(FBlendStateRHIParamRef NewStateRHI, const FLinearColor& BlendFactor)
+{
+	FD3D12BlendState* NewState = FD3D12DynamicRHI::ResourceCast(NewStateRHI);
+	StateCache.SetBlendState(&NewState->Desc, (const float*)&BlendFactor, 0xffffffff);
+}
+
+void FD3D12CommandContext::CommitRenderTargetsAndUAVs()
+{
+	FD3D12RenderTargetView* RTArray[D3D12_SIMULTANEOUS_RENDER_TARGET_COUNT];
+
+	for(uint32 RenderTargetIndex = 0;RenderTargetIndex < NumSimultaneousRenderTargets;++RenderTargetIndex)
+	{
+		RTArray[RenderTargetIndex] = CurrentRenderTargets[RenderTargetIndex];
+	}
+	FD3D12UnorderedAccessView* UAVArray[D3D12_PS_CS_UAV_REGISTER_COUNT];
+
+	uint32 UAVInitialCountArray[D3D12_PS_CS_UAV_REGISTER_COUNT];
+	for(uint32 UAVIndex = 0;UAVIndex < NumUAVs;++UAVIndex)
+	{
+		UAVArray[UAVIndex] = CurrentUAVs[UAVIndex];
+		// Using the value that indicates to keep the current UAV counter
+		UAVInitialCountArray[UAVIndex] = -1;
+	}
+
+	StateCache.SetRenderTargets(NumSimultaneousRenderTargets,
+		RTArray,
+		CurrentDepthStencilTarget
+		);
+
+	if (NumUAVs > 0)
+	{
+		StateCache.SetUAVs(SF_Pixel, NumSimultaneousRenderTargets,
+			NumUAVs,
+			UAVArray,
+			UAVInitialCountArray
+			);
+	}
+}
+
+struct FRTVDesc
+{
+	uint32 Width;
+	uint32 Height;
+	DXGI_SAMPLE_DESC SampleDesc;
+};
+
+// Return an FRTVDesc structure whose
+// Width and height dimensions are adjusted for the RTV's miplevel.
+FRTVDesc GetRenderTargetViewDesc(FD3D12RenderTargetView* RenderTargetView)
+{
+	const D3D12_RENDER_TARGET_VIEW_DESC &TargetDesc = RenderTargetView->GetDesc();
+
+	FD3D12Resource* BaseResource = RenderTargetView->GetResource();
+	uint32 MipIndex = 0;
+	FRTVDesc ret;
+	memset(&ret, 0, sizeof(ret));
+
+	switch (TargetDesc.ViewDimension)
+	{
+	case D3D12_RTV_DIMENSION_TEXTURE2D:
+	case D3D12_RTV_DIMENSION_TEXTURE2DMS:
+	case D3D12_RTV_DIMENSION_TEXTURE2DARRAY:
+	case D3D12_RTV_DIMENSION_TEXTURE2DMSARRAY:
+		{
+			D3D12_RESOURCE_DESC const& Desc = BaseResource->GetDesc();
+			ret.Width = (uint32)Desc.Width;
+			ret.Height = Desc.Height;
+			ret.SampleDesc = Desc.SampleDesc;
+			if (TargetDesc.ViewDimension == D3D12_RTV_DIMENSION_TEXTURE2D || TargetDesc.ViewDimension == D3D12_RTV_DIMENSION_TEXTURE2DARRAY)
+			{
+				// All the non-multisampled texture types have their mip-slice in the same position.
+				MipIndex = TargetDesc.Texture2D.MipSlice;
+			}
+			break;
+		}
+	case D3D12_RTV_DIMENSION_TEXTURE3D:
+		{
+			D3D12_RESOURCE_DESC const& Desc = BaseResource->GetDesc();
+			ret.Width = (uint32)Desc.Width;
+			ret.Height = Desc.Height;
+			ret.SampleDesc.Count = 1;
+			ret.SampleDesc.Quality = 0;
+			MipIndex = TargetDesc.Texture3D.MipSlice;
+			break;
+		}
+	default:
+		{
+			// not expecting 1D targets.
+			checkNoEntry();
+		}
+	}
+	ret.Width >>= MipIndex;
+	ret.Height >>= MipIndex;
+	return ret;
+}
+
+void FD3D12CommandContext::RHISetRenderTargets(
+	uint32 NewNumSimultaneousRenderTargets,
+	const FRHIRenderTargetView* NewRenderTargetsRHI,
+	const FRHIDepthRenderTargetView* NewDepthStencilTargetRHI,
+	uint32 NewNumUAVs,
+	const FUnorderedAccessViewRHIParamRef* UAVs
+	)
+{
+	FD3D12TextureBase* NewDepthStencilTarget = GetD3D11TextureFromRHITexture(NewDepthStencilTargetRHI ? NewDepthStencilTargetRHI->Texture : nullptr);
+
+#if CHECK_SRV_TRANSITIONS
+	// if the depth buffer is writable then it counts as unresolved.
+	if (NewDepthStencilTargetRHI && NewDepthStencilTargetRHI->GetDepthStencilAccess() == FExclusiveDepthStencil::DepthWrite_StencilWrite && NewDepthStencilTarget)
+	{
+		UnresolvedTargets.Add(NewDepthStencilTarget->GetResource(), FUnresolvedRTInfo(NewDepthStencilTargetRHI->Texture->GetName(), 0, 1, -1, 1));
+	}
+#endif
+
+	check(NewNumSimultaneousRenderTargets + NewNumUAVs <= MaxSimultaneousRenderTargets);
+
+	bool bTargetChanged = false;
+
+	// Set the appropriate depth stencil view depending on whether depth writes are enabled or not
+	FD3D12DepthStencilView* DepthStencilView = NULL;
+	if(NewDepthStencilTarget)
+	{
+		CurrentDSVAccessType = NewDepthStencilTargetRHI->GetDepthStencilAccess();
+		DepthStencilView = NewDepthStencilTarget->GetDepthStencilView(CurrentDSVAccessType);
+
+		// Unbind any shader views of the depth stencil target that are bound.
+		ConditionalClearShaderResource(NewDepthStencilTarget->ResourceLocation);
+	}
+
+	// Check if the depth stencil target is different from the old state.
+	if(CurrentDepthStencilTarget != DepthStencilView)
+	{
+		if (DepthStencilView)
+		{
+			// HiZ on Intel appears to be broken without this barrier. In theory this should all be handled in the descriptor cache but 
+			// we're either missing a barrier or there's an Intel driver bug.
+			FD3D12DynamicRHI::TransitionResource(CommandListHandle, DepthStencilView);
+		}
+
+		CurrentDepthTexture = NewDepthStencilTarget;
+		CurrentDepthStencilTarget = DepthStencilView;
+		bTargetChanged = true;
+	}
+
+	// Gather the render target views for the new render targets.
+	FD3D12RenderTargetView* NewRenderTargetViews[MaxSimultaneousRenderTargets];
+	for(uint32 RenderTargetIndex = 0;RenderTargetIndex < MaxSimultaneousRenderTargets;++RenderTargetIndex)
+	{
+		FD3D12RenderTargetView* RenderTargetView = NULL;
+		if (RenderTargetIndex < NewNumSimultaneousRenderTargets && NewRenderTargetsRHI[RenderTargetIndex].Texture != nullptr)
+		{
+			int32 RTMipIndex = NewRenderTargetsRHI[RenderTargetIndex].MipIndex;
+			int32 RTSliceIndex = NewRenderTargetsRHI[RenderTargetIndex].ArraySliceIndex;
+			FD3D12TextureBase* NewRenderTarget = GetD3D11TextureFromRHITexture(NewRenderTargetsRHI[RenderTargetIndex].Texture);
+			RenderTargetView = NewRenderTarget->GetRenderTargetView(RTMipIndex, RTSliceIndex);
+
+			ensureMsgf(RenderTargetView, TEXT("Texture being set as render target has no RTV"));
+#if CHECK_SRV_TRANSITIONS			
+			if (RenderTargetView)
+			{
+				// remember this target as having been bound for write.
+				ID3D12Resource* RTVResource;
+				RenderTargetView->GetResource(&RTVResource);
+				UnresolvedTargets.Add(RTVResource, FUnresolvedRTInfo(NewRenderTargetsRHI[RenderTargetIndex].Texture->GetName(), RTMipIndex, 1, RTSliceIndex, 1));
+				RTVResource->Release();
+			}
+#endif
+
+			// Unbind any shader views of the render target that are bound.
+			ConditionalClearShaderResource(NewRenderTarget->ResourceLocation);
+#if UE_BUILD_DEBUG	
+			// A check to allow you to pinpoint what is using mismatching targets
+			// We filter our d3ddebug spew that checks for this as the d3d runtime's check is wrong.
+			// For filter code, see D3D12Device.cpp look for "OMSETRENDERTARGETS_INVALIDVIEW"
+			if(RenderTargetView && DepthStencilView)
+			{
+				FRTVDesc RTTDesc = GetRenderTargetViewDesc(RenderTargetView);
+
+				FD3D12Resource* DepthTargetTexture = DepthStencilView->GetResource();
+
+				D3D12_RESOURCE_DESC DTTDesc = DepthTargetTexture->GetDesc();
+
+				// enforce color target is <= depth and MSAA settings match
+				if(RTTDesc.Width > DTTDesc.Width || RTTDesc.Height > DTTDesc.Height || 
+					RTTDesc.SampleDesc.Count != DTTDesc.SampleDesc.Count || 
+					RTTDesc.SampleDesc.Quality != DTTDesc.SampleDesc.Quality)
+				{
+					UE_LOG(LogD3D12RHI, Fatal,TEXT("RTV(%i,%i c=%i,q=%i) and DSV(%i,%i c=%i,q=%i) have mismatching dimensions and/or MSAA levels!"),
+						RTTDesc.Width,RTTDesc.Height,RTTDesc.SampleDesc.Count,RTTDesc.SampleDesc.Quality,
+						DTTDesc.Width,DTTDesc.Height,DTTDesc.SampleDesc.Count,DTTDesc.SampleDesc.Quality);
+				}
+			}
+#endif
+		}
+
+		NewRenderTargetViews[RenderTargetIndex] = RenderTargetView;
+
+		// Check if the render target is different from the old state.
+		if(CurrentRenderTargets[RenderTargetIndex] != RenderTargetView)
+		{
+			CurrentRenderTargets[RenderTargetIndex] = RenderTargetView;
+			bTargetChanged = true;
+		}
+	}
+	if(NumSimultaneousRenderTargets != NewNumSimultaneousRenderTargets)
+	{
+		NumSimultaneousRenderTargets = NewNumSimultaneousRenderTargets;
+		bTargetChanged = true;
+	}
+
+	// Gather the new UAVs.
+	for(uint32 UAVIndex = 0;UAVIndex < MaxSimultaneousUAVs;++UAVIndex)
+	{
+		FD3D12UnorderedAccessView* RHIUAV = NULL;
+		if(UAVIndex < NewNumUAVs && UAVs[UAVIndex] != NULL)
+		{
+			RHIUAV = (FD3D12UnorderedAccessView*)UAVs[UAVIndex];
+			FD3D12DynamicRHI::TransitionResource(CommandListHandle, RHIUAV, D3D12_RESOURCE_STATE_UNORDERED_ACCESS);
+
+			// Unbind any shader views of the UAV's resource.
+			ConditionalClearShaderResource(RHIUAV->GetResourceLocation());
+		}
+
+		if(CurrentUAVs[UAVIndex] != RHIUAV)
+		{
+			CurrentUAVs[UAVIndex] = RHIUAV;
+			bTargetChanged = true;
+		}
+	}
+	if(NumUAVs != NewNumUAVs)
+	{
+		NumUAVs = NewNumUAVs;
+		bTargetChanged = true;
+	}
+
+	// Only make the D3D call to change render targets if something actually changed.
+	if(bTargetChanged)
+	{
+		CommitRenderTargetsAndUAVs();
+	}
+
+	// Set the viewport to the full size of render target 0.
+	if (NewRenderTargetViews[0])
+	{
+		// check target 0 is valid
+		check(0 < NewNumSimultaneousRenderTargets && NewRenderTargetsRHI[0].Texture != nullptr);
+		FRTVDesc RTTDesc = GetRenderTargetViewDesc(NewRenderTargetViews[0]);
+		RHISetViewport(0, 0, 0.0f, RTTDesc.Width, RTTDesc.Height, 1.0f);
+	}
+	else if (DepthStencilView)
+	{
+		FD3D12Resource* DepthTargetTexture = DepthStencilView->GetResource();
+		D3D12_RESOURCE_DESC const& DTTDesc = DepthTargetTexture->GetDesc();
+		RHISetViewport(0, 0, 0.0f, (uint32)DTTDesc.Width, DTTDesc.Height, 1.0f);
+	}
+}
+
+void FD3D12DynamicRHI::RHIDiscardRenderTargets(bool Depth, bool Stencil, uint32 ColorBitMask)
+{
+	// Could support in DX12 via ID3D12CommandList::DiscardResource function.
+}
+
+void FD3D12CommandContext::RHISetRenderTargetsAndClear(const FRHISetRenderTargetsInfo& RenderTargetsInfo)
+{
+	this->RHISetRenderTargets(RenderTargetsInfo.NumColorRenderTargets,
+		RenderTargetsInfo.ColorRenderTarget,
+		&RenderTargetsInfo.DepthStencilRenderTarget,
+		0,
+		nullptr);
+	FLinearColor ClearColors[MaxSimultaneousRenderTargets];
+	float DepthClear = 0.0;
+	uint32 StencilClear = 0;
+
+	if (RenderTargetsInfo.bClearColor)
+	{
+		for (int32 i = 0; i < RenderTargetsInfo.NumColorRenderTargets; ++i)
+		{
+			const FClearValueBinding& ClearValue = RenderTargetsInfo.ColorRenderTarget[i].Texture->GetClearBinding();
+			checkf(ClearValue.ColorBinding == EClearBinding::EColorBound, TEXT("Texture: %s does not have a color bound for fast clears"), *RenderTargetsInfo.ColorRenderTarget[i].Texture->GetName().GetPlainNameString());
+			ClearColors[i] = ClearValue.GetClearColor();
+		}
+	}
+	if (RenderTargetsInfo.bClearDepth || RenderTargetsInfo.bClearStencil)
+	{
+		const FClearValueBinding& ClearValue = RenderTargetsInfo.DepthStencilRenderTarget.Texture->GetClearBinding();
+		checkf(ClearValue.ColorBinding == EClearBinding::EDepthStencilBound, TEXT("Texture: %s does not have a DS value bound for fast clears"), *RenderTargetsInfo.DepthStencilRenderTarget.Texture->GetName().GetPlainNameString());
+		ClearValue.GetDepthStencil(DepthClear, StencilClear);
+	}
+
+	this->RHIClearMRTImpl(RenderTargetsInfo.bClearColor, RenderTargetsInfo.NumColorRenderTargets, ClearColors, RenderTargetsInfo.bClearDepth, DepthClear, RenderTargetsInfo.bClearStencil, StencilClear, FIntRect(), false);
+}
+
+// Occlusion/Timer queries.
+void FD3D12CommandContext::RHIBeginRenderQuery(FRenderQueryRHIParamRef QueryRHI)
+{
+	FD3D12OcclusionQuery* Query = FD3D12DynamicRHI::ResourceCast(QueryRHI);
+
+	if(Query->Type == RQT_Occlusion)
+	{
+		Query->bResultIsCached = false;
+		Query->HeapIndex = GetParentDevice()->GetQueryHeap()->BeginQuery(*this, D3D12_QUERY_TYPE_OCCLUSION);
+		Query->OwningCommandList = CommandListHandle;
+		Query->OwningContext = this;
+	}
+	else
+	{
+		// not supported/needed for RQT_AbsoluteTime
+		check(0);
+	}
+
+#if EXECUTE_DEBUG_COMMAND_LISTS
+	GIsDoingQuery = true;
+#endif
+}
+void FD3D12CommandContext::RHIEndRenderQuery(FRenderQueryRHIParamRef QueryRHI)
+{
+	FD3D12OcclusionQuery*  Query = FD3D12DynamicRHI::ResourceCast(QueryRHI);
+	if (Query != NULL)
+	{
+		// This code always assumed it was an occlusion query
+		check(Query->Type == RQT_Occlusion);
+
+		// End the query
+		GetParentDevice()->GetQueryHeap()->EndQuery(*this, D3D12_QUERY_TYPE_OCCLUSION, Query->HeapIndex);
+
+		check(Query->OwningCommandList == CommandListHandle);
+		check(Query->OwningContext == this);
+	}
+
+#if EXECUTE_DEBUG_COMMAND_LISTS
+	GIsDoingQuery = false;
+#endif
+}
+
+// Primitive drawing.
+
+static D3D_PRIMITIVE_TOPOLOGY GetD3D11PrimitiveType(uint32 PrimitiveType, bool bUsingTessellation)
+{
+	if(bUsingTessellation)
+	{
+		switch(PrimitiveType)
+		{
+		case PT_1_ControlPointPatchList: return D3D_PRIMITIVE_TOPOLOGY_1_CONTROL_POINT_PATCHLIST;
+		case PT_2_ControlPointPatchList: return D3D_PRIMITIVE_TOPOLOGY_2_CONTROL_POINT_PATCHLIST;
+
+			// This is the case for tessellation without AEN or other buffers, so just flip to 3 CPs
+		case PT_TriangleList: return D3D_PRIMITIVE_TOPOLOGY_3_CONTROL_POINT_PATCHLIST;
+
+		case PT_LineList:
+		case PT_TriangleStrip:
+		case PT_QuadList:
+		case PT_PointList:
+			UE_LOG(LogD3D12RHI, Fatal,TEXT("Invalid type specified for tessellated render, probably missing a case in FSkeletalMeshSceneProxy::DrawDynamicElementsByMaterial or FStaticMeshSceneProxy::GetMeshElement"));
+			break;
+		default:
+			// Other cases are valid.
+			break;
+		};
+	}
+
+	switch(PrimitiveType)
+	{
+	case PT_TriangleList: return D3D_PRIMITIVE_TOPOLOGY_TRIANGLELIST;
+	case PT_TriangleStrip: return D3D_PRIMITIVE_TOPOLOGY_TRIANGLESTRIP;
+	case PT_LineList: return D3D_PRIMITIVE_TOPOLOGY_LINELIST;
+	case PT_PointList: return D3D_PRIMITIVE_TOPOLOGY_POINTLIST;
+
+		// ControlPointPatchList types will pretend to be TRIANGLELISTS with a stride of N 
+		// (where N is the number of control points specified), so we can return them for
+		// tessellation and non-tessellation. This functionality is only used when rendering a 
+		// default material with something that claims to be tessellated, generally because the 
+		// tessellation material failed to compile for some reason.
+	case PT_3_ControlPointPatchList: return D3D_PRIMITIVE_TOPOLOGY_3_CONTROL_POINT_PATCHLIST;
+	case PT_4_ControlPointPatchList: return D3D_PRIMITIVE_TOPOLOGY_4_CONTROL_POINT_PATCHLIST;
+	case PT_5_ControlPointPatchList: return D3D_PRIMITIVE_TOPOLOGY_5_CONTROL_POINT_PATCHLIST;
+	case PT_6_ControlPointPatchList: return D3D_PRIMITIVE_TOPOLOGY_6_CONTROL_POINT_PATCHLIST;
+	case PT_7_ControlPointPatchList: return D3D_PRIMITIVE_TOPOLOGY_7_CONTROL_POINT_PATCHLIST;
+	case PT_8_ControlPointPatchList: return D3D_PRIMITIVE_TOPOLOGY_8_CONTROL_POINT_PATCHLIST; 
+	case PT_9_ControlPointPatchList: return D3D_PRIMITIVE_TOPOLOGY_9_CONTROL_POINT_PATCHLIST; 
+	case PT_10_ControlPointPatchList: return D3D_PRIMITIVE_TOPOLOGY_10_CONTROL_POINT_PATCHLIST; 
+	case PT_11_ControlPointPatchList: return D3D_PRIMITIVE_TOPOLOGY_11_CONTROL_POINT_PATCHLIST; 
+	case PT_12_ControlPointPatchList: return D3D_PRIMITIVE_TOPOLOGY_12_CONTROL_POINT_PATCHLIST; 
+	case PT_13_ControlPointPatchList: return D3D_PRIMITIVE_TOPOLOGY_13_CONTROL_POINT_PATCHLIST; 
+	case PT_14_ControlPointPatchList: return D3D_PRIMITIVE_TOPOLOGY_14_CONTROL_POINT_PATCHLIST; 
+	case PT_15_ControlPointPatchList: return D3D_PRIMITIVE_TOPOLOGY_15_CONTROL_POINT_PATCHLIST; 
+	case PT_16_ControlPointPatchList: return D3D_PRIMITIVE_TOPOLOGY_16_CONTROL_POINT_PATCHLIST; 
+	case PT_17_ControlPointPatchList: return D3D_PRIMITIVE_TOPOLOGY_17_CONTROL_POINT_PATCHLIST; 
+	case PT_18_ControlPointPatchList: return D3D_PRIMITIVE_TOPOLOGY_18_CONTROL_POINT_PATCHLIST; 
+	case PT_19_ControlPointPatchList: return D3D_PRIMITIVE_TOPOLOGY_19_CONTROL_POINT_PATCHLIST; 
+	case PT_20_ControlPointPatchList: return D3D_PRIMITIVE_TOPOLOGY_20_CONTROL_POINT_PATCHLIST; 
+	case PT_21_ControlPointPatchList: return D3D_PRIMITIVE_TOPOLOGY_21_CONTROL_POINT_PATCHLIST; 
+	case PT_22_ControlPointPatchList: return D3D_PRIMITIVE_TOPOLOGY_22_CONTROL_POINT_PATCHLIST; 
+	case PT_23_ControlPointPatchList: return D3D_PRIMITIVE_TOPOLOGY_23_CONTROL_POINT_PATCHLIST; 
+	case PT_24_ControlPointPatchList: return D3D_PRIMITIVE_TOPOLOGY_24_CONTROL_POINT_PATCHLIST; 
+	case PT_25_ControlPointPatchList: return D3D_PRIMITIVE_TOPOLOGY_25_CONTROL_POINT_PATCHLIST; 
+	case PT_26_ControlPointPatchList: return D3D_PRIMITIVE_TOPOLOGY_26_CONTROL_POINT_PATCHLIST; 
+	case PT_27_ControlPointPatchList: return D3D_PRIMITIVE_TOPOLOGY_27_CONTROL_POINT_PATCHLIST; 
+	case PT_28_ControlPointPatchList: return D3D_PRIMITIVE_TOPOLOGY_28_CONTROL_POINT_PATCHLIST; 
+	case PT_29_ControlPointPatchList: return D3D_PRIMITIVE_TOPOLOGY_29_CONTROL_POINT_PATCHLIST; 
+	case PT_30_ControlPointPatchList: return D3D_PRIMITIVE_TOPOLOGY_30_CONTROL_POINT_PATCHLIST; 
+	case PT_31_ControlPointPatchList: return D3D_PRIMITIVE_TOPOLOGY_31_CONTROL_POINT_PATCHLIST; 
+	case PT_32_ControlPointPatchList: return D3D_PRIMITIVE_TOPOLOGY_32_CONTROL_POINT_PATCHLIST; 
+	default: UE_LOG(LogD3D12RHI, Fatal,TEXT("Unknown primitive type: %u"),PrimitiveType);
+	};
+
+	return D3D_PRIMITIVE_TOPOLOGY_TRIANGLELIST;
+}
+
+void FD3D12CommandContext::CommitNonComputeShaderConstants()
+{
+	SCOPE_CYCLE_COUNTER(STAT_D3D12CommitGraphicsConstants);
+
+	// MSFT: Seb: Do we need to support this none parallel case?
+	//if (GRHISupportsParallelRHIExecute)
+		FD3D12BoundShaderState* RESTRICT CurrentBoundShaderState = this->CurrentBoundShaderState.GetReference();
+	//else
+	//	FD3D12BoundShaderState* RESTRICT CurrentBoundShaderState = (FD3D12BoundShaderState*)OwningRHI.BoundShaderStateHistory.GetLast();
+
+	check(CurrentBoundShaderState);
+
+	// Only set the constant buffer if this shader needs the global constant buffer bound
+	// Otherwise we will overwrite a different constant buffer
+	if (CurrentBoundShaderState->bShaderNeedsGlobalConstantBuffer[SF_Vertex])
+	{
+		// Commit and bind vertex shader constants
+		for(uint32 i=0;i<MAX_CONSTANT_BUFFER_SLOTS; i++)
+		{
+			FD3D12ConstantBuffer* ConstantBuffer = VSConstantBuffers[i];
+			FD3DRHIUtil::CommitConstants<SF_Vertex>(UploadHeapAllocator, ConstantBuffer, StateCache, i, bDiscardSharedConstants);
+		}
+	}
+
+	// Skip HS/DS CB updates in cases where tessellation isn't being used
+	// Note that this is *potentially* unsafe because bDiscardSharedConstants is cleared at the
+	// end of the function, however we're OK for now because bDiscardSharedConstants
+	// is always reset whenever bUsingTessellation changes in SetBoundShaderState()
+	if(bUsingTessellation)
+	{
+		if (CurrentBoundShaderState->bShaderNeedsGlobalConstantBuffer[SF_Hull])
+		{
+			// Commit and bind hull shader constants
+			for(uint32 i=0;i<MAX_CONSTANT_BUFFER_SLOTS; i++)
+			{
+				FD3D12ConstantBuffer* ConstantBuffer = HSConstantBuffers[i];
+				FD3DRHIUtil::CommitConstants<SF_Hull>(UploadHeapAllocator, ConstantBuffer, StateCache, i, bDiscardSharedConstants);
+			}
+		}
+
+		if (CurrentBoundShaderState->bShaderNeedsGlobalConstantBuffer[SF_Domain])
+		{
+			// Commit and bind domain shader constants
+			for(uint32 i=0;i<MAX_CONSTANT_BUFFER_SLOTS; i++)
+			{
+				FD3D12ConstantBuffer* ConstantBuffer = DSConstantBuffers[i];
+				FD3DRHIUtil::CommitConstants<SF_Domain>(UploadHeapAllocator, ConstantBuffer, StateCache, i, bDiscardSharedConstants);
+			}
+		}
+	}
+
+	if (CurrentBoundShaderState->bShaderNeedsGlobalConstantBuffer[SF_Geometry])
+	{
+		// Commit and bind geometry shader constants
+		for(uint32 i=0;i<MAX_CONSTANT_BUFFER_SLOTS; i++)
+		{
+			FD3D12ConstantBuffer* ConstantBuffer = GSConstantBuffers[i];
+			FD3DRHIUtil::CommitConstants<SF_Geometry>(UploadHeapAllocator, ConstantBuffer, StateCache, i, bDiscardSharedConstants);
+		}
+	}
+
+	if (CurrentBoundShaderState->bShaderNeedsGlobalConstantBuffer[SF_Pixel])
+	{
+		// Commit and bind pixel shader constants
+		for(uint32 i=0;i<MAX_CONSTANT_BUFFER_SLOTS; i++)
+		{
+			FD3D12ConstantBuffer* ConstantBuffer = PSConstantBuffers[i];
+			FD3DRHIUtil::CommitConstants<SF_Pixel>(UploadHeapAllocator, ConstantBuffer, StateCache, i, bDiscardSharedConstants);
+		}
+	}
+
+	bDiscardSharedConstants = false;
+}
+
+void FD3D12CommandContext::CommitComputeShaderConstants()
+{
+	bool bLocalDiscardSharedConstants = true;
+
+	// Commit and bind compute shader constants
+	for(uint32 i=0;i<MAX_CONSTANT_BUFFER_SLOTS; i++)
+	{
+		FD3D12ConstantBuffer* ConstantBuffer = CSConstantBuffers[i];
+		FD3DRHIUtil::CommitConstants<SF_Compute>(UploadHeapAllocator, ConstantBuffer, StateCache, i, bDiscardSharedConstants);
+	}
+}
+
+template <EShaderFrequency Frequency>
+FORCEINLINE void SetResource(FD3D12CommandContext& CmdContext, uint32 BindIndex, FD3D12ResourceLocation* RESTRICT ShaderResource, FD3D12ShaderResourceView* RESTRICT SRV)
+{
+	// We set the resource through the RHI to track state for the purposes of unbinding SRVs when a UAV or RTV is bound.
+	// todo: need to support SRV_Static for faster calls when possible
+	CmdContext.SetShaderResourceView<Frequency>(ShaderResource, SRV, BindIndex, FD3D12StateCache::SRV_Unknown);
+}
+
+template <EShaderFrequency Frequency>
+FORCEINLINE void SetResource(FD3D12CommandContext& CmdContext, uint32 BindIndex, FD3D12ResourceLocation* RESTRICT ShaderResource, FD3D12SamplerState* RESTRICT SamplerState)
+{
+	CmdContext.StateCache.SetSamplerState<Frequency>(SamplerState, BindIndex);
+}
+
+template <class D3DResourceType, EShaderFrequency ShaderFrequency>
+inline int32 SetShaderResourcesFromBuffer(FD3D12CommandContext& CmdContext, FD3D12UniformBuffer* RESTRICT Buffer, const uint32 * RESTRICT ResourceMap, int32 BufferIndex)
+{
+	int32 NumSetCalls = 0;
+	uint32 BufferOffset = ResourceMap[BufferIndex];
+	if (BufferOffset > 0)
+	{
+		const uint32* RESTRICT ResourceInfos = &ResourceMap[BufferOffset];
+		uint32 ResourceInfo = *ResourceInfos++;
+		do
+		{
+			checkSlow(FRHIResourceTableEntry::GetUniformBufferIndex(ResourceInfo) == BufferIndex);
+			const uint16 ResourceIndex = FRHIResourceTableEntry::GetResourceIndex(ResourceInfo);
+			const uint8 BindIndex = FRHIResourceTableEntry::GetBindIndex(ResourceInfo);
+
+			// todo: could coalesce adjacent bound resources.
+			FD3D12UniformBuffer::FResourcePair* RESTRICT ResourcePair = &Buffer->RawResourceTable[ResourceIndex];
+			FD3D12ResourceLocation* ShaderResource = ResourcePair->ShaderResourceLocation;
+			D3DResourceType* D3D12Resource = (D3DResourceType*)ResourcePair->D3D11Resource;
+			SetResource<ShaderFrequency>(CmdContext, BindIndex, ShaderResource, D3D12Resource);
+			NumSetCalls++;
+			ResourceInfo = *ResourceInfos++;
+		} while (FRHIResourceTableEntry::GetUniformBufferIndex(ResourceInfo) == BufferIndex);
+	}
+	return NumSetCalls;
+}
+
+template <class ShaderType>
+void FD3D12CommandContext::SetResourcesFromTables(const ShaderType* RESTRICT Shader)
+{
+	checkSlow(Shader);
+
+	// Mask the dirty bits by those buffers from which the shader has bound resources.
+	uint32 DirtyBits = Shader->ShaderResourceTable.ResourceTableBits & DirtyUniformBuffers[ShaderType::StaticFrequency];
+	uint32 NumSetCalls = 0;
+	while (DirtyBits)
+	{
+		// Scan for the lowest set bit, compute its index, clear it in the set of dirty bits.
+		const uint32 LowestBitMask = (DirtyBits)& (-(int32)DirtyBits);
+		const int32 BufferIndex = FMath::FloorLog2(LowestBitMask); // todo: This has a branch on zero, we know it could never be zero...
+		DirtyBits ^= LowestBitMask;
+		FD3D12UniformBuffer* Buffer = (FD3D12UniformBuffer*)BoundUniformBuffers[ShaderType::StaticFrequency][BufferIndex].GetReference();
+		check(Buffer);
+		check(BufferIndex < Shader->ShaderResourceTable.ResourceTableLayoutHashes.Num());
+		check(Buffer->GetLayout().GetHash() == Shader->ShaderResourceTable.ResourceTableLayoutHashes[BufferIndex]);
+		Buffer->CacheResources(OwningRHI.GetResourceTableFrameCounter());
+
+		// todo: could make this two pass: gather then set
+		NumSetCalls += SetShaderResourcesFromBuffer<FD3D12ShaderResourceView, (EShaderFrequency)ShaderType::StaticFrequency>(*this, Buffer, Shader->ShaderResourceTable.ShaderResourceViewMap.GetData(), BufferIndex);
+		NumSetCalls += SetShaderResourcesFromBuffer<FD3D12SamplerState, (EShaderFrequency)ShaderType::StaticFrequency>(*this, Buffer, Shader->ShaderResourceTable.SamplerMap.GetData(), BufferIndex);
+	}
+	DirtyUniformBuffers[ShaderType::StaticFrequency] = 0;
+	OwningRHI.IncrementSetTextureInTableCalls(NumSetCalls);
+}
+
+void FD3D12CommandContext::CommitGraphicsResourceTables()
+{
+	uint32 Start = FPlatformTime::Cycles();
+
+	// MSFT: Seb: Do we need this non parallel case now that context objects are always used?
+//#if PLATFORM_SUPPORTS_PARALLEL_RHI_EXECUTE
+	FD3D12BoundShaderState* RESTRICT CurrentBoundShaderState = this->CurrentBoundShaderState.GetReference();
+//#else
+//	FD3D12BoundShaderState* RESTRICT CurrentBoundShaderState = (FD3D12BoundShaderState*)OwningRHI.BoundShaderStateHistory.GetLast();
+//#endif
+	check(CurrentBoundShaderState);
+
+	if (auto* Shader = CurrentBoundShaderState->GetVertexShader())
+	{
+		SetResourcesFromTables(Shader);
+	}
+	if (auto* Shader = CurrentBoundShaderState->GetPixelShader())
+	{
+		SetResourcesFromTables(Shader);
+	}
+	if (auto* Shader = CurrentBoundShaderState->GetHullShader())
+	{
+		SetResourcesFromTables(Shader);
+	}
+	if (auto* Shader = CurrentBoundShaderState->GetDomainShader())
+	{
+		SetResourcesFromTables(Shader);
+	}
+	if (auto* Shader = CurrentBoundShaderState->GetGeometryShader())
+	{
+		SetResourcesFromTables(Shader);
+	}
+
+	OwningRHI.IncrementCommitComputeResourceTables(FPlatformTime::Cycles() - Start);
+}
+
+void FD3D12CommandContext::CommitComputeResourceTables(FD3D12ComputeShader* InComputeShader)
+{
+	FD3D12ComputeShader* RESTRICT ComputeShader = InComputeShader;
+	check(ComputeShader);
+	SetResourcesFromTables(ComputeShader);
+}
+
+void FD3D12CommandContext::RHIDrawPrimitive(uint32 PrimitiveType, uint32 BaseVertexIndex, uint32 NumPrimitives, uint32 NumInstances)
+{
+	RHI_DRAW_CALL_STATS(PrimitiveType,NumInstances*NumPrimitives);
+
+	CommitGraphicsResourceTables();
+	CommitNonComputeShaderConstants();
+
+	uint32 VertexCount = GetVertexCountForPrimitiveCount(NumPrimitives,PrimitiveType);
+
+	OwningRHI.RegisterGPUWork(NumPrimitives * NumInstances, VertexCount * NumInstances);
+
+	StateCache.SetPrimitiveTopology(GetD3D11PrimitiveType(PrimitiveType, bUsingTessellation));
+
+	StateCache.ApplyState();
+	numDraws++;
+	CommandListHandle->DrawInstanced(VertexCount, FMath::Max<uint32>(1, NumInstances), BaseVertexIndex, 0);
+#if UE_BUILD_DEBUG	
+	OwningRHI.DrawCount++;
+#endif
+	DEBUG_EXECUTE_COMMAND_LIST(this);
+
+}
+
+void FD3D12CommandContext::RHIDrawPrimitiveIndirect(uint32 PrimitiveType, FVertexBufferRHIParamRef ArgumentBufferRHI, uint32 ArgumentOffset)
+{
+	FD3D12VertexBuffer* ArgumentBuffer = FD3D12DynamicRHI::ResourceCast(ArgumentBufferRHI);
+
+	RHI_DRAW_CALL_INC();
+
+	OwningRHI.RegisterGPUWork(0);
+
+	CommitGraphicsResourceTables();
+	CommitNonComputeShaderConstants();
+
+	StateCache.SetPrimitiveTopology(GetD3D11PrimitiveType(PrimitiveType, bUsingTessellation));
+
+	StateCache.ApplyState();
+
+	FD3D12DynamicRHI::TransitionResource(CommandListHandle, ArgumentBuffer->ResourceLocation->GetResource(), D3D12_RESOURCE_STATE_INDIRECT_ARGUMENT, D3D12_RESOURCE_BARRIER_TYPE_TRANSITION);
+
+	numDraws++;
+	CommandListHandle->ExecuteIndirect(
+		GetParentDevice()->GetDrawIndirectCommandSignature(),
+		1,
+		ArgumentBuffer->ResourceLocation->GetResource()->GetResource(),
+		ArgumentBuffer->ResourceLocation->GetOffset() + ArgumentOffset * 20,
+		NULL,
+		0
+		);
+
+#if UE_BUILD_DEBUG	
+	OwningRHI.DrawCount++;
+#endif
+	DEBUG_EXECUTE_COMMAND_LIST(this);
+
+}
+
+void FD3D12CommandContext::RHIDrawIndexedIndirect(FIndexBufferRHIParamRef IndexBufferRHI, uint32 PrimitiveType, FStructuredBufferRHIParamRef ArgumentsBufferRHI, int32 DrawArgumentsIndex, uint32 NumInstances)
+{
+	FD3D12IndexBuffer* IndexBuffer = FD3D12DynamicRHI::ResourceCast(IndexBufferRHI);
+	FD3D12StructuredBuffer* ArgumentsBuffer = FD3D12DynamicRHI::ResourceCast(ArgumentsBufferRHI);
+
+	RHI_DRAW_CALL_INC();
+
+	OwningRHI.RegisterGPUWork(1);
+
+	CommitGraphicsResourceTables();
+	CommitNonComputeShaderConstants();
+
+	// determine 16bit vs 32bit indices
+	uint32 SizeFormat = sizeof(DXGI_FORMAT);
+	const DXGI_FORMAT Format = (IndexBuffer->GetStride() == sizeof(uint16) ? DXGI_FORMAT_R16_UINT : DXGI_FORMAT_R32_UINT);
+
+	StateCache.SetIndexBuffer(IndexBuffer->ResourceLocation.GetReference(), Format, 0);
+	StateCache.SetPrimitiveTopology(GetD3D11PrimitiveType(PrimitiveType,bUsingTessellation));
+	StateCache.ApplyState();
+
+
+	FD3D12DynamicRHI::TransitionResource(CommandListHandle, ArgumentsBuffer->Resource, D3D12_RESOURCE_STATE_INDIRECT_ARGUMENT, D3D12_RESOURCE_BARRIER_TYPE_TRANSITION);
+
+	numDraws++;
+	CommandListHandle->ExecuteIndirect(
+		GetParentDevice()->GetDrawIndexedIndirectCommandSignature(),
+		1,
+		ArgumentsBuffer->Resource->GetResource(),
+		ArgumentsBuffer->ResourceLocation->GetOffset() + DrawArgumentsIndex * ArgumentsBuffer->GetStride(),
+		NULL,
+		0
+		);
+
+#if UE_BUILD_DEBUG	
+	OwningRHI.DrawCount++;
+#endif
+	DEBUG_EXECUTE_COMMAND_LIST(this);
+}
+
+void FD3D12CommandContext::RHIDrawIndexedPrimitive(FIndexBufferRHIParamRef IndexBufferRHI, uint32 PrimitiveType, int32 BaseVertexIndex, uint32 FirstInstance, uint32 NumVertices, uint32 StartIndex, uint32 NumPrimitives, uint32 NumInstances)
+{
+	FD3D12IndexBuffer* IndexBuffer = FD3D12DynamicRHI::ResourceCast(IndexBufferRHI);
+
+	// called should make sure the input is valid, this avoid hidden bugs
+	ensure(NumPrimitives > 0);
+
+	RHI_DRAW_CALL_STATS(PrimitiveType,NumInstances*NumPrimitives);
+
+	OwningRHI.RegisterGPUWork(NumPrimitives * NumInstances, NumVertices * NumInstances);
+
+	CommitGraphicsResourceTables();
+	CommitNonComputeShaderConstants();
+
+	// determine 16bit vs 32bit indices
+	uint32 SizeFormat = sizeof(DXGI_FORMAT);
+	const DXGI_FORMAT Format = (IndexBuffer->GetStride() == sizeof(uint16) ? DXGI_FORMAT_R16_UINT : DXGI_FORMAT_R32_UINT);
+
+	uint32 IndexCount = GetVertexCountForPrimitiveCount(NumPrimitives,PrimitiveType);
+
+	// Verify that we are not trying to read outside the index buffer range
+	// test is an optimized version of: StartIndex + IndexCount <= IndexBuffer->GetSize() / IndexBuffer->GetStride() 
+	checkf((StartIndex + IndexCount) * IndexBuffer->GetStride() <= IndexBuffer->GetSize(), 		
+		TEXT("Start %u, Count %u, Type %u, Buffer Size %u, Buffer stride %u"), StartIndex, IndexCount, PrimitiveType, IndexBuffer->GetSize(), IndexBuffer->GetStride());
+
+	StateCache.SetIndexBuffer(IndexBuffer->ResourceLocation.GetReference(), Format, 0);
+	StateCache.SetPrimitiveTopology(GetD3D11PrimitiveType(PrimitiveType,bUsingTessellation));
+	StateCache.ApplyState();
+
+	numDraws++;
+	CommandListHandle->DrawIndexedInstanced(IndexCount, FMath::Max<uint32>(1, NumInstances), StartIndex, BaseVertexIndex, FirstInstance);
+#if UE_BUILD_DEBUG	
+	OwningRHI.DrawCount++;
+#endif
+	DEBUG_EXECUTE_COMMAND_LIST(this);
+
+}
+
+void FD3D12CommandContext::RHIDrawIndexedPrimitiveIndirect(uint32 PrimitiveType, FIndexBufferRHIParamRef IndexBufferRHI, FVertexBufferRHIParamRef ArgumentBufferRHI, uint32 ArgumentOffset)
+{
+	FD3D12IndexBuffer* IndexBuffer = FD3D12DynamicRHI::ResourceCast(IndexBufferRHI);
+	FD3D12VertexBuffer* ArgumentBuffer = FD3D12DynamicRHI::ResourceCast(ArgumentBufferRHI);
+
+	RHI_DRAW_CALL_INC();
+
+	OwningRHI.RegisterGPUWork(0);
+
+	CommitGraphicsResourceTables();
+	CommitNonComputeShaderConstants();
+
+	// Set the index buffer.
+	const uint32 SizeFormat = sizeof(DXGI_FORMAT);
+	const DXGI_FORMAT Format = (IndexBuffer->GetStride() == sizeof(uint16) ? DXGI_FORMAT_R16_UINT : DXGI_FORMAT_R32_UINT);
+	StateCache.SetIndexBuffer(IndexBuffer->ResourceLocation.GetReference(), Format, 0);
+	StateCache.SetPrimitiveTopology(GetD3D11PrimitiveType(PrimitiveType,bUsingTessellation));
+	StateCache.ApplyState();
+
+	FD3D12DynamicRHI::TransitionResource(CommandListHandle, ArgumentBuffer->ResourceLocation->GetResource(), D3D12_RESOURCE_STATE_INDIRECT_ARGUMENT, D3D12_RESOURCE_BARRIER_TYPE_TRANSITION);
+
+	numDraws++;
+	CommandListHandle->ExecuteIndirect(
+		GetParentDevice()->GetDrawIndexedIndirectCommandSignature(),
+		1,
+		ArgumentBuffer->ResourceLocation->GetResource()->GetResource(),
+		ArgumentBuffer->ResourceLocation->GetOffset() + ArgumentOffset * 20,
+		NULL,
+		0
+		);
+
+#if UE_BUILD_DEBUG	
+	OwningRHI.DrawCount++;
+#endif
+	DEBUG_EXECUTE_COMMAND_LIST(this);
+}
+
+/**
+* Preallocate memory or get a direct command stream pointer to fill up for immediate rendering . This avoids memcpys below in DrawPrimitiveUP
+* @param PrimitiveType The type (triangles, lineloop, etc) of primitive to draw
+* @param NumPrimitives The number of primitives in the VertexData buffer
+* @param NumVertices The number of vertices to be written
+* @param VertexDataStride Size of each vertex 
+* @param OutVertexData Reference to the allocated vertex memory
+*/
+void FD3D12CommandContext::RHIBeginDrawPrimitiveUP(uint32 PrimitiveType, uint32 NumPrimitives, uint32 NumVertices, uint32 VertexDataStride, void*& OutVertexData)
+{
+	checkSlow( PendingNumVertices == 0 );
+
+	// Remember the parameters for the draw call.
+	PendingPrimitiveType = PrimitiveType;
+	PendingNumPrimitives = NumPrimitives;
+	PendingNumVertices = NumVertices;
+	PendingVertexDataStride = VertexDataStride;
+
+	// Map the dynamic buffer.
+	OutVertexData = DynamicVB->Lock(NumVertices * VertexDataStride);
+}
+
+/**
+* Draw a primitive using the vertex data populated since RHIBeginDrawPrimitiveUP and clean up any memory as needed
+*/
+void FD3D12CommandContext::RHIEndDrawPrimitiveUP()
+{
+	RHI_DRAW_CALL_STATS(PendingPrimitiveType,PendingNumPrimitives);
+
+	checkSlow(!bUsingTessellation || PendingPrimitiveType == PT_TriangleList);
+
+	OwningRHI.RegisterGPUWork(PendingNumPrimitives, PendingNumVertices);
+
+	// Unmap the dynamic vertex buffer.
+	FD3D12ResourceLocation* BufferLocation = DynamicVB->Unlock();
+	uint32 VBOffset = 0;
+
+	// Issue the draw call.
+	CommitGraphicsResourceTables();
+	CommitNonComputeShaderConstants();
+	StateCache.SetStreamSource(BufferLocation, 0, PendingVertexDataStride, VBOffset);
+	StateCache.SetPrimitiveTopology(GetD3D11PrimitiveType(PendingPrimitiveType,bUsingTessellation));
+	StateCache.ApplyState();
+	numDraws++;
+	CommandListHandle->DrawInstanced(PendingNumVertices, 1, 0, 0);
+#if UE_BUILD_DEBUG	
+	OwningRHI.DrawCount++;
+#endif
+	DEBUG_EXECUTE_COMMAND_LIST(this);
+
+	// Clear these parameters.
+	PendingPrimitiveType = 0;
+	PendingNumPrimitives = 0;
+	PendingNumVertices = 0;
+	PendingVertexDataStride = 0;
+}
+
+/**
+* Preallocate memory or get a direct command stream pointer to fill up for immediate rendering . This avoids memcpys below in DrawIndexedPrimitiveUP
+* @param PrimitiveType The type (triangles, lineloop, etc) of primitive to draw
+* @param NumPrimitives The number of primitives in the VertexData buffer
+* @param NumVertices The number of vertices to be written
+* @param VertexDataStride Size of each vertex
+* @param OutVertexData Reference to the allocated vertex memory
+* @param MinVertexIndex The lowest vertex index used by the index buffer
+* @param NumIndices Number of indices to be written
+* @param IndexDataStride Size of each index (either 2 or 4 bytes)
+* @param OutIndexData Reference to the allocated index memory
+*/
+void FD3D12CommandContext::RHIBeginDrawIndexedPrimitiveUP(uint32 PrimitiveType, uint32 NumPrimitives, uint32 NumVertices, uint32 VertexDataStride, void*& OutVertexData, uint32 MinVertexIndex, uint32 NumIndices, uint32 IndexDataStride, void*& OutIndexData)
+{
+	checkSlow((sizeof(uint16) == IndexDataStride) || (sizeof(uint32) == IndexDataStride));
+
+	// Store off information needed for the draw call.
+	PendingPrimitiveType = PrimitiveType;
+	PendingNumPrimitives = NumPrimitives;
+	PendingMinVertexIndex = MinVertexIndex;
+	PendingIndexDataStride = IndexDataStride;
+	PendingNumVertices = NumVertices;
+	PendingNumIndices = NumIndices;
+	PendingVertexDataStride = VertexDataStride;
+
+	// Map dynamic vertex and index buffers.
+	OutVertexData = DynamicVB->Lock(NumVertices * VertexDataStride);
+	OutIndexData = DynamicIB->Lock(NumIndices * IndexDataStride);
+}
+
+/**
+* Draw a primitive using the vertex and index data populated since RHIBeginDrawIndexedPrimitiveUP and clean up any memory as needed
+*/
+void FD3D12CommandContext::RHIEndDrawIndexedPrimitiveUP()
+{
+	// tessellation only supports trilists
+	checkSlow(!bUsingTessellation || PendingPrimitiveType == PT_TriangleList);
+
+	RHI_DRAW_CALL_STATS(PendingPrimitiveType,PendingNumPrimitives);
+
+	OwningRHI.RegisterGPUWork(PendingNumPrimitives, PendingNumVertices);
+
+	// Unmap the dynamic buffers.
+	FD3D12ResourceLocation* VertexBufferLocation = DynamicVB->Unlock();
+	FD3D12ResourceLocation* IndexBufferLocation = DynamicIB->Unlock();
+	uint32 VBOffset = 0;
+
+	// Issue the draw call.
+	CommitGraphicsResourceTables();
+	CommitNonComputeShaderConstants();
+	StateCache.SetStreamSource(VertexBufferLocation, 0, PendingVertexDataStride, VBOffset);
+	StateCache.SetIndexBuffer(IndexBufferLocation, PendingIndexDataStride == sizeof(uint16) ? DXGI_FORMAT_R16_UINT : DXGI_FORMAT_R32_UINT, 0);
+	StateCache.SetPrimitiveTopology(GetD3D11PrimitiveType(PendingPrimitiveType,bUsingTessellation));
+	StateCache.ApplyState();
+
+	numDraws++;
+	CommandListHandle->DrawIndexedInstanced(PendingNumIndices, 1, 0, PendingMinVertexIndex, 0);
+#if UE_BUILD_DEBUG	
+	OwningRHI.DrawCount++;
+#endif
+	DEBUG_EXECUTE_COMMAND_LIST(this);
+
+	// Clear these parameters.
+	PendingPrimitiveType = 0;
+	PendingNumPrimitives = 0;
+	PendingMinVertexIndex = 0;
+	PendingIndexDataStride = 0;
+	PendingNumVertices = 0;
+	PendingNumIndices = 0;
+	PendingVertexDataStride = 0;
+}
+
+// Raster operations.
+void FD3D12CommandContext::RHIClear(bool bClearColor, const FLinearColor& Color, bool bClearDepth, float Depth, bool bClearStencil, uint32 Stencil, FIntRect ExcludeRect)
+{
+	RHIClearMRTImpl(bClearColor, 1, &Color, bClearDepth, Depth, bClearStencil, Stencil, ExcludeRect, true);
+}
+
+void FD3D12CommandContext::RHIClearMRT(bool bClearColor, int32 NumClearColors, const FLinearColor* ClearColorArray, bool bClearDepth, float Depth, bool bClearStencil, uint32 Stencil, FIntRect ExcludeRect)
+{
+	RHIClearMRTImpl(bClearColor, NumClearColors, ClearColorArray, bClearDepth, Depth, bClearStencil, Stencil, ExcludeRect, true);
+}
+
+void FD3D12CommandContext::RHIClearMRTImpl(bool bClearColor, int32 NumClearColors, const FLinearColor* ClearColorArray, bool bClearDepth, float Depth, bool bClearStencil, uint32 Stencil, FIntRect ExcludeRect, bool bForceShaderClear)
+{
+	SCOPE_CYCLE_COUNTER(STAT_D3D12ClearMRT);
+	// Helper struct to record and restore device states RHIClearMRT modifies.
+	class FDeviceStateHelper
+	{
+		enum { ResourceCount = D3D12_COMMONSHADER_INPUT_RESOURCE_SLOT_COUNT };
+		//////////////////////////////////////////////////////////////////////////
+		// Relevant recorded states:
+		FD3D12ShaderResourceView* VertResources[ResourceCount];
+		uint32 NumVertResources;
+		FD3D12BoundShaderState* OldShaderState;
+		D3D12_DEPTH_STENCIL_DESC* pOldDepthStencilState;
+		D3D12_RASTERIZER_DESC* pOldRasterizerState;
+		D3D12_BLEND_DESC* pOldBlendState;
+		uint32 StencilRef;
+		float BlendFactor[4];
+		uint32 SampleMask;
+		//////////////////////////////////////////////////////////////////////////
+		void ReleaseResources()
+		{
+			FD3D12ShaderResourceView** Resources = VertResources;
+			for (uint32 i = 0 ; i < NumVertResources; i++, Resources++)
+			{
+				SAFE_RELEASE(*Resources);
+			}
+		}
+	public:
+		/** The global D3D device's immediate context */
+		FDeviceStateHelper() {}
+
+		void CaptureDeviceState(FD3D12StateCache& StateCache)
+		{
+			StateCache.GetBoundShaderState(&OldShaderState);
+			StateCache.GetShaderResourceViews<SF_Vertex>(0, NumVertResources, &VertResources[0]);
+			StateCache.GetDepthStencilState(&pOldDepthStencilState, &StencilRef);
+			StateCache.GetBlendState(&pOldBlendState, BlendFactor, &SampleMask);
+			StateCache.GetRasterizerState(&pOldRasterizerState);
+		}
+
+		void ClearCurrentVertexResources(FD3D12StateCache& StateCache)
+		{
+			static FD3D12ShaderResourceView* NullResources[ResourceCount] = {};
+			for (uint32 ResourceLoop = 0 ; ResourceLoop < NumVertResources; ResourceLoop++)
+			{
+				StateCache.SetShaderResourceView<SF_Vertex>(NullResources[0],0);
+			}
+		}
+
+		void RestoreDeviceState(FD3D12StateCache& StateCache) 
+		{
+
+			// Restore the old shaders
+			StateCache.SetBoundShaderState(OldShaderState);
+			for (uint32 ResourceLoop = 0 ; ResourceLoop < NumVertResources; ResourceLoop++)
+			{
+				StateCache.SetShaderResourceView<SF_Vertex>(VertResources[ResourceLoop],ResourceLoop);
+			}
+			StateCache.SetDepthStencilState(pOldDepthStencilState, StencilRef);
+			StateCache.SetBlendState(pOldBlendState, BlendFactor, SampleMask);
+			StateCache.SetRasterizerState(pOldRasterizerState);
+
+			ReleaseResources();
+		}
+	};
+
+	{
+		// <0: Auto
+		int32 ClearWithExcludeRects = 2;
+
+#if !(UE_BUILD_SHIPPING || UE_BUILD_TEST)
+		static const auto ExcludeRectCVar = IConsoleManager::Get().FindTConsoleVariableDataInt(TEXT("r.ClearWithExcludeRects"));
+		ClearWithExcludeRects = ExcludeRectCVar->GetValueOnRenderThread();
+#endif
+
+		if(ClearWithExcludeRects >= 2)
+		{
+			// by default use the exclude rect
+			ClearWithExcludeRects = 1;
+
+			if(IsRHIDeviceIntel())
+			{
+				// Disable exclude rect (Intel has fast clear so better we disable)
+				ClearWithExcludeRects = 0;
+			}
+		}
+
+		if(!ClearWithExcludeRects)
+		{
+			// Disable exclude rect
+			ExcludeRect = FIntRect();
+		}
+	}
+
+	FD3D12RenderTargetView* RenderTargetViews[MaxSimultaneousRenderTargets];
+	FD3D12DepthStencilView* DSView = NULL;
+	uint32 NumSimultaneousRTs = 0;
+	StateCache.GetRenderTargets(RenderTargetViews, &NumSimultaneousRTs, &DSView);
+	FD3D12BoundRenderTargets BoundRenderTargets(RenderTargetViews, NumSimultaneousRTs, DSView);
+
+	// Must specify enough clear colors for all active RTs
+	check(!bClearColor || NumClearColors >= BoundRenderTargets.GetNumActiveTargets());
+
+	FD3D12DepthStencilView* DepthStencilView = BoundRenderTargets.GetDepthStencilView();
+
+	// If we're clearing depth or stencil and we have a readonly depth/stencil view bound, we need to use a writable depth/stencil view
+	if (CurrentDepthTexture)
+	{
+		FExclusiveDepthStencil RequestedAccess;
+
+		RequestedAccess.SetDepthStencilWrite(bClearDepth, bClearStencil);
+
+		ensure(RequestedAccess.IsValid(CurrentDSVAccessType));
+	}
+
+	// Determine if we're trying to clear a subrect of the screen
+	bool UseDrawClear = bForceShaderClear;
+	uint32 NumViews = 1;
+	D3D12_VIEWPORT Viewport;
+	StateCache.GetViewports(&NumViews, &Viewport);
+	if (Viewport.TopLeftX > 0 || Viewport.TopLeftY > 0)
+	{
+		UseDrawClear = true;
+	}
+
+	/*	// possible optimization
+	if(ExcludeRect.Width() > 0 && ExcludeRect.Height() > 0 && HardwareHasLinearClearPerformance) 
+	{
+	UseDrawClear = true;
+	}
+	*/
+	if(ExcludeRect.Min.X == 0 && ExcludeRect.Width() == Viewport.Width && ExcludeRect.Min.Y == 0 && ExcludeRect.Height() == Viewport.Height)
+	{
+		// no need to do anything
+		return;
+	}
+
+	D3D12_RECT ScissorRect = { 0 };
+	uint32 NumRects = 1;
+	StateCache.GetScissorRect(&ScissorRect);
+	if (ScissorRect.left > 0
+		|| ScissorRect.right < Viewport.TopLeftX + Viewport.Width
+		|| ScissorRect.top > 0
+		|| ScissorRect.bottom < Viewport.TopLeftY + Viewport.Height)
+	{
+		UseDrawClear = true;
+	}
+
+	if (!UseDrawClear)
+	{
+		uint32 Width = 0;
+		uint32 Height = 0;
+		if (BoundRenderTargets.GetRenderTargetView(0))
+		{
+			FRTVDesc RTVDesc = GetRenderTargetViewDesc(BoundRenderTargets.GetRenderTargetView(0));
+			Width = RTVDesc.Width;
+			Height = RTVDesc.Height;
+		}
+		else if (DepthStencilView)
+		{
+			FD3D12Resource* BaseTexture = DepthStencilView->GetResource();
+			D3D12_RESOURCE_DESC const& Desc = BaseTexture->GetDesc();
+			Width = (uint32)Desc.Width;
+			Height = Desc.Height;
+		}
+
+		if ((Viewport.Width < Width || Viewport.Height < Height) 
+			&& (Viewport.Width > 1 && Viewport.Height > 1))
+		{
+			UseDrawClear = true;
+		}
+	}
+
+	if (UseDrawClear)
+	{
+		// we don't support draw call clears before the RHI is initialized, reorder the code or make sure it's not a draw call clear
+		check(GIsRHIInitialized);
+
+		if (CurrentDepthTexture)
+		{
+			// Clear all texture references to this depth buffer
+			ConditionalClearShaderResource(CurrentDepthTexture->ResourceLocation);
+		}
+
+		// Build new states
+		FBlendStateRHIParamRef BlendStateRHI;
+
+		if (BoundRenderTargets.GetNumActiveTargets() <= 1)
+		{
+			BlendStateRHI = (bClearColor && BoundRenderTargets.GetRenderTargetView(0))
+				? TStaticBlendState<>::GetRHI()
+				: TStaticBlendState<CW_NONE>::GetRHI();
+		}
+		else
+		{
+			BlendStateRHI = (bClearColor && BoundRenderTargets.GetRenderTargetView(0))
+				? TStaticBlendState<>::GetRHI()
+				: TStaticBlendStateWriteMask<CW_NONE,CW_NONE,CW_NONE,CW_NONE,CW_NONE,CW_NONE,CW_NONE,CW_NONE>::GetRHI();
+		}
+
+		FRasterizerStateRHIParamRef RasterizerStateRHI = TStaticRasterizerState<FM_Solid,CM_None>::GetRHI();
+		float BF[4] = {0,0,0,0};
+
+		const FDepthStencilStateRHIParamRef DepthStencilStateRHI = 
+			(bClearDepth && bClearStencil)
+			? TStaticDepthStencilState<
+			true, CF_Always,
+			true,CF_Always,SO_Replace,SO_Replace,SO_Replace,
+			false,CF_Always,SO_Replace,SO_Replace,SO_Replace,
+			0xff,0xff
+			>::GetRHI()
+			: bClearDepth
+			? TStaticDepthStencilState<true, CF_Always>::GetRHI()
+			: bClearStencil
+			? TStaticDepthStencilState<
+			false, CF_Always,
+			true,CF_Always,SO_Replace,SO_Replace,SO_Replace,
+			false,CF_Always,SO_Replace,SO_Replace,SO_Replace,
+			0xff,0xff
+			>::GetRHI()
+			:     TStaticDepthStencilState<false, CF_Always>::GetRHI();
+
+		if (CurrentDepthTexture)
+		{
+			FExclusiveDepthStencil RequestedAccess;
+
+			RequestedAccess.SetDepthStencilWrite(bClearDepth, bClearStencil);
+
+			ValidateExclusiveDepthStencilAccess(RequestedAccess);
+		}
+
+		FD3D12BlendState* BlendState = FD3D12DynamicRHI::ResourceCast(BlendStateRHI);
+		FD3D12RasterizerState* RasterizerState = FD3D12DynamicRHI::ResourceCast(RasterizerStateRHI);
+		FD3D12DepthStencilState* DepthStencilState = FD3D12DynamicRHI::ResourceCast(DepthStencilStateRHI);
+
+		// Store the current device state
+		FDeviceStateHelper OriginalResourceState;
+		OriginalResourceState.CaptureDeviceState(StateCache);
+
+		// Set the cached state objects
+		StateCache.SetBlendState(&BlendState->Desc, BF, 0xffffffff);
+		StateCache.SetDepthStencilState(&DepthStencilState->Desc, Stencil);
+		StateCache.SetRasterizerState(&RasterizerState->Desc);
+		OriginalResourceState.ClearCurrentVertexResources(StateCache);		
+
+		// Set the new shaders
+		auto ShaderMap = GetGlobalShaderMap(GMaxRHIFeatureLevel);
+		TShaderMapRef<TOneColorVS<true> > VertexShader(ShaderMap);
+
+		FOneColorPS* PixelShader = NULL;
+
+		// Set the shader to write to the appropriate number of render targets
+		// On AMD PC hardware, outputting to a color index in the shader without a matching render target set has a significant performance hit
+		if (BoundRenderTargets.GetNumActiveTargets() <= 1)
+		{
+			TShaderMapRef<TOneColorPixelShaderMRT<1> > MRTPixelShader(ShaderMap);
+			PixelShader = *MRTPixelShader;
+		}
+		else if (BoundRenderTargets.GetNumActiveTargets() == 2)
+		{
+			TShaderMapRef<TOneColorPixelShaderMRT<2> > MRTPixelShader(ShaderMap);
+			PixelShader = *MRTPixelShader;
+		}
+		else if (BoundRenderTargets.GetNumActiveTargets() == 3)
+		{
+			TShaderMapRef<TOneColorPixelShaderMRT<3> > MRTPixelShader(ShaderMap);
+			PixelShader = *MRTPixelShader;
+		}
+		else if (BoundRenderTargets.GetNumActiveTargets() == 4)
+		{
+			TShaderMapRef<TOneColorPixelShaderMRT<4> > MRTPixelShader(ShaderMap);
+			PixelShader = *MRTPixelShader;
+		}
+		else if (BoundRenderTargets.GetNumActiveTargets() == 5)
+		{
+			TShaderMapRef<TOneColorPixelShaderMRT<5> > MRTPixelShader(ShaderMap);
+			PixelShader = *MRTPixelShader;
+		}
+		else if (BoundRenderTargets.GetNumActiveTargets() == 6)
+		{
+			TShaderMapRef<TOneColorPixelShaderMRT<6> > MRTPixelShader(ShaderMap);
+			PixelShader = *MRTPixelShader;
+		}
+		else if (BoundRenderTargets.GetNumActiveTargets() == 7)
+		{
+			TShaderMapRef<TOneColorPixelShaderMRT<7> > MRTPixelShader(ShaderMap);
+			PixelShader = *MRTPixelShader;
+		}
+		else if (BoundRenderTargets.GetNumActiveTargets() == 8)
+		{
+			TShaderMapRef<TOneColorPixelShaderMRT<8> > MRTPixelShader(ShaderMap);
+			PixelShader = *MRTPixelShader;
+		}
+
+		{
+			FRHICommandList_RecursiveHazardous RHICmdList(this);
+			SetGlobalBoundShaderState(RHICmdList, GMaxRHIFeatureLevel, GD3D11ClearMRTBoundShaderState[FMath::Max(BoundRenderTargets.GetNumActiveTargets() - 1, 0)], GD3D11Vector4VertexDeclaration.VertexDeclarationRHI, *VertexShader, PixelShader);
+			PixelShader->SetColors(RHICmdList, ClearColorArray, NumClearColors);
+
+			{
+				// Draw a fullscreen quad
+				if(ExcludeRect.Width() > 0 && ExcludeRect.Height() > 0)
+				{
+					// with a hole in it (optimization in case the hardware has non constant clear performance)
+					FVector4 OuterVertices[4];
+					OuterVertices[0].Set( -1.0f,  1.0f, Depth, 1.0f );
+					OuterVertices[1].Set(  1.0f,  1.0f, Depth, 1.0f );
+					OuterVertices[2].Set(  1.0f, -1.0f, Depth, 1.0f );
+					OuterVertices[3].Set( -1.0f, -1.0f, Depth, 1.0f );
+
+					float InvViewWidth = 1.0f / Viewport.Width;
+					float InvViewHeight = 1.0f / Viewport.Height;
+					FVector4 FractionRect = FVector4(ExcludeRect.Min.X * InvViewWidth, ExcludeRect.Min.Y * InvViewHeight, (ExcludeRect.Max.X - 1) * InvViewWidth, (ExcludeRect.Max.Y - 1) * InvViewHeight);
+
+					FVector4 InnerVertices[4];
+					InnerVertices[0].Set( FMath::Lerp(-1.0f,  1.0f, FractionRect.X), FMath::Lerp(1.0f, -1.0f, FractionRect.Y), Depth, 1.0f );
+					InnerVertices[1].Set( FMath::Lerp(-1.0f,  1.0f, FractionRect.Z), FMath::Lerp(1.0f, -1.0f, FractionRect.Y), Depth, 1.0f );
+					InnerVertices[2].Set( FMath::Lerp(-1.0f,  1.0f, FractionRect.Z), FMath::Lerp(1.0f, -1.0f, FractionRect.W), Depth, 1.0f );
+					InnerVertices[3].Set( FMath::Lerp(-1.0f,  1.0f, FractionRect.X), FMath::Lerp(1.0f, -1.0f, FractionRect.W), Depth, 1.0f );
+
+					FVector4 Vertices[10];
+					Vertices[0] = OuterVertices[0];
+					Vertices[1] = InnerVertices[0];
+					Vertices[2] = OuterVertices[1];
+					Vertices[3] = InnerVertices[1];
+					Vertices[4] = OuterVertices[2];
+					Vertices[5] = InnerVertices[2];
+					Vertices[6] = OuterVertices[3];
+					Vertices[7] = InnerVertices[3];
+					Vertices[8] = OuterVertices[0];
+					Vertices[9] = InnerVertices[0];
+
+					DrawPrimitiveUP(RHICmdList, PT_TriangleStrip, 8, Vertices, sizeof(Vertices[0]));
+				}
+				else
+				{
+					// without a hole
+					FVector4 Vertices[4];
+					Vertices[0].Set( -1.0f,  1.0f, Depth, 1.0f );
+					Vertices[1].Set(  1.0f,  1.0f, Depth, 1.0f );
+					Vertices[2].Set( -1.0f, -1.0f, Depth, 1.0f );
+					Vertices[3].Set(  1.0f, -1.0f, Depth, 1.0f );
+					DrawPrimitiveUP(RHICmdList, PT_TriangleStrip, 2, Vertices, sizeof(Vertices[0]));
+				}
+			}
+			// Implicit flush. Always call flush when using a command list in RHI implementations before doing anything else. This is super hazardous.
+		}
+
+		// Restore the original device state
+		OriginalResourceState.RestoreDeviceState(StateCache); 
+	}
+	else
+	{
+		if (bClearColor && BoundRenderTargets.GetNumActiveTargets() > 0)
+		{
+			for (int32 TargetIndex = 0; TargetIndex < BoundRenderTargets.GetNumActiveTargets(); TargetIndex++)
+			{
+				TRefCountPtr<FD3D12RenderTargetView> RTView = BoundRenderTargets.GetRenderTargetView(TargetIndex);
+
+				FD3D12DynamicRHI::TransitionResource(CommandListHandle, RTView, D3D12_RESOURCE_STATE_RENDER_TARGET);
+				numClears++;
+				CommandListHandle->ClearRenderTargetView(RTView->GetView(), (float*)&ClearColorArray[TargetIndex], 0, nullptr);
+			}
+		}
+
+		if ((bClearDepth || bClearStencil) && DepthStencilView)
+		{
+			FExclusiveDepthStencil ExclusiveDepthStencil;
+			uint32 ClearFlags = 0;
+			if (bClearDepth)
+			{
+				ClearFlags |= D3D12_CLEAR_FLAG_DEPTH;
+				check(DepthStencilView->HasDepth());
+				ExclusiveDepthStencil.SetDepthWrite();
+			}
+			if (bClearStencil)
+			{
+				ClearFlags |= D3D12_CLEAR_FLAG_STENCIL;
+				check(DepthStencilView->HasStencil());
+				ExclusiveDepthStencil.SetStencilWrite();
+			}
+
+			if (ExclusiveDepthStencil.IsDepthWrite() && (!DepthStencilView->HasStencil() || ExclusiveDepthStencil.IsStencilWrite()))
+			{
+				// Transition the entire view (Both depth and stencil planes if applicable)
+				// Some DSVs don't have stencil bits.
+				FD3D12DynamicRHI::TransitionResource(CommandListHandle, DepthStencilView, D3D12_RESOURCE_STATE_DEPTH_WRITE);
+			}
+			else 
+			{
+				if (ExclusiveDepthStencil.IsDepthWrite())
+				{
+					// Transition just the depth plane
+					check(ExclusiveDepthStencil.IsDepthWrite() && !ExclusiveDepthStencil.IsStencilWrite());
+					FD3D12DynamicRHI::TransitionResource(CommandListHandle, DepthStencilView->GetResource(), D3D12_RESOURCE_STATE_DEPTH_WRITE, DepthStencilView->GetDepthOnlyViewSubresourceSubset());
+				}
+				else
+				{
+					// Transition just the stencil plane
+					check(!ExclusiveDepthStencil.IsDepthWrite() && ExclusiveDepthStencil.IsStencilWrite());
+					FD3D12DynamicRHI::TransitionResource(CommandListHandle, DepthStencilView->GetResource(), D3D12_RESOURCE_STATE_DEPTH_WRITE, DepthStencilView->GetStencilOnlyViewSubresourceSubset());
+				}
+			}
+			
+			numClears++;
+			CommandListHandle->ClearDepthStencilView(DepthStencilView->GetView(), (D3D12_CLEAR_FLAGS)ClearFlags, Depth, Stencil, 0, nullptr);
+		}
+	}
+
+	OwningRHI.RegisterGPUWork(0);
+
+	DEBUG_EXECUTE_COMMAND_LIST(this);
+}
+
+void FD3D12CommandContext::RHIBeginAsyncComputeJob_DrawThread(EAsyncComputePriority Priority)
+{
+#if USE_ASYNC_COMPUTE_CONTEXT
+#error Implement me!
+#endif
+}
+
+void FD3D12CommandContext::RHIEndAsyncComputeJob_DrawThread(uint32 FenceIndex)
+{
+#if USE_ASYNC_COMPUTE_CONTEXT
+#error Implement me!
+#endif
+}
+
+void FD3D12CommandContext::RHIGraphicsWaitOnAsyncComputeJob(uint32 FenceIndex)
+{
+#if USE_ASYNC_COMPUTE_CONTEXT
+#error Implement me!
+#endif
+}
+
+// Functions to yield and regain rendering control from D3D
+
+void FD3D12DynamicRHI::RHISuspendRendering()
+{
+	// Not supported
+}
+
+void FD3D12DynamicRHI::RHIResumeRendering()
+{
+	// Not supported
+}
+
+bool FD3D12DynamicRHI::RHIIsRenderingSuspended()
+{
+	// Not supported
+	return false;
+}
+
+// Blocks the CPU until the GPU catches up and goes idle.
+void FD3D12DynamicRHI::RHIBlockUntilGPUIdle()
+{
+	// Not really supported
+}
+
+/*
+* Returns the total GPU time taken to render the last frame. Same metric as FPlatformTime::Cycles().
+*/
+uint32 FD3D12DynamicRHI::RHIGetGPUFrameCycles()
+{
+	return GGPUFrameTime;
+}
+
+void FD3D12DynamicRHI::RHIExecuteCommandList(FRHICommandList* CmdList)
+{
+	check(0); // this path has gone stale and needs updated methods, starting at ERCT_SetScissorRect
+}
+
+// NVIDIA Depth Bounds Test interface
+void FD3D12CommandContext::RHIEnableDepthBoundsTest(bool bEnable, float MinDepth, float MaxDepth)
+{
+	//UE_LOG(LogD3D12RHI, Warning, TEXT("RHIEnableDepthBoundsTest not supported on DX12."));
+}
+
+void FD3D12CommandContext::RHISubmitCommandsHint()
+{
+
+}
diff --git a/Engine/Source/Runtime/Windows/D3D12RHI/Private/D3D12ConstantBuffer.cpp b/Engine/Source/Runtime/Windows/D3D12RHI/Private/D3D12ConstantBuffer.cpp
new file mode 100644
--- /dev/null
+++ b/Engine/Source/Runtime/Windows/D3D12RHI/Private/D3D12ConstantBuffer.cpp
@@ -0,0 +1,58 @@
+// Copyright 1998-2014 Epic Games, Inc. All Rights Reserved.
+
+/*=============================================================================
+D3D12ConstantBuffer.cpp: D3D Constant buffer RHI implementation.
+=============================================================================*/
+
+#include "D3D12RHIPrivate.h"
+
+// SHANOND: Looks like UpdateSubresource is going to be the way to update these CBs.
+//			The driver writers are trying to optimize for UpdateSubresource and this
+//			will also avoid any driver renaming issues we may hit with map_discard.
+
+#define MAX_POOL_BUFFERS 1 // use update subresource and plop it into the command stream
+
+namespace D3D12RHI
+{
+	/** Sizes of constant buffers defined in ED3D11ShaderOffsetBuffer. */
+	const uint32 GConstantBufferSizes[MAX_CONSTANT_BUFFER_SLOTS] =
+	{
+		// CBs must be a multiple of 16
+		Align(MAX_GLOBAL_CONSTANT_BUFFER_SIZE, 16),
+	};
+}
+
+// New circular buffer system for faster constant uploads.  Avoids CopyResource and speeds things up considerably
+FD3D12ConstantBuffer::FD3D12ConstantBuffer(FD3D12Device* InParent, uint32 InSize, uint32 SubBuffers) :
+	MaxSize(InSize),
+	ShadowData(NULL),
+	CurrentUpdateSize(0),
+	TotalUpdateSize(0),
+    FD3D12DeviceChild(InParent)
+{
+	InitResource();
+}
+
+FD3D12ConstantBuffer::~FD3D12ConstantBuffer()
+{
+	ReleaseResource();
+}
+
+/**
+* Creates a constant buffer on the device
+*/
+void FD3D12ConstantBuffer::InitDynamicRHI()
+{
+	// New circular buffer system for faster constant uploads.  Avoids CopyResource and speeds things up considerably
+	// aligned for best performance
+	ShadowData = (uint8*)FMemory::Malloc(MaxSize, 16);
+	FMemory::Memzero(ShadowData,MaxSize);
+}
+
+void FD3D12ConstantBuffer::ReleaseDynamicRHI()
+{
+	if(ShadowData)
+	{
+		FMemory::Free(ShadowData);
+	}
+}
