UE4 - D3D12RHI (not working)

diff --git a/Engine/Source/Runtime/Windows/D3D12RHI/Private/D3D12Device.cpp b/Engine/Source/Runtime/Windows/D3D12RHI/Private/D3D12Device.cpp
new file mode 100644
--- /dev/null
+++ b/Engine/Source/Runtime/Windows/D3D12RHI/Private/D3D12Device.cpp
@@ -0,0 +1,877 @@
+// Copyright 1998-2014 Epic Games, Inc. All Rights Reserved.
+
+/*=============================================================================
+D3D12Device.cpp: D3D device RHI implementation.
+=============================================================================*/
+
+#include "D3D12RHIPrivate.h"
+#include "AllowWindowsPlatformTypes.h"
+	#include <delayimp.h>
+#include "HideWindowsPlatformTypes.h"
+
+namespace D3D12RHI
+{
+	extern void EmptyD3DSamplerStateCache();
+}
+using namespace D3D12RHI;
+
+bool D3D12RHI_ShouldCreateWithD3DDebug()
+{
+	// Use a debug device if specified on the command line.
+	return
+		FParse::Param(FCommandLine::Get(),TEXT("d3ddebug")) ||
+		FParse::Param(FCommandLine::Get(),TEXT("d3debug")) ||
+		FParse::Param(FCommandLine::Get(),TEXT("dxdebug"));
+}
+
+bool D3D12RHI_ShouldAllowAsyncResourceCreation()
+{
+	static bool bAllowAsyncResourceCreation = !FParse::Param(FCommandLine::Get(),TEXT("nod3dasync"));
+	return bAllowAsyncResourceCreation;
+}
+
+IMPLEMENT_MODULE(FD3D12DynamicRHIModule, D3D12RHI);
+
+FD3D12DynamicRHI *FD3D12DynamicRHI::SingleD3DRHI = nullptr;
+
+IRHICommandContext* FD3D12DynamicRHI::RHIGetDefaultContext()
+{
+	return static_cast<IRHICommandContext*>(&GetRHIDevice()->GetDefaultCommandContext());
+}
+
+#if D3D12_SUPPORTS_PARALLEL_RHI_EXECUTE
+
+//todo recycle these to avoid alloc
+
+
+class FD3D12CommandContextContainer : public IRHICommandContextContainer
+{
+	FD3D12Device* OwningDevice;
+	FD3D12CommandContext* CmdContext;
+	TArray<FD3D12CommandListHandle> CommandLists;
+
+public:
+
+	/** Custom new/delete with recycling */
+	void* operator new(size_t Size);
+	void operator delete(void *RawMemory);
+
+	FD3D12CommandContextContainer(FD3D12Device* InOwningDevice)
+		: OwningDevice(InOwningDevice), CmdContext(nullptr)
+	{
+		CommandLists.Reserve (16);
+	}
+
+	virtual ~FD3D12CommandContextContainer() override
+	{
+	}
+
+	virtual IRHICommandContext* GetContext() override
+	{
+		check(CmdContext == nullptr);
+		CmdContext = OwningDevice->ObtainCommandContext();
+		check(CmdContext != nullptr);
+
+		if (CmdContext->CommandListHandle == nullptr)
+			CmdContext->OpenCommandList();
+
+		CmdContext->ClearState();
+
+		return CmdContext;
+	}
+
+	virtual void FinishContext() override
+	{
+		// We never "Finish" the default context. It gets submitted when FlushCommands() is called.
+		check(!CmdContext->IsDefaultContext());
+
+		CmdContext->Finish(CommandLists);
+
+		OwningDevice->ReleaseCommandContext(CmdContext);
+		CmdContext = nullptr;
+	}
+
+	virtual void SubmitAndFreeContextContainer(int32 Index, int32 Num) override
+	{
+		if (Index == 0)
+		{
+			check((IsInRenderingThread() || IsInRHIThread()));
+
+			OwningDevice->GetDefaultCommandContext().FlushCommands();
+		}
+
+		// Add the current lists for execution (now or possibly later)
+		for (int32 i = 0; i < CommandLists.Num(); ++i)
+		{
+			OwningDevice->PendingCommandLists.Add(CommandLists[i]);
+			OwningDevice->PendingCommandListsTotalWorkCommands +=
+				CommandLists[i].GetCurrentOwningContext()->numClears +
+				CommandLists[i].GetCurrentOwningContext()->numCopies +
+				CommandLists[i].GetCurrentOwningContext()->numDraws;
+		}
+
+		CommandLists.Reset();
+
+		// Submission occurs when a batch is finished
+		const bool FinalCommandListInBatch = Index == (Num - 1);
+		if (FinalCommandListInBatch && OwningDevice->PendingCommandLists.Num() > 0)
+		{
+#if SUPPORTS_MEMORY_RESIDENCY
+			OwningDevice->GetOwningRHI()->GetResourceResidencyManager().MakeResident();
+#endif
+			OwningDevice->GetCommandListManager().ExecuteCommandLists(OwningDevice->PendingCommandLists);
+			OwningDevice->PendingCommandLists.Reset();
+			OwningDevice->PendingCommandListsTotalWorkCommands = 0;
+		}
+
+		delete this;
+	}
+};
+
+static TLockFreeFixedSizeAllocator<sizeof(FD3D12CommandContextContainer), FThreadSafeCounter> FD3D12CommandContextContainerAllocator;
+
+void* FD3D12CommandContextContainer::operator new(size_t Size)
+{
+	// doesn't support derived classes with a different size
+	check(Size == sizeof(FD3D12CommandContextContainer));
+	return FD3D12CommandContextContainerAllocator.Allocate();
+}
+
+void FD3D12CommandContextContainer::operator delete(void* RawMemory)
+{
+	FD3D12CommandContextContainerAllocator.Free(RawMemory);
+}	
+
+IRHICommandContextContainer* FD3D12DynamicRHI::RHIGetCommandContextContainer()
+{
+	return new FD3D12CommandContextContainer(GetRHIDevice());
+}
+
+#endif // D3D12_SUPPORTS_PARALLEL_RHI_EXECUTE
+
+FD3D12CommandContext::FD3D12CommandContext(FD3D12Device* InParent) :
+	OwningRHI(*InParent->GetOwningRHI()),
+	UploadHeapAllocator(InParent, D3D12_HEAP_TYPE_UPLOAD),
+	bUsingTessellation(false),
+	PendingNumVertices(0),
+	PendingVertexDataStride(0),
+	PendingPrimitiveType(0),
+	PendingNumPrimitives(0),
+	PendingMinVertexIndex(0),
+	PendingNumIndices(0),
+	PendingIndexDataStride(0),
+	CurrentDepthTexture(nullptr),
+	NumSimultaneousRenderTargets(0),
+	NumUAVs(0),
+	CurrentDSVAccessType(FExclusiveDepthStencil::DepthWrite_StencilWrite),
+	bDiscardSharedConstants(false),
+	CommandListHandle(),
+    FD3D12DeviceChild(InParent)
+{
+	FMemory::Memzero(DirtyUniformBuffers, sizeof(DirtyUniformBuffers));
+
+	UploadHeapAllocator.SetCurrentCommandContext(this);
+
+	// Initialize the constant buffers.
+	InitConstantBuffers();
+
+	// Create the dynamic vertex and index buffers used for Draw[Indexed]PrimitiveUP.
+	DynamicVB = new FD3D12DynamicBuffer(GetParentDevice(), UploadHeapAllocator);
+	DynamicIB = new FD3D12DynamicBuffer(GetParentDevice(), UploadHeapAllocator);
+
+	StateCache.Init(GetParentDevice(), this, nullptr);
+}
+
+FD3D12CommandContext::~FD3D12CommandContext()
+{
+	ClearState();
+	StateCache.Clear();	// Clears its descriptor cache too
+
+	DynamicVB = nullptr;
+	DynamicIB = nullptr;
+
+	// Release references to bound uniform buffers.
+	for (int32 Frequency = 0; Frequency < SF_NumFrequencies; ++Frequency)
+	{
+		for (int32 BindIndex = 0; BindIndex < MAX_UNIFORM_BUFFERS_PER_SHADER_STAGE; ++BindIndex)
+		{
+			BoundUniformBuffers[Frequency][BindIndex].SafeRelease();
+		}
+	}
+
+	// Release allocators last because when other objects are deleted, they may be returned to an allocator
+	UploadHeapAllocator.ReleaseAllResources();
+}
+
+FD3D12Device::FD3D12Device(FD3D12DynamicRHI* InOwningRHI, IDXGIFactory4* InDXGIFactory, FD3D12Adapter& InAdapter) :
+    OwningRHI(InOwningRHI),
+    DXGIFactory(InDXGIFactory),
+    DeviceAdapter(InAdapter),
+    Direct3DDevice(nullptr),
+    DxgiAdapter3(nullptr),
+    bDeviceRemoved(false),
+    RTVAllocator(D3D12_DESCRIPTOR_HEAP_TYPE_RTV, 256),
+    DSVAllocator(D3D12_DESCRIPTOR_HEAP_TYPE_DSV, 256),
+    SRVAllocator(D3D12_DESCRIPTOR_HEAP_TYPE_CBV_SRV_UAV, 1024),
+    UAVAllocator(D3D12_DESCRIPTOR_HEAP_TYPE_CBV_SRV_UAV, 1024),
+    SamplerAllocator(D3D12_DESCRIPTOR_HEAP_TYPE_SAMPLER, 128),
+    CBVAllocator(D3D12_DESCRIPTOR_HEAP_TYPE_CBV_SRV_UAV, 1024),
+    SamplerID(0),
+    OcclusionQueryHeap(this, D3D12_QUERY_HEAP_TYPE_OCCLUSION, 32768),
+    PipelineStateCache(this),
+    ResourceHelper(this),
+    DeferredDeletionQueue(this),
+    DefaultBufferAllocator(this),
+    PendingCommandListsTotalWorkCommands(0)
+{
+}
+
+FD3D12DynamicHeapAllocator* FD3D12DynamicRHI::HelperThreadDynamicHeapAllocator = nullptr;
+
+TAutoConsoleVariable<int32> CVarD3D12ZeroBufferSizeInMB(
+	TEXT("d3d12.ZeroBufferSizeInMB"),
+	4,
+	TEXT("The D3D12 RHI needs a static allocation of zeroes to use when streaming textures asynchronously. It should be large enough to support the largest mipmap you need to stream. The default is 4MB."),
+	ECVF_ReadOnly
+	);
+
+FD3D12DynamicRHI::FD3D12DynamicRHI(IDXGIFactory4* InDXGIFactory, FD3D12Adapter& InAdapter) :
+	DXGIFactory(InDXGIFactory),
+	CommitResourceTableCycles(0),
+	CacheResourceTableCalls(0),
+	CacheResourceTableCycles(0),
+	SetShaderTextureCycles(0),
+	SetShaderTextureCalls(0),
+	SetTextureInTableCalls(0),
+	SceneFrameCounter(0),
+	ResourceTableFrameCounter(INDEX_NONE),
+	bForceSingleQueueGPU(false),
+	NumThreadDynamicHeapAllocators(0),
+	ViewportFrameCounter(0),
+	MainAdapter(InAdapter),
+	MainDevice(nullptr)
+{
+	FMemory::Memzero(ThreadDynamicHeapAllocatorArray, sizeof(ThreadDynamicHeapAllocatorArray));
+
+	// The FD3D12DynamicRHI must be a singleton
+	check(SingleD3DRHI == nullptr);
+
+	// This should be called once at the start 
+	check( IsInGameThread() );
+	check( !GIsThreadedRendering );
+
+	SingleD3DRHI = this;
+
+	check(MainAdapter.IsValid());
+	MainDevice = new FD3D12Device(this, DXGIFactory, MainAdapter);
+
+    FeatureLevel = MainAdapter.MaxSupportedFeatureLevel;
+
+    GPUProfilingData.Init(MainDevice);
+
+#if SUPPORTS_MEMORY_RESIDENCY
+	ResourceResidencyManager.Init(MainDevice);
+#endif
+
+	// Allocate a buffer of zeroes. This is used when we need to pass D3D memory
+	// that we don't care about and will overwrite with valid data in the future.
+	ZeroBufferSize = FMath::Max(CVarD3D12ZeroBufferSizeInMB.GetValueOnAnyThread(), 0) * (1 << 20);
+	ZeroBuffer = FMemory::Malloc(ZeroBufferSize);
+	FMemory::Memzero(ZeroBuffer,ZeroBufferSize);
+
+	GPoolSizeVRAMPercentage = 0;
+	GTexturePoolSize = 0;
+	GConfig->GetInt(TEXT("TextureStreaming"), TEXT("PoolSizeVRAMPercentage"), GPoolSizeVRAMPercentage, GEngineIni);
+
+	// Initialize the RHI capabilities.
+	check(FeatureLevel == D3D_FEATURE_LEVEL_11_0 || FeatureLevel == D3D_FEATURE_LEVEL_10_0 );
+
+	if(FeatureLevel == D3D_FEATURE_LEVEL_10_0)
+	{
+		GSupportsDepthFetchDuringDepthTest = false;
+	}
+
+	// ES2 feature level emulation in D3D11
+	if (FParse::Param(FCommandLine::Get(), TEXT("FeatureLevelES2")) && !GIsEditor)
+	{
+		GMaxRHIFeatureLevel = ERHIFeatureLevel::ES2;
+		GMaxRHIShaderPlatform = SP_PCD3D_ES2;
+	}
+	else if ((FParse::Param(FCommandLine::Get(), TEXT("FeatureLevelES31")) || FParse::Param(FCommandLine::Get(), TEXT("FeatureLevelES3_1"))) && !GIsEditor)
+	{
+		GMaxRHIFeatureLevel = ERHIFeatureLevel::ES3_1;
+		GMaxRHIShaderPlatform = SP_PCD3D_ES3_1;
+	}
+	else if (FeatureLevel == D3D_FEATURE_LEVEL_11_0)
+	{
+		GMaxRHIFeatureLevel = ERHIFeatureLevel::SM5;
+		GMaxRHIShaderPlatform = SP_PCD3D_SM5;
+	}
+	else if (FeatureLevel == D3D_FEATURE_LEVEL_10_0)
+	{
+		GMaxRHIFeatureLevel = ERHIFeatureLevel::SM4;
+		GMaxRHIShaderPlatform = SP_PCD3D_SM4;
+	}
+
+	if (FParse::Param(FCommandLine::Get(), TEXT("ForceSingleQueue")))
+	{
+		bForceSingleQueueGPU = true;
+	}
+
+	// Initialize the platform pixel format map.
+	GPixelFormats[ PF_Unknown		].PlatformFormat	= DXGI_FORMAT_UNKNOWN;
+	GPixelFormats[ PF_A32B32G32R32F	].PlatformFormat	= DXGI_FORMAT_R32G32B32A32_FLOAT;
+	GPixelFormats[ PF_B8G8R8A8		].PlatformFormat	= DXGI_FORMAT_B8G8R8A8_TYPELESS;
+	GPixelFormats[ PF_G8			].PlatformFormat	= DXGI_FORMAT_R8_UNORM;
+	GPixelFormats[ PF_G16			].PlatformFormat	= DXGI_FORMAT_R16_UNORM;
+	GPixelFormats[ PF_DXT1			].PlatformFormat	= DXGI_FORMAT_BC1_TYPELESS;
+	GPixelFormats[ PF_DXT3			].PlatformFormat	= DXGI_FORMAT_BC2_TYPELESS;
+	GPixelFormats[ PF_DXT5			].PlatformFormat	= DXGI_FORMAT_BC3_TYPELESS;
+	GPixelFormats[ PF_BC4			].PlatformFormat	= DXGI_FORMAT_BC4_UNORM;
+	GPixelFormats[ PF_UYVY			].PlatformFormat	= DXGI_FORMAT_UNKNOWN;		// TODO: Not supported in D3D11
+#if DEPTH_32_BIT_CONVERSION
+	GPixelFormats[ PF_DepthStencil	].PlatformFormat	= DXGI_FORMAT_R32G8X24_TYPELESS; 
+	GPixelFormats[ PF_DepthStencil	].BlockBytes		= 5;
+	GPixelFormats[ PF_X24_G8 ].PlatformFormat			= DXGI_FORMAT_X32_TYPELESS_G8X24_UINT;
+	GPixelFormats[ PF_X24_G8].BlockBytes				= 5;
+#else
+	GPixelFormats[ PF_DepthStencil	].PlatformFormat	= DXGI_FORMAT_R24G8_TYPELESS;
+	GPixelFormats[ PF_DepthStencil	].BlockBytes		= 4;
+	GPixelFormats[ PF_X24_G8 ].PlatformFormat			= DXGI_FORMAT_X24_TYPELESS_G8_UINT;
+	GPixelFormats[ PF_X24_G8].BlockBytes				= 4;
+#endif
+	GPixelFormats[ PF_ShadowDepth	].PlatformFormat	= DXGI_FORMAT_R16_TYPELESS;
+	GPixelFormats[ PF_ShadowDepth	].BlockBytes		= 2;
+	GPixelFormats[ PF_R32_FLOAT		].PlatformFormat	= DXGI_FORMAT_R32_FLOAT;
+	GPixelFormats[ PF_G16R16		].PlatformFormat	= DXGI_FORMAT_R16G16_UNORM;
+	GPixelFormats[ PF_G16R16F		].PlatformFormat	= DXGI_FORMAT_R16G16_FLOAT;
+	GPixelFormats[ PF_G16R16F_FILTER].PlatformFormat	= DXGI_FORMAT_R16G16_FLOAT;
+	GPixelFormats[ PF_G32R32F		].PlatformFormat	= DXGI_FORMAT_R32G32_FLOAT;
+	GPixelFormats[ PF_A2B10G10R10   ].PlatformFormat    = DXGI_FORMAT_R10G10B10A2_UNORM;
+	GPixelFormats[ PF_A16B16G16R16  ].PlatformFormat    = DXGI_FORMAT_R16G16B16A16_UNORM;
+	GPixelFormats[ PF_D24 ].PlatformFormat				= DXGI_FORMAT_R24G8_TYPELESS;
+	GPixelFormats[ PF_R16F			].PlatformFormat	= DXGI_FORMAT_R16_FLOAT;
+	GPixelFormats[ PF_R16F_FILTER	].PlatformFormat	= DXGI_FORMAT_R16_FLOAT;
+
+	GPixelFormats[ PF_FloatRGB	].PlatformFormat		= DXGI_FORMAT_R11G11B10_FLOAT;
+	GPixelFormats[ PF_FloatRGB	].BlockBytes			= 4;
+	GPixelFormats[ PF_FloatRGBA	].PlatformFormat		= DXGI_FORMAT_R16G16B16A16_FLOAT;
+	GPixelFormats[ PF_FloatRGBA	].BlockBytes			= 8;
+
+	GPixelFormats[ PF_FloatR11G11B10].PlatformFormat	= DXGI_FORMAT_R11G11B10_FLOAT;
+	GPixelFormats[ PF_FloatR11G11B10].BlockBytes		= 4;
+
+	GPixelFormats[ PF_V8U8			].PlatformFormat	= DXGI_FORMAT_R8G8_SNORM;
+	GPixelFormats[ PF_BC5			].PlatformFormat	= DXGI_FORMAT_BC5_UNORM;
+	GPixelFormats[ PF_A1			].PlatformFormat	= DXGI_FORMAT_R1_UNORM; // Not supported for rendering.
+	GPixelFormats[ PF_A8			].PlatformFormat	= DXGI_FORMAT_A8_UNORM;
+	GPixelFormats[ PF_R32_UINT		].PlatformFormat	= DXGI_FORMAT_R32_UINT;
+	GPixelFormats[ PF_R32_SINT		].PlatformFormat	= DXGI_FORMAT_R32_SINT;
+
+	GPixelFormats[ PF_R16_UINT         ].PlatformFormat = DXGI_FORMAT_R16_UINT;
+	GPixelFormats[ PF_R16_SINT         ].PlatformFormat = DXGI_FORMAT_R16_SINT;
+	GPixelFormats[ PF_R16G16B16A16_UINT].PlatformFormat = DXGI_FORMAT_R16G16B16A16_UINT;
+	GPixelFormats[ PF_R16G16B16A16_SINT].PlatformFormat = DXGI_FORMAT_R16G16B16A16_SINT;
+
+	GPixelFormats[ PF_R5G6B5_UNORM	].PlatformFormat	= DXGI_FORMAT_B5G6R5_UNORM;
+	GPixelFormats[ PF_R8G8B8A8		].PlatformFormat	= DXGI_FORMAT_R8G8B8A8_TYPELESS;
+	GPixelFormats[ PF_R8G8			].PlatformFormat	= DXGI_FORMAT_R8G8_UNORM;
+	GPixelFormats[ PF_R32G32B32A32_UINT].PlatformFormat = DXGI_FORMAT_R32G32B32A32_UINT;
+	GPixelFormats[ PF_R16G16_UINT].PlatformFormat = DXGI_FORMAT_R16G16_UINT;
+
+	GPixelFormats[ PF_BC6H			].PlatformFormat	= DXGI_FORMAT_BC6H_UF16;
+	GPixelFormats[ PF_BC7			].PlatformFormat	= DXGI_FORMAT_BC7_TYPELESS;
+
+    // MS - Not doing any feature level checks. D3D12 currently supports these limits.
+    // However this may need to be revisited if new feature levels are introduced with different HW requirement
+	GSupportsSeparateRenderTargetBlendState = true;
+	GMaxTextureDimensions = D3D12_REQ_TEXTURE2D_U_OR_V_DIMENSION;
+	GMaxCubeTextureDimensions = D3D12_REQ_TEXTURECUBE_DIMENSION;
+	GMaxTextureArrayLayers = D3D12_REQ_TEXTURE2D_ARRAY_AXIS_DIMENSION;
+
+	GMaxShadowDepthBufferSizeX = 4096;
+	GMaxShadowDepthBufferSizeY = 4096;
+
+	// Enable multithreading if not in the editor (editor crashes with multithreading enabled).
+	if (!GIsEditor)
+	{
+		GRHISupportsRHIThread = true;
+	}
+	GRHISupportsParallelRHIExecute = D3D12_SUPPORTS_PARALLEL_RHI_EXECUTE;
+}
+
+void FD3D12Device::CreateCommandContexts()
+{
+    check(CommandContextArray.Num() == 0);
+
+	const uint32 NumContexts = FTaskGraphInterface::Get().GetNumWorkerThreads() + 1;
+
+	// We never make the default context free for allocation by the context containers
+	CommandContextArray.Reserve(NumContexts);
+	FreeCommandContexts.Reserve(NumContexts - 1);
+
+	for (uint32 i = 0; i < NumContexts; ++i)
+	{
+		FD3D12CommandContext* NewCmdContext = new FD3D12CommandContext(this);
+
+		// without that the first RHIClear would get a scissor rect of (0,0)-(0,0) which means we get a draw call clear 
+		NewCmdContext->RHISetScissorRect(false, 0, 0, 0, 0);
+
+		CommandContextArray.Add(NewCmdContext);
+
+		// Make available all but the first command context for parallel threads
+		if (i > 0)
+			FreeCommandContexts.Add(CommandContextArray[i]);
+	}
+
+	CommandContextArray[0]->OpenCommandList();
+}
+
+void FD3D12DynamicRHI::Shutdown()
+{
+	check(IsInGameThread() && IsInRenderingThread());  // require that the render thread has been shut down
+
+	// Cleanup the D3D devices.
+	MainDevice->CleanupD3DDevice();
+	delete(MainDevice);
+
+	// Release buffered timestamp queries
+	GPUProfilingData.FrameTiming.ReleaseResource();
+
+	// Release the buffer of zeroes.
+	FMemory::Free(ZeroBuffer);
+	ZeroBuffer = NULL;
+	ZeroBufferSize = 0;
+}
+
+void FD3D12CommandContext::RHIPushEvent(const TCHAR* Name)
+{
+    OwningRHI.PushGPUEvent(Name);
+}
+void FD3D12CommandContext::RHIPopEvent()
+{
+    OwningRHI.PopGPUEvent();
+}
+
+/**
+* Returns a supported screen resolution that most closely matches the input.
+* @param Width - Input: Desired resolution width in pixels. Output: A width that the platform supports.
+* @param Height - Input: Desired resolution height in pixels. Output: A height that the platform supports.
+*/
+void FD3D12DynamicRHI::RHIGetSupportedResolution( uint32 &Width, uint32 &Height )
+{
+	uint32 InitializedMode = false;
+	DXGI_MODE_DESC BestMode;
+	BestMode.Width = 0;
+	BestMode.Height = 0;
+
+	{
+		HRESULT hr = S_OK;
+		TRefCountPtr<IDXGIAdapter> Adapter;
+        hr = DXGIFactory->EnumAdapters(GetRHIDevice()->GetAdapterIndex(), Adapter.GetInitReference());
+		if( DXGI_ERROR_NOT_FOUND == hr )
+		{
+			return;
+		}
+		if( FAILED(hr) )
+		{
+			return;
+		}
+
+		// get the description of the adapter
+		DXGI_ADAPTER_DESC AdapterDesc;
+		VERIFYD3D11RESULT(Adapter->GetDesc(&AdapterDesc));
+
+#ifndef PLATFORM_XBOXONE // No need for display mode enumeration on console
+		// Enumerate outputs for this adapter
+		// TODO: Cap at 1 for default output
+		for(uint32 o = 0;o < 1; o++)
+		{
+			TRefCountPtr<IDXGIOutput> Output;
+			hr = Adapter->EnumOutputs(o,Output.GetInitReference());
+			if(DXGI_ERROR_NOT_FOUND == hr)
+				break;
+			if(FAILED(hr))
+				return;
+
+			// TODO: GetDisplayModeList is a terribly SLOW call.  It can take up to a second per invocation.
+			//  We might want to work around some DXGI badness here.
+			DXGI_FORMAT Format = DXGI_FORMAT_R8G8B8A8_UNORM;
+			uint32 NumModes = 0;
+			hr = Output->GetDisplayModeList(Format,0,&NumModes,NULL);
+			if(hr == DXGI_ERROR_NOT_FOUND)
+			{
+				return;
+			}
+			else if(hr == DXGI_ERROR_NOT_CURRENTLY_AVAILABLE)
+			{
+				UE_LOG(LogD3D12RHI, Fatal,
+					TEXT("This application cannot be run over a remote desktop configuration")
+					);
+				return;
+			}
+			DXGI_MODE_DESC* ModeList = new DXGI_MODE_DESC[ NumModes ];
+			VERIFYD3D11RESULT(Output->GetDisplayModeList(Format,0,&NumModes,ModeList));
+
+			for(uint32 m = 0;m < NumModes;m++)
+			{
+				// Search for the best mode
+
+				// Suppress static analysis warnings about a potentially out-of-bounds read access to ModeList. This is a false positive - Index is always within range.
+				CA_SUPPRESS( 6385 );
+				bool IsEqualOrBetterWidth = FMath::Abs((int32)ModeList[m].Width - (int32)Width) <= FMath::Abs((int32)BestMode.Width - (int32)Width);
+				bool IsEqualOrBetterHeight = FMath::Abs((int32)ModeList[m].Height - (int32)Height) <= FMath::Abs((int32)BestMode.Height - (int32)Height);
+				if(!InitializedMode || (IsEqualOrBetterWidth && IsEqualOrBetterHeight))
+				{
+					BestMode = ModeList[m];
+					InitializedMode = true;
+				}
+			}
+
+			delete[] ModeList;
+		}
+#endif // PLATFORM_XBOXONE
+	}
+
+	check(InitializedMode);
+	Width = BestMode.Width;
+	Height = BestMode.Height;
+}
+
+// Suppress static analysis warnings in FD3D12DynamicRHI::RHIGetAvailableResolutions() about a potentially out-of-bounds read access to ModeList. This is a false positive - Index is always within range.
+#if USING_CODE_ANALYSIS
+MSVC_PRAGMA(warning(push))
+	MSVC_PRAGMA(warning(disable:6385))
+#endif	// USING_CODE_ANALYSIS
+
+	// Re-enable static code analysis warning C6385.
+#if USING_CODE_ANALYSIS
+	MSVC_PRAGMA(warning(pop))
+#endif	// USING_CODE_ANALYSIS
+
+	void FD3D12DynamicRHI::GetBestSupportedMSAASetting( DXGI_FORMAT PlatformFormat, uint32 MSAACount, uint32& OutBestMSAACount, uint32& OutMSAAQualityLevels )
+{
+	//  We disable MSAA for Feature level 10
+	if (GMaxRHIFeatureLevel == ERHIFeatureLevel::SM4)
+	{
+		OutBestMSAACount = 1;
+		OutMSAAQualityLevels = 0;
+		return;
+	}
+
+	// start counting down from current setting (indicated the current "best" count) and move down looking for support
+	for (uint32 SampleCount = MSAACount; SampleCount > 0; SampleCount--)
+	{
+		// The multisampleQualityLevels struct serves as both the input and output to CheckFeatureSupport.
+		D3D12_FEATURE_DATA_MULTISAMPLE_QUALITY_LEVELS multisampleQualityLevels = {};
+		multisampleQualityLevels.SampleCount = SampleCount;
+
+        if (SUCCEEDED(GetRHIDevice()->GetDevice()->CheckFeatureSupport(D3D12_FEATURE_MULTISAMPLE_QUALITY_LEVELS, &multisampleQualityLevels, sizeof(multisampleQualityLevels))))
+		{
+			OutBestMSAACount = SampleCount;
+			OutMSAAQualityLevels = multisampleQualityLevels.NumQualityLevels;
+			break;
+		}
+	}
+
+	return;
+}
+
+uint32 FD3D12DynamicRHI::GetMaxMSAAQuality(uint32 SampleCount)
+{
+	if(SampleCount <= DX_MAX_MSAA_COUNT)
+	{
+		// 0 has better quality (a more even distribution)
+		// higher quality levels might be useful for non box filtered AA or when using weighted samples 
+		return 0;
+		//		return AvailableMSAAQualities[SampleCount];
+	}
+	// not supported
+	return 0xffffffff;
+}
+
+void FD3D12Device::CreateSignatures()
+{
+	// Graphics
+	struct
+	{
+		D3D12_SHADER_VISIBILITY Vis;
+		D3D12_DESCRIPTOR_RANGE_TYPE Type;
+		uint32 Count;
+		uint32 BaseShaderReg;
+	} RangeDesc[Graphics_DescriptorTableCount] =
+	{
+		{ D3D12_SHADER_VISIBILITY_PIXEL, D3D12_DESCRIPTOR_RANGE_TYPE_SRV, MAX_SRVS, 0 },
+		{ D3D12_SHADER_VISIBILITY_PIXEL, D3D12_DESCRIPTOR_RANGE_TYPE_CBV, MAX_CBS, 0 },
+		{ D3D12_SHADER_VISIBILITY_PIXEL, D3D12_DESCRIPTOR_RANGE_TYPE_SAMPLER, D3D12_COMMONSHADER_SAMPLER_SLOT_COUNT, 0 },
+
+		{ D3D12_SHADER_VISIBILITY_VERTEX, D3D12_DESCRIPTOR_RANGE_TYPE_SRV, MAX_SRVS, 0 },
+		{ D3D12_SHADER_VISIBILITY_VERTEX, D3D12_DESCRIPTOR_RANGE_TYPE_CBV, MAX_CBS, 0 },
+		{ D3D12_SHADER_VISIBILITY_VERTEX, D3D12_DESCRIPTOR_RANGE_TYPE_SAMPLER, D3D12_COMMONSHADER_SAMPLER_SLOT_COUNT, 0 },
+
+		{ D3D12_SHADER_VISIBILITY_GEOMETRY, D3D12_DESCRIPTOR_RANGE_TYPE_SRV, MAX_SRVS, 0 },
+		{ D3D12_SHADER_VISIBILITY_GEOMETRY, D3D12_DESCRIPTOR_RANGE_TYPE_CBV, MAX_CBS, 0 },
+		{ D3D12_SHADER_VISIBILITY_GEOMETRY, D3D12_DESCRIPTOR_RANGE_TYPE_SAMPLER, D3D12_COMMONSHADER_SAMPLER_SLOT_COUNT, 0 },
+
+		{ D3D12_SHADER_VISIBILITY_HULL, D3D12_DESCRIPTOR_RANGE_TYPE_SRV, MAX_SRVS, 0 },
+		{ D3D12_SHADER_VISIBILITY_HULL, D3D12_DESCRIPTOR_RANGE_TYPE_CBV, MAX_CBS, 0 },
+		{ D3D12_SHADER_VISIBILITY_HULL, D3D12_DESCRIPTOR_RANGE_TYPE_SAMPLER, D3D12_COMMONSHADER_SAMPLER_SLOT_COUNT, 0 },
+
+		{ D3D12_SHADER_VISIBILITY_DOMAIN, D3D12_DESCRIPTOR_RANGE_TYPE_SRV, MAX_SRVS, 0 },
+		{ D3D12_SHADER_VISIBILITY_DOMAIN, D3D12_DESCRIPTOR_RANGE_TYPE_CBV, MAX_CBS, 0 },
+		{ D3D12_SHADER_VISIBILITY_DOMAIN, D3D12_DESCRIPTOR_RANGE_TYPE_SAMPLER, D3D12_COMMONSHADER_SAMPLER_SLOT_COUNT, 0 },
+
+		{ D3D12_SHADER_VISIBILITY_ALL, D3D12_DESCRIPTOR_RANGE_TYPE_UAV, D3D12_PS_CS_UAV_REGISTER_COUNT, 0 },
+	};
+
+	CD3DX12_ROOT_PARAMETER TableSlots[Graphics_DescriptorTableCount];
+	CD3DX12_DESCRIPTOR_RANGE DescriptorRanges[Graphics_DescriptorTableCount];
+
+	for (uint32 i = 0; i < Graphics_DescriptorTableCount; i++)
+	{
+		DescriptorRanges[i].Init(
+			RangeDesc[i].Type,
+			RangeDesc[i].Count,
+			RangeDesc[i].BaseShaderReg
+			);
+
+		TableSlots[i].InitAsDescriptorTable(1, &DescriptorRanges[i], RangeDesc[i].Vis);
+	}
+
+	{
+		CD3DX12_ROOT_SIGNATURE_DESC RootDesc;
+		RootDesc.Init(Graphics_DescriptorTableCount, TableSlots, 0, nullptr, D3D12_ROOT_SIGNATURE_FLAG_ALLOW_INPUT_ASSEMBLER_INPUT_LAYOUT);
+
+		TRefCountPtr<ID3DBlob> SerializedRS;
+		VERIFYD3D11RESULT(D3D12SerializeRootSignature(&RootDesc, D3D_ROOT_SIGNATURE_VERSION_1, SerializedRS.GetInitReference(), nullptr));
+		VERIFYD3D11RESULT(GetDevice()->CreateRootSignature(0, SerializedRS->GetBufferPointer(), SerializedRS->GetBufferSize(), IID_PPV_ARGS(GraphicsRS.GetInitReference())));
+	}
+
+	// Compute
+	uint32 RangeIndex = 0;
+	DescriptorRanges[RangeIndex].Init(D3D12_DESCRIPTOR_RANGE_TYPE_SRV, MAX_SRVS, 0);
+	TableSlots[RangeIndex].InitAsDescriptorTable(1, &DescriptorRanges[RangeIndex], D3D12_SHADER_VISIBILITY_ALL);
+	++RangeIndex;
+	DescriptorRanges[RangeIndex].Init(D3D12_DESCRIPTOR_RANGE_TYPE_CBV, MAX_CBS, 0);
+	TableSlots[RangeIndex].InitAsDescriptorTable(1, &DescriptorRanges[RangeIndex], D3D12_SHADER_VISIBILITY_ALL);
+	++RangeIndex;
+	DescriptorRanges[RangeIndex].Init(D3D12_DESCRIPTOR_RANGE_TYPE_SAMPLER, D3D12_COMMONSHADER_SAMPLER_SLOT_COUNT, 0);
+	TableSlots[RangeIndex].InitAsDescriptorTable(1, &DescriptorRanges[RangeIndex], D3D12_SHADER_VISIBILITY_ALL);
+	++RangeIndex;
+	DescriptorRanges[RangeIndex].Init(D3D12_DESCRIPTOR_RANGE_TYPE_UAV, D3D12_PS_CS_UAV_REGISTER_COUNT, 0);
+	TableSlots[RangeIndex].InitAsDescriptorTable(1, &DescriptorRanges[RangeIndex], D3D12_SHADER_VISIBILITY_ALL);
+	++RangeIndex;
+
+	{
+		CD3DX12_ROOT_SIGNATURE_DESC RootDesc;
+		RootDesc.Init(RangeIndex, TableSlots, 0, nullptr, D3D12_ROOT_SIGNATURE_FLAG_NONE);
+
+		TRefCountPtr<ID3DBlob> SerializedRS;
+		VERIFYD3D11RESULT(D3D12SerializeRootSignature(&RootDesc, D3D_ROOT_SIGNATURE_VERSION_1, SerializedRS.GetInitReference(), nullptr));
+		VERIFYD3D11RESULT(GetDevice()->CreateRootSignature(0, SerializedRS->GetBufferPointer(), SerializedRS->GetBufferSize(), IID_PPV_ARGS(ComputeRS.GetInitReference())));
+	}
+
+	// ExecuteIndirect command signatures
+	D3D12_COMMAND_SIGNATURE_DESC commandSignatureDesc = {};
+	commandSignatureDesc.NumArgumentDescs = 1;
+	commandSignatureDesc.ByteStride = 20;
+
+	D3D12_INDIRECT_ARGUMENT_DESC indirectParameterDesc[1] = {};
+    indirectParameterDesc[0].Type = D3D12_INDIRECT_ARGUMENT_TYPE_DRAW_INDEXED;
+
+	commandSignatureDesc.pArgumentDescs = indirectParameterDesc;
+
+	VERIFYD3D11RESULT(GetDevice()->CreateCommandSignature(&commandSignatureDesc, nullptr, IID_PPV_ARGS(DrawIndirectCommandSignature.GetInitReference())));
+
+	indirectParameterDesc[0].Type = D3D12_INDIRECT_ARGUMENT_TYPE_DRAW_INDEXED;
+	VERIFYD3D11RESULT(GetDevice()->CreateCommandSignature(&commandSignatureDesc, nullptr, IID_PPV_ARGS(DrawIndexedIndirectCommandSignature.GetInitReference())));
+
+    indirectParameterDesc[0].Type = D3D12_INDIRECT_ARGUMENT_TYPE_DISPATCH;
+	VERIFYD3D11RESULT(GetDevice()->CreateCommandSignature(&commandSignatureDesc, nullptr, IID_PPV_ARGS(DispatchIndirectCommandSignature.GetInitReference())));
+
+	CD3DX12_HEAP_PROPERTIES CounterHeapProperties(D3D12_HEAP_TYPE_UPLOAD);
+	D3D12_RESOURCE_DESC CounterBufferDesc = CD3DX12_RESOURCE_DESC::Buffer(COUNTER_HEAP_SIZE);
+
+	VERIFYD3D11RESULT(
+		GetDevice()->CreateCommittedResource(
+		&CounterHeapProperties,
+		D3D12_HEAP_FLAG_NONE,
+		&CounterBufferDesc,
+		D3D12_RESOURCE_STATE_GENERIC_READ,
+		nullptr,
+		IID_PPV_ARGS(CounterUploadHeap.GetInitReference())
+		)
+		);
+
+	VERIFYD3D11RESULT(
+		CounterUploadHeap->Map (
+		0,
+		NULL,
+		&CounterUploadHeapData)
+		);
+
+	CounterUploadHeapIndex = 0;
+}
+
+void FD3D12Device::SetupAfterDeviceCreation()
+{
+	CreateSignatures();
+
+    PipelineStateCache = FD3D12PipelineStateCache(this);
+    FString GraphicsCacheFile = FPaths::GameContentDir() / TEXT("D3DGraphics.ushaderprecache");
+    FString ComputeCacheFile = FPaths::GameContentDir() / TEXT("D3DCompute.ushaderprecache");
+
+	PipelineStateCache.Init(GraphicsCacheFile, ComputeCacheFile);
+	PipelineStateCache.RebuildFromDiskCache(GraphicsRS, ComputeRS);
+
+	CreateCommandContexts();
+
+	UpdateMSAASettings();
+
+	if (GRHISupportsAsyncTextureCreation)
+	{
+		UE_LOG(LogD3D12RHI, Log, TEXT("Async texture creation enabled"));
+	}
+	else
+	{
+		UE_LOG(LogD3D12RHI, Log, TEXT("Async texture creation disabled: %s"),
+			D3D12RHI_ShouldAllowAsyncResourceCreation() ? TEXT("no driver support") : TEXT("disabled by user"));
+	}
+}
+
+void FD3D12Device::UpdateMSAASettings()
+{	
+	check(DX_MAX_MSAA_COUNT == 8);
+
+	// quality levels are only needed for CSAA which we cannot use with custom resolves
+
+	// 0xffffffff means not available
+	AvailableMSAAQualities[0] = 0xffffffff;
+	AvailableMSAAQualities[1] = 0xffffffff;
+	AvailableMSAAQualities[2] = 0;
+	AvailableMSAAQualities[3] = 0xffffffff;
+	AvailableMSAAQualities[4] = 0;
+	AvailableMSAAQualities[5] = 0xffffffff;
+	AvailableMSAAQualities[6] = 0xffffffff;
+	AvailableMSAAQualities[7] = 0xffffffff;
+	AvailableMSAAQualities[8] = 0;
+}
+
+void FD3D12Device::CleanupD3DDevice()
+{
+	if(GIsRHIInitialized)
+	{
+		// Ensure any pending rendering is finished
+		GetCommandListManager().SignalFrameComplete(true);
+
+		check(Direct3DDevice);
+
+        PipelineStateCache.Close();
+
+		// Reset the RHI initialized flag.
+		GIsRHIInitialized = false;
+
+		check(!GIsCriticalError);
+
+		// Ask all initialized FRenderResources to release their RHI resources.
+		for(TLinkedList<FRenderResource*>::TIterator ResourceIt(FRenderResource::GetResourceList());ResourceIt;ResourceIt.Next())
+		{
+			FRenderResource* Resource = *ResourceIt;
+			check(Resource->IsInitialized());
+			Resource->ReleaseRHI();
+		}
+
+		for(TLinkedList<FRenderResource*>::TIterator ResourceIt(FRenderResource::GetResourceList());ResourceIt;ResourceIt.Next())
+		{
+			ResourceIt->ReleaseDynamicRHI();
+		}
+
+		Viewports.Empty();
+		DrawingViewport = nullptr;
+		SamplerMap.Empty();
+
+		ReleasePooledUniformBuffers();
+		ReleasePooledTextures();
+
+		// Wait for the command queues to flush
+		CommandListManager.WaitForCommandQueueFlush();
+		CopyCommandListManager.WaitForCommandQueueFlush();
+
+		// Delete array index 0 (the default context) last
+		for (int32 i = CommandContextArray.Num() - 1; i > 0; i--)
+		{
+			delete CommandContextArray[i];
+			CommandContextArray[i] = nullptr;
+		}
+
+		// This is needed so the deferred deletion queue empties out completely because the resources in the
+		// queue use the frame fence values. This includes any deleted resources that may have been deleted 
+		// after the above signal frame
+		GetCommandListManager().SignalFrameComplete(true);
+
+		// Cleanup resources
+		DeferredDeletionQueue.Clear();
+
+		// Cleanup the allocator near the end, as some resources may be returned to the allocator
+		DefaultBufferAllocator.FreeDefaultBufferPools();
+		/*
+		// Cleanup thread resources
+        for (uint32 index; (index = InterlockedDecrement(&NumThreadDynamicHeapAllocators)) != (uint32)-1;)
+		{
+            FD3D12DynamicHeapAllocator *pHeapAllocator = ThreadDynamicHeapAllocatorArray[index];
+			pHeapAllocator->ReleaseAllResources();
+			delete(pHeapAllocator);
+		}
+		*/
+
+		// enable the code below to view detailed resource leaks report
+		//TRefCountPtr<ID3D11Debug> d3dDebug;
+		//if (SUCCEEDED(Direct3DDevice->QueryInterface(__uuidof(ID3D11Debug), (void**)d3dDebug.GetInitReference())))
+		//{
+		//          d3dDebug->ReportLiveDeviceObjects((D3D11_RLDO_FLAGS)0x2);
+		//}
+
+		CommandListManager.Destroy();
+		CopyCommandListManager.Destroy();
+
+		// Flush all pending deletes before destroying the device.
+		FRHIResource::FlushPendingDeletes();
+
+		Direct3DDevice = nullptr;
+	}
+}
+
+void FD3D12DynamicRHI::RHIFlushResources()
+{
+	// Nothing to do (yet!)
+}
+
+void FD3D12DynamicRHI::RHIAcquireThreadOwnership()
+{
+#if WITH_EDITOR
+	if (GRHISupportsRHIThread)
+	{
+		SetupRecursiveResources();
+	}
+#endif
+}
+
+void FD3D12DynamicRHI::RHIReleaseThreadOwnership()
+{
+	// Nothing to do
+}
+
+void FD3D12CommandContext::RHIAutomaticCacheFlushAfterComputeShader(bool bEnable)
+{
+	StateCache.AutoFlushComputeShaderCache(bEnable);
+}
+
+void FD3D12CommandContext::RHIFlushComputeShaderCache()
+{
+	StateCache.FlushComputeShaderCache(true);
+}
+
+void* FD3D12DynamicRHI::RHIGetNativeDevice()
+{
+    return (void*)GetRHIDevice()->GetDevice();
+}
diff --git a/Engine/Source/Runtime/Windows/D3D12RHI/Private/D3D12DirectCommandListManager.cpp b/Engine/Source/Runtime/Windows/D3D12RHI/Private/D3D12DirectCommandListManager.cpp
new file mode 100644
--- /dev/null
+++ b/Engine/Source/Runtime/Windows/D3D12RHI/Private/D3D12DirectCommandListManager.cpp
@@ -0,0 +1,470 @@
+#include "D3D12RHIPrivate.h"
+#include "Windows.h"
+
+FD3D12Fence::FD3D12Fence()
+	: CurrentFence(-1)
+	, LastCompletedFence(-1)
+	, hFenceCompleteEvent(INVALID_HANDLE_VALUE)
+{
+	// Create an event
+	hFenceCompleteEvent = CreateEvent(nullptr, false, false, nullptr);
+	check(INVALID_HANDLE_VALUE != hFenceCompleteEvent);
+}
+
+FD3D12Fence::~FD3D12Fence()
+{
+	if (hFenceCompleteEvent != INVALID_HANDLE_VALUE)
+	{
+		CloseHandle(hFenceCompleteEvent);
+		hFenceCompleteEvent = INVALID_HANDLE_VALUE;
+	}
+
+}
+
+void FD3D12Fence::CreateFence(ID3D12Device* pDirect3DDevice, uint64 InitialValue)
+{
+	check(Fence == nullptr);
+
+	VERIFYD3D11RESULT(pDirect3DDevice->CreateFence(InitialValue, D3D12_FENCE_FLAG_NONE, IID_PPV_ARGS(Fence.GetInitReference())));
+	LastCompletedFence = Fence->GetCompletedValue();
+	CurrentFence = LastCompletedFence + 1;
+}
+
+uint64 FD3D12Fence::Signal(FD3D12CommandQueue* pCommandQueue)
+{
+	check(pCommandQueue != nullptr);
+
+	VERIFYD3D11RESULT(pCommandQueue->Signal(Fence.GetReference(), CurrentFence));
+
+	// Save the current fence and increment it
+	const uint64 SignaledFence = CurrentFence++;
+
+	// Return the value that was signaled
+	return SignaledFence;
+}
+
+bool FD3D12Fence::IsFenceFinished(uint64 FenceValue)
+{
+	check(Fence != nullptr);
+
+	// Avoid repeatedly calling GetCompletedValue()
+	if (FenceValue <= LastCompletedFence)
+		return true;
+
+	// Refresh the completed fence value
+	LastCompletedFence = Fence->GetCompletedValue();
+	return FenceValue <= LastCompletedFence;
+}
+
+uint64 FD3D12Fence::GetLastCompletedFence()
+{
+	LastCompletedFence = Fence->GetCompletedValue();
+	return LastCompletedFence;
+}
+
+void FD3D12Fence::WaitForFence(uint64 FenceValue)
+{
+	SCOPE_CYCLE_COUNTER(STAT_D3D12WaitForFenceTime);
+
+	check(Fence != nullptr);
+
+	if (IsFenceFinished(FenceValue))
+		return;
+
+	// We must wait.  Do so with an event handler so we don't oversleep.
+	VERIFYD3D11RESULT(Fence->SetEventOnCompletion(FenceValue, hFenceCompleteEvent));
+
+	// Wait for the event to complete (the event is automatically reset afterwards)
+	const uint32 WaitResult = WaitForSingleObject(hFenceCompleteEvent, INFINITE);
+	check(0 == WaitResult);
+
+	LastCompletedFence = FenceValue;
+}
+
+FD3D12CommandListManager::FD3D12CommandListManager()
+	: NumCommandListsAllocated(0)
+{
+}
+
+FD3D12CommandListManager::~FD3D12CommandListManager()
+{
+	Destroy();
+}
+
+void FD3D12CommandListManager::Destroy()
+{
+	// Wait for the queue to empty
+	WaitForCommandQueueFlush();
+
+	D3DCommandQueue.SafeRelease();
+
+	FD3D12CommandListHandle hList;
+	while (!ReadyLists.IsEmpty())
+		ReadyLists.Dequeue (hList);
+}
+
+void FD3D12CommandListManager::Create(ID3D12Device* InDirect3DDevice,  D3D12_COMMAND_LIST_TYPE InCommandListType, uint32 NumCommandLists)
+{
+	Direct3DDevice = InDirect3DDevice;
+	CommandListType = InCommandListType;
+
+	check(D3DCommandQueue.GetReference() == nullptr);
+	check(ReadyLists.IsEmpty());
+	checkf(NumCommandLists <= 0xffff, TEXT("Exceeded maximum supported command lists"));
+
+    D3D12_COMMAND_QUEUE_DESC CommandQueueDesc = {};
+    CommandQueueDesc.Flags = D3D12_COMMAND_QUEUE_FLAG_NONE;
+    CommandQueueDesc.NodeMask = 0;
+    CommandQueueDesc.Priority = 0;
+    CommandQueueDesc.Type = CommandListType;
+    VERIFYD3D11RESULT(Direct3DDevice->CreateCommandQueue(&CommandQueueDesc, IID_PPV_ARGS(D3DCommandQueue.GetInitReference())));
+
+	for (uint32 i = 0; i < FT_NumTypes; i++)
+	{
+		Fences[i].CreateFence(Direct3DDevice, 0);
+	}
+
+	for (uint32 i = 0; i < NumCommandLists; ++i)
+	{
+		FD3D12CommandListHandle hList = CreateCommandListHandle();
+		ReadyLists.Enqueue(hList);
+	}
+}
+
+FD3D12CommandListHandle FD3D12CommandListManager::BeginCommandList(ID3D12PipelineState* InitialPSO)
+{
+    FD3D12CommandListHandle List;
+
+	bool FoundFreeList = false;
+	{
+		// Need to lock in case of race conditions between the peek and dequeue
+		FScopeLock Lock(&PeekReadyListCS);
+		FD3D12Fence& Fence = Fences[FT_CommandList];
+		FoundFreeList = ReadyLists.Peek(List);
+		if (FoundFreeList && Fence.IsFenceFinished(List.FenceIndex()))
+		{
+			ReadyLists.Dequeue(List);
+		}
+		else
+		{
+			// Command List in the ReadyList is still pending, just allocate another command list instead of waiting
+			FoundFreeList = false;
+		}
+	}
+
+	if (!FoundFreeList)
+	{
+		List = CreateCommandListHandle();
+	}
+
+	List.Reset(InitialPSO);
+
+	return List;
+}
+
+void FD3D12CommandListManager::SignalFrameComplete(bool WaitForCompletion)
+{
+	FD3D12Fence& Fence = Fences[FT_Frame];
+	const uint64 SignaledFence = Fence.Signal(D3DCommandQueue);
+
+	if (WaitForCompletion)
+	{
+		Fence.WaitForFence(SignaledFence);
+	}
+}
+
+void FD3D12CommandListManager::ExecuteCommandList(FD3D12CommandListHandle& hList, bool WaitForCompletion)
+{
+	SCOPE_CYCLE_COUNTER(STAT_D3D12ExecuteCommandListTime);
+
+	TArray<FD3D12CommandListHandle> Lists;
+	Lists.Add(hList);
+
+	ExecuteCommandLists(Lists, WaitForCompletion);
+}
+
+uint64 FD3D12CommandListManager::ExecuteAndIncrementFence(ID3D12CommandList* pD3DCommandLists[], uint64 NumCommandLists, FD3D12Fence &Fence)
+{
+	FScopeLock Lock(&FenceCS);
+
+	// Execute, signal, and wait (if requested)
+#if UE_BUILD_DEBUG
+	if (D3D12RHI_ShouldCreateWithD3DDebug())
+	{
+		// Debug layer will break when a command list does bad stuff. Helps identify the command list in question.
+		for (int32 i = 0; i < NumCommandLists; i++)
+		{
+			D3DCommandQueue->ExecuteCommandLists(1, &(pD3DCommandLists[i]));
+
+#if LOG_EXECUTE_COMMAND_LISTS
+			LogExecuteCommandLists(1, &(pD3DCommandLists[i]));
+#endif
+		}
+	}
+	else
+#endif
+	{
+		D3DCommandQueue->ExecuteCommandLists(NumCommandLists, pD3DCommandLists);
+
+#if LOG_EXECUTE_COMMAND_LISTS
+		LogExecuteCommandLists(NumCommandLists, pD3DCommandLists);
+#endif
+	}
+
+	return Fence.Signal(D3DCommandQueue);
+}
+
+void FD3D12CommandListManager::ExecuteCommandLists(TArray<FD3D12CommandListHandle>& Lists, bool WaitForCompletion)
+{
+	SCOPE_CYCLE_COUNTER(STAT_D3D12ExecuteCommandListTime);
+
+	bool NeedsResourceBarriers = false;
+	for (int32 i = 0; i < Lists.Num(); i++)
+	{
+		FD3D12CommandListHandle& commandList = Lists[i];
+		if (commandList.PendingResourceBarriers().Num() > 0)
+		{
+			NeedsResourceBarriers = true;
+			break;
+		}
+	}
+
+	FD3D12Fence& Fence = Fences[FT_CommandList];
+	uint64 SignaledFence;
+	int32 commandListIndex = 0;
+	int32 barrierCommandListIndex = 0;
+
+	// Close the resource barrier lists, get the raw command list pointers, and enqueue the command list handles
+	// Note: All command lists will share the same fence
+	check(Lists.Num() <= 128);
+	ID3D12CommandList* pD3DCommandLists[128];
+	FD3D12CommandListHandle BarrierCommandList[128];
+	if (NeedsResourceBarriers)
+	{
+#if 0//UE_BUILD_DEBUG	
+		if (!ResourceStateCS.TryLock())
+		{
+			FD3D12DynamicRHI::GetD3DRHI()->SubmissionLockStalls++;
+			// We don't think this will get hit but it's possible. If we do see this happen,
+			// we should evaluate how often and why this is happening
+			check(0); 
+		}
+#endif
+		FScopeLock Lock(&ResourceStateCS);
+
+		for (int32 i = 0; i < Lists.Num(); i++)
+		{
+			FD3D12CommandListHandle& commandList = Lists[i];
+
+			FD3D12CommandListHandle barrierCommandList = {};
+			const uint32 numBarriers = GetResourceBarrierCommandList(commandList, barrierCommandList);
+			if (numBarriers)
+			{
+				// TODO: Unnecessary assignment here, but fixing this will require refactoring GetResourceBarrierCommandList
+				BarrierCommandList[barrierCommandListIndex] = barrierCommandList;
+				barrierCommandListIndex++;
+
+				check(barrierCommandList.FenceIndex() == (uint64)-1);
+				barrierCommandList.Close();
+
+				pD3DCommandLists[commandListIndex] = barrierCommandList.CommandList();
+				commandListIndex++;
+
+				check(commandListIndex < 127);
+			}
+
+			check(commandList.FenceIndex() == (uint64)-1);
+
+			pD3DCommandLists[commandListIndex] = commandList.CommandList();
+			commandListIndex++;
+		}
+		SignaledFence = ExecuteAndIncrementFence(pD3DCommandLists, commandListIndex, Fence);
+	}
+	else
+	{
+		for (int32 i = 0; i < Lists.Num(); i++)
+		{
+			FD3D12CommandListHandle& commandList = Lists[i];
+			check(commandList.FenceIndex() == (uint64)-1);
+
+			pD3DCommandLists[commandListIndex] = commandList.CommandList();
+			commandListIndex++;
+		}
+		SignaledFence = ExecuteAndIncrementFence(pD3DCommandLists, commandListIndex, Fence);
+	}
+
+	for (int32 i = 0; i < Lists.Num(); i++)
+	{
+		FD3D12CommandListHandle& commandList = Lists[i];
+		commandList.SetFenceIndex(SignaledFence);
+		ReadyLists.Enqueue(commandList);
+	}
+
+	for (int32 i = 0; i < barrierCommandListIndex; i++)
+	{
+		FD3D12CommandListHandle& commandList = BarrierCommandList[i];
+		commandList.SetFenceIndex(SignaledFence);
+		ReadyLists.Enqueue(commandList);
+	}
+
+	if (WaitForCompletion)
+	{
+		Fence.WaitForFence(SignaledFence);
+	}
+}
+
+uint32 FD3D12CommandListManager::GetResourceBarrierCommandList(FD3D12CommandListHandle& hList, FD3D12CommandListHandle& hResourceBarrierList)
+{
+	SCOPE_CYCLE_COUNTER(STAT_D3D12ExecuteCommandListTime);
+
+	TArray<FD3D12PendingResourceBarrier>& PendingResourceBarriers = hList.PendingResourceBarriers();
+	const uint32 NumPendingResourceBarriers = PendingResourceBarriers.Num();
+	if (NumPendingResourceBarriers)
+	{
+		// Reserve space for the descs
+		TArray<D3D12_RESOURCE_BARRIER> BarrierDescs;
+		BarrierDescs.Reserve(NumPendingResourceBarriers);
+
+		// Fill out the descs
+		D3D12_RESOURCE_BARRIER desc = {};
+		desc.Type = D3D12_RESOURCE_BARRIER_TYPE_TRANSITION;
+
+		for (uint32 i = 0; i < NumPendingResourceBarriers; ++i)
+		{
+			const FD3D12PendingResourceBarrier& PRB = PendingResourceBarriers[i];
+
+			// Should only be doing this for the few resources that need state tracking
+			check(PRB.Resource->RequiresResourceStateTracking());
+
+			CResourceState* pResourceState = PRB.Resource->GetResourceState();
+			check(pResourceState);
+
+			desc.Transition.Subresource = PRB.SubResource;
+			const D3D12_RESOURCE_STATES before = pResourceState->GetSubresourceState(desc.Transition.Subresource);
+			const D3D12_RESOURCE_STATES after = PRB.State;
+
+			check(before != D3D12_RESOURCE_STATE_TBD && before != D3D12_RESOURCE_STATE_CORRUPT);
+			if (before != after)
+			{
+				desc.Transition.pResource = PRB.Resource->GetResource();
+				desc.Transition.StateBefore = before;
+				desc.Transition.StateAfter = after;
+
+				// Add the desc
+				BarrierDescs.Add(desc);
+			}
+
+			// Update the state to the what it will be after hList executes
+			const D3D12_RESOURCE_STATES LastState = hList.GetResourceState(PRB.Resource).GetSubresourceState(desc.Transition.Subresource);
+			if (before != LastState)
+			{
+				pResourceState->SetSubresourceState(desc.Transition.Subresource, LastState);
+			}
+		}
+
+		// Remove all pendering barriers from the command list
+		PendingResourceBarriers.Reset();
+
+		// Empty tracked resource state for this command list
+		hList.EmptyTrackedResourceState();
+
+		if (BarrierDescs.Num() > 0)
+		{
+			// MSFT: TODO: Use a special pool of "smaller" command lists+allocators, specifically for resource barriers. Makes the cost
+			// of reseting the allocator less. Also consider opening the resource barrier command list while it's being
+			// recorded on another thread.
+			hResourceBarrierList = BeginCommandList(nullptr);
+
+#if DEBUG_RESOURCE_STATES
+			LogResourceBarriers(BarrierDescs.Num(), BarrierDescs.GetData(), hResourceBarrierList.CommandList());
+#endif
+
+			hResourceBarrierList->ResourceBarrier(BarrierDescs.Num(), BarrierDescs.GetData());
+		}
+
+		return BarrierDescs.Num();
+	}
+
+	return 0;
+}
+
+void FD3D12CommandListManager::WaitForCompletion(const FD3D12CLSyncPoint& hSyncPoint)
+{
+	// First compare the list handle's generation to the managed list's current generation.  If it's
+	// an older generation, then we know it was completed a while ago.
+	if ((hSyncPoint.Generation) < hSyncPoint.CommandList.CurrentGeneration())
+		return;
+
+	const uint64 ListFenceValue = hSyncPoint.CommandList.FenceIndex();
+	checkf(ListFenceValue != (uint64)-1, TEXT("You can't wait for an unsubmitted command list to complete.  Kick first!"));
+
+	FD3D12Fence& Fence = Fences[FT_CommandList];
+	Fence.WaitForFence(ListFenceValue);
+}
+
+bool FD3D12CommandListManager::IsFinished(const FD3D12CLSyncPoint& hSyncPoint, uint64 FenceOffset)
+{
+	if (!hSyncPoint)
+		return false;
+
+    if (FenceOffset == 0 && hSyncPoint.Generation < hSyncPoint.CommandList.CurrentGeneration())
+		return true;
+
+	const uint64 ListFenceValue = hSyncPoint.CommandList.FenceIndex() + FenceOffset;
+	FD3D12Fence& Fence = Fences[FT_CommandList];
+	return Fence.IsFenceFinished(ListFenceValue);
+}
+
+CommandListState FD3D12CommandListManager::GetCommandListState(const FD3D12CLSyncPoint& hSyncPoint)
+{
+	check(hSyncPoint);
+
+	// First compare the list handle's generation to the managed list's current generation.  If it's
+	// an older generation, then we know it was completed a while ago.
+	if (hSyncPoint.Generation < hSyncPoint.CommandList.CurrentGeneration())
+		return CommandListState::kFinished;
+
+	uint64 ListFenceValue = hSyncPoint.CommandList.FenceIndex();
+	if (ListFenceValue == (uint64)-1)
+		return CommandListState::kOpen;
+
+	FD3D12Fence& Fence = Fences[FT_CommandList];
+	return Fence.IsFenceFinished(ListFenceValue) ? CommandListState::kFinished : CommandListState::kQueued;
+}
+
+bool FD3D12CommandListManager::IsFinished(const FD3D12FrameSyncPoint& hSyncPoint)
+{
+	FD3D12Fence& Fence = Fences[FT_Frame];
+	const uint64 FrameFenceValue = hSyncPoint.FrameFence;
+	return Fence.IsFenceFinished(FrameFenceValue);
+}
+
+void FD3D12CommandListManager::WaitForCompletion(const FD3D12FrameSyncPoint& hSyncPoint)
+{
+	FD3D12Fence& Fence = Fences[FT_Frame];
+	const uint64 FrameFenceValue = hSyncPoint.FrameFence;
+	Fence.WaitForFence(FrameFenceValue);
+}
+
+void FD3D12CommandListManager::WaitForCommandQueueFlush()
+{
+	if (D3DCommandQueue)
+	{
+		FD3D12Fence& Fence = Fences[FT_Frame];
+		const uint64 SignaledFence = Fence.Signal(D3DCommandQueue);
+		Fence.WaitForFence(SignaledFence);
+	}
+}
+
+FD3D12CommandListHandle FD3D12CommandListManager::CreateCommandListHandle()
+{
+	FD3D12CommandListHandle List = FD3D12CommandListHandle();
+	List.Create(Direct3DDevice, CommandListType, this);
+	const int64 CommandListCount = FPlatformAtomics::InterlockedIncrement(&NumCommandListsAllocated);
+
+	// If we hit this, we should really evaluate why we need so many command lists, especially since
+	// each command list is paired with an allocator. This can consume a lot of memory even if the command lists
+	// are empty
+	check(CommandListCount < MAX_ALLOCATED_COMMAND_LISTS);
+	return List;
+}
\ No newline at end of file
diff --git a/Engine/Source/Runtime/Windows/D3D12RHI/Private/D3D12DirectCommandListManager.h b/Engine/Source/Runtime/Windows/D3D12RHI/Private/D3D12DirectCommandListManager.h
new file mode 100644
--- /dev/null
+++ b/Engine/Source/Runtime/Windows/D3D12RHI/Private/D3D12DirectCommandListManager.h
@@ -0,0 +1,109 @@
+#pragma once
+#include "Core.h"
+#include "RHI.h"
+#include "Engine.h"
+#include "Windows/D3D12RHIBasePrivate.h"
+#define DEFAULT_MAIN_POOL_COMMAND_LISTS 3
+#define MAX_ALLOCATED_COMMAND_LISTS 256
+
+enum class CommandListState
+{
+	kOpen,
+	kQueued,
+	kFinished
+};
+
+enum EFenceType
+{
+	FT_CommandList,
+	FT_Frame,
+	FT_NumTypes
+};
+
+class FD3D12Fence
+{
+public:
+	FD3D12Fence();
+	~FD3D12Fence();
+
+	void CreateFence(ID3D12Device* pDirect3DDevice, uint64 InitialValue = 0);
+	uint64 Signal(FD3D12CommandQueue* pCommandQueue);
+	bool IsFenceFinished(uint64 FenceValue);
+	void WaitForFence(uint64 FenceValue);
+
+	uint64 GetCurrentFence() const { return CurrentFence; }
+	uint64 GetLastCompletedFence();
+
+private:
+	TRefCountPtr<ID3D12Fence> Fence;
+	uint64 CurrentFence;
+	uint64 LastCompletedFence;
+	HANDLE hFenceCompleteEvent;
+};
+
+class FD3D12CommandListManager
+{
+public:
+	FD3D12CommandListManager();
+	~FD3D12CommandListManager();
+
+	void Create(ID3D12Device* Direct3DDevice, D3D12_COMMAND_LIST_TYPE CommandListType, uint32 NumCommandLists = 16);
+	void Destroy();
+
+	inline bool IsReady()
+	{
+		return D3DCommandQueue.GetReference() != nullptr;
+	}
+
+	FD3D12CommandListHandle BeginCommandList(ID3D12PipelineState* InitialPSO = nullptr);
+	void ExecuteCommandList(FD3D12CommandListHandle& hList, bool WaitForCompletion = false);
+	void ExecuteCommandLists(TArray<FD3D12CommandListHandle>& Lists, bool WaitForCompletion = false);
+
+	uint32 GetResourceBarrierCommandList(FD3D12CommandListHandle& hList, FD3D12CommandListHandle& hResourceBarrierList);
+
+	void SignalFrameComplete(bool WaitForCompletion = false);
+
+	CommandListState GetCommandListState(const FD3D12CLSyncPoint& hSyncPoint);
+	bool IsFinished(const FD3D12CLSyncPoint& hSyncPoint, uint64 FenceOffset = 0);
+	void WaitForCompletion(const FD3D12CLSyncPoint& hSyncPoint);
+
+	bool IsFinished(const FD3D12FrameSyncPoint& hSyncPoint);
+	void WaitForCompletion(const FD3D12FrameSyncPoint& hSyncPoint);
+
+	inline HRESULT GetTimestampFrequency(uint64* Frequency)
+	{
+		return D3DCommandQueue->GetTimestampFrequency(Frequency);
+	}
+
+    inline FD3D12CommandQueue* GetD3DCommandQueue()
+    {
+        return D3DCommandQueue.GetReference();
+    }
+
+	inline FD3D12Fence& GetFence(EFenceType FenceType)
+	{
+		check(FenceType < FT_NumTypes);
+		return Fences[FenceType];
+	}
+
+	void WaitForCommandQueueFlush();
+
+private:
+	// Returns signaled Fence
+	uint64 ExecuteAndIncrementFence(ID3D12CommandList* pD3DCommandLists[], uint64 NumCommandLists, FD3D12Fence &Fence);
+	FD3D12CommandListHandle CreateCommandListHandle();
+
+	TRefCountPtr<FD3D12CommandQueue>		D3DCommandQueue;
+
+	FThreadsafeQueue<FD3D12CommandListHandle> ReadyLists;
+
+	FD3D12Fence Fences[FT_NumTypes];
+	
+	int64									NumCommandListsAllocated;
+	D3D12_COMMAND_LIST_TYPE					CommandListType;
+	ID3D12Device*							Direct3DDevice;
+	FCriticalSection						PeekReadyListCS;
+	FCriticalSection						ResourceStateCS;
+	FCriticalSection						FenceCS;
+};
+
diff --git a/Engine/Source/Runtime/Windows/D3D12RHI/Private/D3D12IndexBuffer.cpp b/Engine/Source/Runtime/Windows/D3D12RHI/Private/D3D12IndexBuffer.cpp
new file mode 100644
--- /dev/null
+++ b/Engine/Source/Runtime/Windows/D3D12RHI/Private/D3D12IndexBuffer.cpp
@@ -0,0 +1,219 @@
+// Copyright 1998-2014 Epic Games, Inc. All Rights Reserved.
+
+/*=============================================================================
+	D3D12IndexBuffer.cpp: D3D Index buffer RHI implementation.
+=============================================================================*/
+
+#include "D3D12RHIPrivate.h"
+
+FIndexBufferRHIRef FD3D12DynamicRHI::RHICreateIndexBuffer(uint32 Stride, uint32 Size, uint32 InUsage, FRHIResourceCreateInfo& CreateInfo)
+{
+	// Explicitly check that the size is nonzero before allowing CreateIndexBuffer to opaquely fail.
+	check(Size > 0);
+
+	// Describe the index buffer.
+    const bool bIsDynamic = (InUsage & BUF_AnyDynamic) ? true : false;
+    D3D12_RESOURCE_DESC Desc = CD3DX12_RESOURCE_DESC::Buffer(Size);
+
+	if (InUsage & BUF_UnorderedAccess)
+	{
+        Desc.Flags |= D3D12_RESOURCE_FLAG_ALLOW_UNORDERED_ACCESS;
+	}
+
+	if ((InUsage & BUF_ShaderResource) == 0)
+	{
+        Desc.Flags |= D3D12_RESOURCE_FLAG_DENY_SHADER_RESOURCE;
+	}
+
+	// If a resource array was provided for the resource, create the resource pre-populated
+	D3D12_SUBRESOURCE_DATA InitData = { 0 };
+	D3D12_SUBRESOURCE_DATA* pInitData = nullptr;
+	if (CreateInfo.ResourceArray)
+	{
+		check(Size == CreateInfo.ResourceArray->GetResourceDataSize());
+		InitData.pData = CreateInfo.ResourceArray->GetResourceData();
+        InitData.RowPitch = Size;
+        InitData.SlicePitch = 0;
+		pInitData = &InitData;
+	}
+
+	TRefCountPtr<FD3D12ResourceLocation> ResourceLocation;
+    if (bIsDynamic)
+	{
+		// Create an empty resource location
+		ResourceLocation = new FD3D12ResourceLocation(GetRHIDevice(), Size);
+
+		if (pInitData)
+		{
+			// Handle initial data
+			void *pData = GetRHIDevice()->GetDefaultUploadHeapAllocator().Alloc(Size, D3D12_TEXTURE_DATA_PLACEMENT_ALIGNMENT, ResourceLocation);
+			FMemory::Memcpy(pData, InitData.pData, Size);
+		}
+	}
+	else
+	{
+        ResourceLocation = new FD3D12ResourceLocation(GetRHIDevice(), Size);
+        VERIFYD3D11RESULT(GetRHIDevice()->GetDefaultBufferAllocator().AllocDefaultResource(Desc, pInitData, ResourceLocation));
+	}
+
+	if (CreateInfo.ResourceArray)
+	{
+		// Discard the resource array's contents.
+		CreateInfo.ResourceArray->Discard();
+	}
+
+	UpdateBufferStats(ResourceLocation, true, D3D12_BUFFER_TYPE_INDEX);
+
+	return new FD3D12IndexBuffer(GetRHIDevice(), ResourceLocation, Stride, Size, InUsage);
+}
+
+FD3D12IndexBuffer::~FD3D12IndexBuffer()
+{
+	UpdateBufferStats(ResourceLocation, false, D3D12_BUFFER_TYPE_INDEX);
+
+	// Dynamic resources have a permanent entry in the OutstandingLocks that needs to be cleaned up. 
+	if (GetUsage() & BUF_AnyDynamic)
+	{
+		FD3D12LockedKey LockedKey(ResourceLocation.GetReference());
+		FD3D12DynamicRHI::GetD3DRHI()->OutstandingLocks.Remove (LockedKey);
+	}
+}
+
+void* FD3D12DynamicRHI::RHILockIndexBuffer(FIndexBufferRHIParamRef IndexBufferRHI,uint32 Offset,uint32 Size,EResourceLockMode LockMode)
+{
+	FD3D12IndexBuffer* IndexBuffer = FD3D12DynamicRHI::ResourceCast(IndexBufferRHI);
+	
+    FD3D12CommandContext& DefaultContext = GetRHIDevice()->GetDefaultCommandContext();
+
+	// If this resource is bound to the device, unbind it
+	DefaultContext.ConditionalClearShaderResource(IndexBuffer->ResourceLocation);
+
+	// Determine whether the index buffer is dynamic or not.
+	const bool bIsDynamic = (IndexBuffer->GetUsage() & BUF_AnyDynamic) ? true : false;
+
+	FD3D12LockedKey LockedKey(IndexBuffer->ResourceLocation.GetReference());
+	FD3D12LockedData LockedData;
+    uint32 ResourceSize = IndexBuffer->ResourceLocation->GetEffectiveBufferSize();
+
+	if(bIsDynamic)
+	{
+		check(LockMode == RLM_WriteOnly);
+
+		// Use an upload heap for dynamic resources and map its memory for writing.
+        void *pData = GetRHIDevice()->GetDefaultUploadHeapAllocator().Alloc(Offset + Size, D3D12_TEXTURE_DATA_PLACEMENT_ALIGNMENT, IndexBuffer->ResourceLocation.GetReference());
+		check(pData);
+
+		// Add the lock to the lock map.
+		FD3D12LockedKey LockedKey(IndexBuffer->ResourceLocation.GetReference());
+		FD3D12LockedData LockedData;
+		LockedData.SetData(pData);
+		LockedData.Pitch = Offset + Size;
+		OutstandingLocks.Add(LockedKey, LockedData);
+
+		return (void *) ((uint8 *) pData + Offset);
+	}
+	else
+	{
+		if(LockMode == RLM_ReadOnly)
+		{
+			FD3D12Resource *Resource = IndexBuffer->ResourceLocation->GetResource();
+
+			// If the static buffer is being locked for reading, create a staging buffer.
+			TRefCountPtr<FD3D12Resource> StagingIndexBuffer;
+            VERIFYD3D11RESULT(GetRHIDevice()->GetResourceHelper().CreateBuffer(D3D12_HEAP_TYPE_READBACK, Offset + Size, StagingIndexBuffer.GetInitReference()));
+			LockedData.StagingResource = StagingIndexBuffer;
+
+			// Copy the contents of the index buffer to the staging buffer.
+			DefaultContext.numCopies++;
+			DefaultContext.CommandListHandle->CopyResource(StagingIndexBuffer->GetResource(), Resource->GetResource());
+
+			// Execute the command list and wait for the contents to be ready so it can be returned to the caller
+            GetRHIDevice()->GetDefaultCommandContext().FlushCommands(true);
+
+			// Map the staging buffer's memory for reading.
+			void* pData;
+			VERIFYD3D11RESULT(StagingIndexBuffer->GetResource()->Map(0, nullptr, &pData));
+			LockedData.SetData(pData);
+			LockedData.Pitch = Offset + Size;
+		}
+		else
+		{
+			// If the static buffer is being locked for writing, allocate memory for the contents to be written to.
+
+			// Use an upload heap to copy data to a default resource.
+			TRefCountPtr<FD3D12ResourceLocation> UploadBuffer = new FD3D12ResourceLocation();
+            void *pData = GetRHIDevice()->GetDefaultUploadHeapAllocator().FastAlloc(Offset + Size, D3D12_TEXTURE_DATA_PLACEMENT_ALIGNMENT, UploadBuffer);
+
+			// Add the lock to the lock map.
+			LockedData.SetData(pData);
+			LockedData.Pitch = Offset + Size;
+
+			// Keep track of the underlying resource for the upload heap so it can be referenced during Unmap.
+			LockedData.UploadHeapLocation = UploadBuffer;
+		}
+	}
+	
+	// Add the lock to the lock map.
+	OutstandingLocks.Add(LockedKey,LockedData);
+
+	// Return the offset pointer
+	return (void*)((uint8*)LockedData.GetData() + Offset);
+}
+
+void FD3D12DynamicRHI::RHIUnlockIndexBuffer(FIndexBufferRHIParamRef IndexBufferRHI)
+{
+	FD3D12IndexBuffer* IndexBuffer = FD3D12DynamicRHI::ResourceCast(IndexBufferRHI);
+
+	// Determine whether the index buffer is dynamic or not.
+	const bool bIsDynamic = (IndexBuffer->GetUsage() & BUF_AnyDynamic) ? true : false;
+
+	// Find the outstanding lock for this IB.
+	FD3D12LockedKey LockedKey(IndexBuffer->ResourceLocation.GetReference());
+	FD3D12LockedData* LockedData = OutstandingLocks.Find(LockedKey);
+	check(LockedData);
+
+	FD3D12Resource *Resource = IndexBuffer->ResourceLocation->GetResource();
+
+	if(bIsDynamic)
+	{
+		// If the IB is dynamic, its upload heap memory can always stay mapped. Don't do anything.
+	}
+	else
+	{
+		// If the static IB lock involved a staging resource, it was locked for reading.
+		if(LockedData->StagingResource)
+		{
+			// Unmap the staging buffer's memory.
+			FD3D12Resource* StagingBuffer = LockedData->StagingResource.GetReference();
+			StagingBuffer->GetResource()->Unmap(0, nullptr);
+		}
+		else 
+		{
+			check(LockedData->UploadHeapLocation != nullptr);
+
+			// Copy the contents of the temporary memory buffer allocated for writing into the IB.
+			FD3D12ResourceLocation* UploadHeap = LockedData->UploadHeapLocation;
+
+			FD3D12CommandListHandle& hCommandList = GetRHIDevice()->GetDefaultCommandContext().CommandListHandle;
+
+			{
+				FScopeResourceBarrier ScopeResourceBarrier(hCommandList, Resource, D3D12_RESOURCE_STATE_GENERIC_READ, D3D12_RESOURCE_STATE_COPY_DEST, 0);
+				// Don't need to transition upload heaps
+
+				GetRHIDevice()->GetDefaultCommandContext().numCopies++;
+				hCommandList->CopyBufferRegion(
+					Resource->GetResource(), IndexBuffer->ResourceLocation->GetOffset(),
+					UploadHeap->GetResource()->GetResource(), UploadHeap->GetOffset(),
+					LockedData->Pitch);
+			}
+
+			
+			DEBUG_RHI_EXECUTE_COMMAND_LIST(this);
+		}
+	}
+
+	// Remove the FD3D12LockedData from the lock map.
+	// If the lock involved a staging resource, this releases it.
+	// If the lock involved an upload heap resource, this releases it.
+	OutstandingLocks.Remove(LockedKey);
+}
diff --git a/Engine/Source/Runtime/Windows/D3D12RHI/Private/D3D12Query.cpp b/Engine/Source/Runtime/Windows/D3D12RHI/Private/D3D12Query.cpp
new file mode 100644
--- /dev/null
+++ b/Engine/Source/Runtime/Windows/D3D12RHI/Private/D3D12Query.cpp
@@ -0,0 +1,639 @@
+// Copyright 1998-2014 Epic Games, Inc. All Rights Reserved.
+
+/*=============================================================================
+	D3D12Query.cpp: D3D query RHI implementation.
+=============================================================================*/
+
+#include "D3D12RHIPrivate.h"
+
+namespace D3D12RHI
+{
+/**
+* RHI console variables used by queries.
+*/
+namespace RHIConsoleVariables
+{
+    int32 bStablePowerState = 0;
+    static FAutoConsoleVariableRef CVarStablePowerState(
+        TEXT("D3D12RHI.StablePowerState"),
+        bStablePowerState,
+        TEXT("If true, enable stable power state. This increases GPU timing measurement accuracy but may decrease overall GPU clock rate."),
+        ECVF_Default
+        );
+}
+}
+using namespace D3D12RHI;
+
+FRenderQueryRHIRef FD3D12DynamicRHI::RHICreateRenderQuery(ERenderQueryType QueryType)
+{
+    TRefCountPtr<ID3D12QueryHeap> queryHeap;
+    TRefCountPtr<ID3D12Resource> queryResultBuffer;
+
+	if(QueryType == RQT_Occlusion)
+	{
+        // Don't actually allocate an element until BeginQuery
+        GetRHIDevice()->GetQueryHeap()->ReserveElement();
+	}
+	else if(QueryType == RQT_AbsoluteTime)
+	{
+        D3D12_QUERY_HEAP_DESC queryHeapDesc = {};
+        queryHeapDesc.Count = 1;
+
+		queryHeapDesc.Type = D3D12_QUERY_HEAP_TYPE_TIMESTAMP;
+
+        VERIFYD3D11RESULT(GetRHIDevice()->GetDevice()->CreateQueryHeap(&queryHeapDesc, IID_PPV_ARGS(queryHeap.GetInitReference())));
+
+        D3D12_HEAP_PROPERTIES heapProperties = CD3DX12_HEAP_PROPERTIES(D3D12_HEAP_TYPE_READBACK);
+
+        D3D12_RESOURCE_DESC heapDesc = CD3DX12_RESOURCE_DESC::Buffer(8);
+
+        VERIFYD3D11RESULT(
+            GetRHIDevice()->GetDevice()->CreateCommittedResource(
+            &heapProperties,
+            D3D12_HEAP_FLAG_NONE,
+            &heapDesc,
+            D3D12_RESOURCE_STATE_COPY_DEST,
+            nullptr,
+            IID_PPV_ARGS(queryResultBuffer.GetInitReference())
+            )
+            );
+	}
+	else
+	{
+		check(0);
+	}
+
+	return new FD3D12OcclusionQuery(queryHeap, queryResultBuffer, QueryType);
+}
+
+bool FD3D12DynamicRHI::RHIGetRenderQueryResult(FRenderQueryRHIParamRef QueryRHI,uint64& OutResult,bool bWait)
+{
+	check(IsInRenderingThread());
+	FD3D12OcclusionQuery* Query = FD3D12DynamicRHI::ResourceCast(QueryRHI);
+
+	bool bSuccess = true;
+	if (Query != NULL)
+	{
+		if (!Query->bResultIsCached)
+		{
+            Query->bResultIsCached = GetRHIDevice()->GetQueryData(*Query, bWait);
+		}
+
+		if (Query->Type == RQT_AbsoluteTime)
+		{
+			// GetTimingFrequency is the number of ticks per second
+			uint64 Div = FGPUTiming::GetTimingFrequency() / (1000 * 1000);
+
+			// convert from GPU specific timestamp to micro sec (1 / 1 000 000 s) which seems a reasonable resolution
+			OutResult = Query->Result / Div;
+		}
+		else
+		{
+			OutResult = Query->Result;
+		}
+	}
+	return bSuccess;
+}
+
+bool FD3D12Device::GetQueryData(FD3D12OcclusionQuery& Query, bool bWait)
+{
+	FD3D12CommandContext& CmdContext = (Query.Type == RQT_Occlusion) ? OcclusionQueryHeap.GetOwningContext() : *Query.OwningContext;
+	FD3D12CLSyncPoint SyncPoint = (Query.Type == RQT_Occlusion) ? OcclusionQueryHeap.GetSyncPoint() : Query.OwningCommandList;
+    CommandListState ListState = GetCommandListManager().GetCommandListState(SyncPoint);
+
+	if (ListState != CommandListState::kFinished)
+	{
+		if (!bWait)
+			return false;
+
+		if (ListState == CommandListState::kOpen)
+			CmdContext.FlushCommands(true);
+		else
+			GetCommandListManager().WaitForCompletion(SyncPoint);
+	}
+	
+	// Read the data from the query's buffer.
+	if (Query.Type == RQT_Occlusion)
+	{
+		const uint32 BeginOffset = Query.HeapIndex * sizeof(Query.Result);
+		CD3DX12_RANGE range(BeginOffset, BeginOffset + sizeof(Query.Result));
+		VERIFYD3D11RESULT(OcclusionQueryHeap.MapResultBufferRange(range));
+		FMemory::Memcpy(&Query.Result, OcclusionQueryHeap.GetResultBufferData(range), sizeof(Query.Result));
+		OcclusionQueryHeap.UnmapResultBufferRange();
+        return true;
+	}
+	else
+	{
+        void* pSrcData = nullptr;
+        HRESULT Result = Query.ResultBuffer->Map(0, nullptr, &pSrcData);
+        if (Result == S_OK)
+        {
+            // Copy the data from the query result buffer to the destination
+            FMemory::Memcpy(&Query.Result, pSrcData, sizeof(Query.Result));
+
+            // Unmap
+            Query.ResultBuffer->Unmap(0, nullptr);
+
+            return true;
+        }
+        else if (Result == DXGI_ERROR_DEVICE_REMOVED || Result == DXGI_ERROR_DEVICE_RESET || Result == DXGI_ERROR_DRIVER_INTERNAL_ERROR)
+        {
+            bDeviceRemoved = true;
+            return false;
+        }
+        else
+        {
+            VERIFYD3D11RESULT(Result);
+            return false;
+        }
+	}
+}
+
+void FD3D12CommandContext::RHIBeginOcclusionQueryBatch()
+{
+	GetParentDevice()->GetQueryHeap()->StartQueryBatch(*this);
+}
+
+void FD3D12CommandContext::RHIEndOcclusionQueryBatch()
+{
+	GetParentDevice()->GetQueryHeap()->EndQueryBatchAndResolveQueryData(*this, D3D12_QUERY_TYPE_OCCLUSION);
+	
+	// MSFT: Seb: Do we still need to do this if we support parallel execute?
+	// Break up the command list here so that the wait on the previous frame's results don't block.
+	FlushCommands();
+}
+
+/*=============================================================================
+* class FD3D12QueryHeap
+*=============================================================================*/
+
+FD3D12QueryHeap::FD3D12QueryHeap(FD3D12Device* InParent, const D3D12_QUERY_HEAP_TYPE &InQueryHeapType, uint32 InQueryHeapCount)
+:HeadActiveElement(0)
+, TailActiveElement(0)
+, ActiveAllocatedElementCount(0)
+, LastAllocatedElement(InQueryHeapCount - 1)
+, ReservedElementCount(0)
+, ResultSize(8)
+, OwningContext(nullptr)
+, pResultData(nullptr)
+, FD3D12DeviceChild(InParent)
+, LastBatch(MAX_ACTIVE_BATCHES - 1)
+{
+	if (InQueryHeapType == D3D12_QUERY_HEAP_TYPE_OCCLUSION)
+		QueryType = D3D12_QUERY_TYPE_OCCLUSION;
+
+	else if (InQueryHeapType == D3D12_QUERY_HEAP_TYPE_TIMESTAMP)
+		QueryType = D3D12_QUERY_TYPE_TIMESTAMP;
+
+    // Setup the query heap desc
+    QueryHeapDesc = {};
+    QueryHeapDesc.Type = InQueryHeapType;
+    QueryHeapDesc.Count = InQueryHeapCount;
+
+    CurrentQueryBatch.Clear();
+
+    ActiveQueryBatches.Reserve(MAX_ACTIVE_BATCHES);
+    ActiveQueryBatches.AddZeroed(MAX_ACTIVE_BATCHES);
+
+    // Don't Init() until the RHI has created the device
+}
+
+FD3D12QueryHeap::~FD3D12QueryHeap()
+{
+    // Unmap the result buffer
+    if (pResultData)
+    {
+        ResultBuffer->GetResource()->Unmap(0, nullptr);
+        pResultData = nullptr;
+    }
+}
+
+void FD3D12QueryHeap::Init()
+{
+    check(GetParentDevice());
+    check(GetParentDevice()->GetDevice());
+
+    // Create the query heap
+    CreateQueryHeap();
+
+    // Create the result buffer
+    CreateResultBuffer();
+}
+
+uint32 FD3D12QueryHeap::ReserveElement()
+{
+    // Increment the reserved count
+    ReservedElementCount++;
+
+    // Make sure we have enough room in the query heap
+    if (ReservedElementCount > GetQueryHeapCount())
+    {
+        // We need more queries than what currently fit in the query heap. Need to grow/rollover.
+        // MSFT: TODO: Do some checks to make sure active heap elements rollover?
+        check(0);
+    }
+
+    return ReservedElementCount;
+}
+
+uint32 FD3D12QueryHeap::GetNextElement(uint32 InElement)
+{
+    // Increment the provided element
+    InElement++;
+
+    // See if we need to wrap around to the begining of the heap
+    if (InElement >= GetQueryHeapCount())
+    {
+        InElement = 0;
+    }
+
+    return InElement;
+}
+
+uint32 FD3D12QueryHeap::GetPreviousElement(uint32 InElement)
+{
+    // Decrement the provided element
+    InElement--;
+
+    // See if we need to wrap around to the end of the heap
+    if (InElement == UINT_MAX)
+    {
+        InElement = GetQueryHeapCount() - 1;
+    }
+
+    return InElement;
+}
+
+uint32 FD3D12QueryHeap::GetNextBatchElement(uint32 InBatchElement)
+{
+    // Increment the provided element
+    InBatchElement++;
+
+    // See if we need to wrap around to the begining of the heap
+    if (InBatchElement >= MAX_ACTIVE_BATCHES)
+    {
+        InBatchElement = 0;
+    }
+
+    return InBatchElement;
+}
+
+uint32 FD3D12QueryHeap::GetPreviousBatchElement(uint32 InBatchElement)
+{
+    // Decrement the provided element
+    InBatchElement--;
+
+    // See if we need to wrap around to the end of the heap
+    if (InBatchElement == UINT_MAX)
+    {
+        InBatchElement = MAX_ACTIVE_BATCHES - 1;
+    }
+
+    return InBatchElement;
+}
+
+bool FD3D12QueryHeap::IsHeapFull()
+{
+    // Find the next element after the active tail and compare with the head
+    return (GetNextElement(TailActiveElement) == HeadActiveElement);
+}
+
+uint32 FD3D12QueryHeap::AllocQuery(FD3D12CommandContext& CmdContext, D3D12_QUERY_TYPE InQueryType)
+{
+    check(CurrentQueryBatch.bOpen);
+
+    // Get the element for this allocation
+    const uint32 CurrentElement = GetNextElement(LastAllocatedElement);
+
+    if (CurrentQueryBatch.StartElement > CurrentElement)
+    {
+        // We're in the middle of a batch, but we're at the end of the heap
+        // We need to split the batch in two and resolve the first piece
+        EndQueryBatchAndResolveQueryData(CmdContext, InQueryType);
+
+        StartQueryBatch(CmdContext);
+    }
+
+    // Increment the count for the current batch
+    CurrentQueryBatch.ElementCount++;
+
+    LastAllocatedElement = CurrentElement;
+    check(CurrentElement < GetQueryHeapCount());
+    return CurrentElement;
+}
+
+void FD3D12QueryHeap::StartQueryBatch(FD3D12CommandContext& CmdContext)
+{
+    check(&CmdContext == &GetParentDevice()->GetDefaultCommandContext());
+    check(!CurrentQueryBatch.bOpen);
+
+	OwningContext = &CmdContext;
+	SyncPoint = CmdContext.CommandListHandle;
+
+    // Clear the current batch
+    CurrentQueryBatch.Clear();
+
+    // Start a new batch
+    CurrentQueryBatch.StartElement = GetNextElement(LastAllocatedElement);
+    CurrentQueryBatch.bOpen = true;
+}
+
+void FD3D12QueryHeap::EndQueryBatchAndResolveQueryData(FD3D12CommandContext& CmdContext, D3D12_QUERY_TYPE InQueryType)
+{
+    check(CurrentQueryBatch.bOpen);
+	check(&CmdContext == OwningContext);
+	//check(CommandListHandle == CmdContext.CommandListHandle);
+
+    // Close the current batch
+    CurrentQueryBatch.bOpen = false;
+
+    // Discard empty batches
+    if (CurrentQueryBatch.ElementCount == 0)
+    {
+        return;
+    }
+
+    // Update the end element
+    CurrentQueryBatch.EndElement = CurrentQueryBatch.StartElement + CurrentQueryBatch.ElementCount - 1;
+
+    // Update the tail
+    TailActiveElement = CurrentQueryBatch.EndElement;
+    check(TailActiveElement < GetQueryHeapCount());
+
+    // Increment the active element count
+    ActiveAllocatedElementCount += CurrentQueryBatch.ElementCount;
+    check(ActiveAllocatedElementCount <= GetQueryHeapCount());
+
+    // Track the current active batches (application is using the data)
+    LastBatch = GetNextBatchElement(LastBatch);
+    ActiveQueryBatches[LastBatch] = CurrentQueryBatch;
+
+    // Update the head
+    QueryBatch& OldestBatch = ActiveQueryBatches[GetNextBatchElement(LastBatch)];
+    HeadActiveElement = OldestBatch.StartElement;
+    ActiveAllocatedElementCount -= OldestBatch.ElementCount;
+
+	CmdContext.otherWorkCounter++;
+    CmdContext.CommandListHandle->ResolveQueryData(
+		QueryHeap, InQueryType, CurrentQueryBatch.StartElement, CurrentQueryBatch.ElementCount,
+        ResultBuffer->GetResource(), GetResultBufferOffsetForElement(CurrentQueryBatch.StartElement));
+}
+
+uint32 FD3D12QueryHeap::BeginQuery(FD3D12CommandContext& CmdContext, D3D12_QUERY_TYPE InQueryType)
+{
+    const uint32 Element = AllocQuery(CmdContext, InQueryType);
+	CmdContext.otherWorkCounter++;
+	CmdContext.CommandListHandle->BeginQuery(QueryHeap, InQueryType, Element);
+    return Element;
+}
+
+void FD3D12QueryHeap::EndQuery(FD3D12CommandContext& CmdContext, D3D12_QUERY_TYPE InQueryType, uint32 InElement)
+{
+	CmdContext.otherWorkCounter++;
+	CmdContext.CommandListHandle->EndQuery(QueryHeap, InQueryType, InElement);
+}
+
+void FD3D12QueryHeap::CreateQueryHeap()
+{
+    // Create the upload heap
+    VERIFYD3D11RESULT(GetParentDevice()->GetDevice()->CreateQueryHeap(&QueryHeapDesc, IID_PPV_ARGS(QueryHeap.GetInitReference())));
+}
+
+void FD3D12QueryHeap::CreateResultBuffer()
+{
+    D3D12_HEAP_PROPERTIES ResultBufferHeapProperties = CD3DX12_HEAP_PROPERTIES(D3D12_HEAP_TYPE_READBACK);
+    D3D12_RESOURCE_DESC ResultBufferDesc = CD3DX12_RESOURCE_DESC::Buffer(ResultSize * QueryHeapDesc.Count); // Each query's result occupies ResultSize bytes.
+
+    // Create the readback heap
+    VERIFYD3D11RESULT(GetParentDevice()->GetResourceHelper().CreateCommittedResource(
+        ResultBufferDesc, 
+        ResultBufferHeapProperties, 
+        D3D12_RESOURCE_STATE_COPY_DEST, 
+        ResultBuffer.GetInitReference()));
+
+    // Map the result buffer (and keep it mapped)
+    VERIFYD3D11RESULT(ResultBuffer->GetResource()->Map(0, nullptr, &pResultData));
+}
+
+
+/*=============================================================================
+ * class FD3D12BufferedGPUTiming
+ *=============================================================================*/
+
+/**
+ * Constructor.
+ *
+ * @param InD3DRHI			RHI interface
+ * @param InBufferSize		Number of buffered measurements
+ */
+FD3D12BufferedGPUTiming::FD3D12BufferedGPUTiming( FD3D12Device* InParent, int32 InBufferSize )
+: 	BufferSize( InBufferSize )
+,	CurrentTimestamp( -1 )
+,	NumIssuedTimestamps( 0 )
+,   StartTimestampQueryHeap(nullptr)
+,   StartTimestampListHandles(nullptr)
+,   StartTimestampQueryHeapBuffer(nullptr)
+,   StartTimestampQueryHeapBufferData(nullptr)
+,   EndTimestampQueryHeap(nullptr)
+,   EndTimestampListHandles(nullptr)
+,   EndTimestampQueryHeapBuffer(nullptr)
+,   EndTimestampQueryHeapBufferData(nullptr)
+,	bIsTiming( false )
+,   bStablePowerState( false)
+,   FD3D12DeviceChild(InParent)
+{
+}
+
+/**
+ * Initializes the static variables, if necessary.
+ */
+void FD3D12BufferedGPUTiming::PlatformStaticInitialize(void* UserData)
+{
+	// Are the static variables initialized?
+	check( !GAreGlobalsInitialized );
+
+	GTimingFrequency = 0;
+    FD3D12Device* ParentDevice = (FD3D12Device*) UserData;
+    VERIFYD3D11RESULT(ParentDevice->GetCommandListManager().GetTimestampFrequency(&GTimingFrequency));
+}
+
+/**
+ * Initializes all D3D resources and if necessary, the static variables.
+ */
+void FD3D12BufferedGPUTiming::InitDynamicRHI()
+{
+	StaticInitialize(GetParentDevice(), PlatformStaticInitialize);
+
+	CurrentTimestamp = 0;
+	NumIssuedTimestamps = 0;
+	bIsTiming = false;
+
+	// Now initialize the queries and backing buffers for this timing object.
+
+	if (GIsSupported)
+	{
+		D3D12_QUERY_HEAP_DESC QueryHeapDesc = {};
+		QueryHeapDesc.Type = D3D12_QUERY_HEAP_TYPE_TIMESTAMP;
+		QueryHeapDesc.Count = BufferSize;
+
+		VERIFYD3D11RESULT(GetParentDevice()->GetDevice()->CreateQueryHeap(&QueryHeapDesc, IID_PPV_ARGS(StartTimestampQueryHeap.GetInitReference())));
+		VERIFYD3D11RESULT(GetParentDevice()->GetDevice()->CreateQueryHeap(&QueryHeapDesc, IID_PPV_ARGS(EndTimestampQueryHeap.GetInitReference())));
+
+		D3D12_HEAP_PROPERTIES queryHeapBufferHeapProperties = CD3DX12_HEAP_PROPERTIES(D3D12_HEAP_TYPE_READBACK);
+
+        D3D12_RESOURCE_DESC queryHeapBufferDesc = CD3DX12_RESOURCE_DESC::Buffer(8 * BufferSize); // Each timestamp query occupies 8 bytes.
+		
+        VERIFYD3D11RESULT(GetParentDevice()->GetDevice()->CreateCommittedResource(&queryHeapBufferHeapProperties, D3D12_HEAP_FLAG_NONE, &queryHeapBufferDesc, D3D12_RESOURCE_STATE_COPY_DEST, nullptr, IID_PPV_ARGS(StartTimestampQueryHeapBuffer.GetInitReference())));
+		VERIFYD3D11RESULT(StartTimestampQueryHeapBuffer->Map(0, nullptr, reinterpret_cast<void**>(&StartTimestampQueryHeapBufferData)));
+		
+        VERIFYD3D11RESULT(GetParentDevice()->GetDevice()->CreateCommittedResource(&queryHeapBufferHeapProperties, D3D12_HEAP_FLAG_NONE, &queryHeapBufferDesc, D3D12_RESOURCE_STATE_COPY_DEST, nullptr, IID_PPV_ARGS(EndTimestampQueryHeapBuffer.GetInitReference())));
+        VERIFYD3D11RESULT(EndTimestampQueryHeapBuffer->Map(0, nullptr, reinterpret_cast<void**>(&EndTimestampQueryHeapBufferData)));
+
+		StartTimestampListHandles = new FD3D12CLSyncPoint[BufferSize];
+		ZeroMemory(StartTimestampListHandles, sizeof(StartTimestampListHandles));
+		EndTimestampListHandles = new FD3D12CLSyncPoint[BufferSize];
+		ZeroMemory(EndTimestampListHandles, sizeof(EndTimestampListHandles));
+	}
+}
+
+/**
+ * Releases all D3D resources.
+ */
+void FD3D12BufferedGPUTiming::ReleaseDynamicRHI()
+{
+	StartTimestampQueryHeap = nullptr;
+    if (StartTimestampListHandles)
+    {
+        delete[] StartTimestampListHandles;
+        StartTimestampListHandles = nullptr;
+    }
+	StartTimestampQueryHeapBuffer = nullptr;
+
+	EndTimestampQueryHeap = nullptr;
+    if (EndTimestampListHandles)
+    {
+        delete[] EndTimestampListHandles;
+        EndTimestampListHandles = nullptr;
+    }
+	EndTimestampQueryHeapBuffer = nullptr;
+}
+
+/**
+ * Start a GPU timing measurement.
+ */
+void FD3D12BufferedGPUTiming::StartTiming()
+{
+	// Issue a timestamp query for the 'start' time.
+	if (GIsSupported && !bIsTiming)
+	{
+        // Check to see if stable power state cvar has changed
+        bool bStablePowerStateCVar = RHIConsoleVariables::bStablePowerState != 0;
+        if (bStablePowerState != bStablePowerStateCVar)
+        {
+            if (SUCCEEDED(GetParentDevice()->GetDevice()->SetStablePowerState(bStablePowerStateCVar)))
+            {
+                // SetStablePowerState succeeded. Update timing frequency.
+                VERIFYD3D11RESULT(GetParentDevice()->GetCommandListManager().GetTimestampFrequency(&GTimingFrequency));
+                bStablePowerState = bStablePowerStateCVar;
+            }
+            else
+            {
+                // SetStablePowerState failed. This can occur if SDKLayers is not present on the system.
+                RHIConsoleVariables::CVarStablePowerState->Set(0, ECVF_SetByConsole);
+            }
+        }
+		int32 NewTimestampIndex = (CurrentTimestamp + 1) % BufferSize;
+        FD3D12CommandContext& CmdContext = GetParentDevice()->GetDefaultCommandContext();
+		CmdContext.otherWorkCounter += 2;
+		CmdContext.CommandListHandle->EndQuery(StartTimestampQueryHeap, D3D12_QUERY_TYPE_TIMESTAMP, NewTimestampIndex);
+		CmdContext.CommandListHandle->ResolveQueryData(StartTimestampQueryHeap, D3D12_QUERY_TYPE_TIMESTAMP, CurrentTimestamp, 1, StartTimestampQueryHeapBuffer, 8 * CurrentTimestamp);
+		StartTimestampListHandles[NewTimestampIndex] = CmdContext.CommandListHandle;
+		CurrentTimestamp = NewTimestampIndex;
+		bIsTiming = true;
+	}
+}
+
+/**
+ * End a GPU timing measurement.
+ * The timing for this particular measurement will be resolved at a later time by the GPU.
+ */
+void FD3D12BufferedGPUTiming::EndTiming()
+{
+	// Issue a timestamp query for the 'end' time.
+	if (GIsSupported && bIsTiming)
+	{
+		check(CurrentTimestamp >= 0 && CurrentTimestamp < BufferSize);
+        FD3D12CommandContext& CmdContext = GetParentDevice()->GetDefaultCommandContext();
+		CmdContext.otherWorkCounter += 2;
+		CmdContext.CommandListHandle->EndQuery(EndTimestampQueryHeap, D3D12_QUERY_TYPE_TIMESTAMP, CurrentTimestamp);
+		CmdContext.CommandListHandle->ResolveQueryData(EndTimestampQueryHeap, D3D12_QUERY_TYPE_TIMESTAMP, CurrentTimestamp, 1, EndTimestampQueryHeapBuffer, 8 * CurrentTimestamp);
+		EndTimestampListHandles[CurrentTimestamp] = CmdContext.CommandListHandle;
+		NumIssuedTimestamps = FMath::Min<int32>(NumIssuedTimestamps + 1, BufferSize);
+		bIsTiming = false;
+	}
+}
+
+/**
+ * Retrieves the most recently resolved timing measurement.
+ * The unit is the same as for FPlatformTime::Cycles(). Returns 0 if there are no resolved measurements.
+ *
+ * @return	Value of the most recently resolved timing, or 0 if no measurements have been resolved by the GPU yet.
+ */
+uint64 FD3D12BufferedGPUTiming::GetTiming(bool bGetCurrentResultsAndBlock)
+{
+	if ( GIsSupported )
+	{
+		check( CurrentTimestamp >= 0 && CurrentTimestamp < BufferSize );
+		uint64 StartTime, EndTime;
+
+		FD3D12CommandListManager& CommandListManager = GetParentDevice()->GetCommandListManager();
+
+		int32 TimestampIndex = CurrentTimestamp;
+		if (!bGetCurrentResultsAndBlock)
+		{
+			// Quickly check the most recent measurements to see if any of them has been resolved.  Do not flush these queries.
+			for ( int32 IssueIndex = 1; IssueIndex < NumIssuedTimestamps; ++IssueIndex )
+			{
+				if (CommandListManager.IsFinished(EndTimestampListHandles[TimestampIndex]))
+				{
+					check(CommandListManager.IsFinished(StartTimestampListHandles[TimestampIndex]));
+
+					StartTime = StartTimestampQueryHeapBufferData[TimestampIndex];
+					EndTime = EndTimestampQueryHeapBufferData[TimestampIndex];
+
+					if (EndTime > StartTime)
+						return EndTime - StartTime;
+				}
+
+				TimestampIndex = (TimestampIndex + BufferSize - 1) % BufferSize;
+			}
+		}
+		
+		if ( NumIssuedTimestamps > 0 || bGetCurrentResultsAndBlock )
+		{
+			// None of the (NumIssuedTimestamps - 1) measurements were ready yet,
+			// so check the oldest measurement more thoroughly.
+			// This really only happens if occlusion and frame sync event queries are disabled, otherwise those will block until the GPU catches up to 1 frame behind
+
+			const bool bBlocking = ( NumIssuedTimestamps == BufferSize ) || bGetCurrentResultsAndBlock;
+			uint32 IdleStart = FPlatformTime::Cycles();
+			double StartTimeoutTime = FPlatformTime::Seconds();
+
+			SCOPE_CYCLE_COUNTER( STAT_RenderQueryResultTime );
+
+			if (bBlocking)
+				GetParentDevice()->GetDefaultCommandContext().FlushCommands(true);
+
+			GRenderThreadIdle[ERenderThreadIdleTypes::WaitingForGPUQuery] += FPlatformTime::Cycles() - IdleStart;
+			GRenderThreadNumIdle[ERenderThreadIdleTypes::WaitingForGPUQuery]++;
+
+			StartTime = StartTimestampQueryHeapBufferData[TimestampIndex];
+			EndTime = EndTimestampQueryHeapBufferData[TimestampIndex];
+
+			if (EndTime > StartTime)
+				return EndTime - StartTime;
+		}
+	}
+
+	return 0;
+}
\ No newline at end of file
