UE4 - D3D12RHI (not working)

diff --git a/Engine/Source/Runtime/Windows/D3D12RHI/Private/D3D12RHI.cpp b/Engine/Source/Runtime/Windows/D3D12RHI/Private/D3D12RHI.cpp
new file mode 100644
--- /dev/null
+++ b/Engine/Source/Runtime/Windows/D3D12RHI/Private/D3D12RHI.cpp
@@ -0,0 +1,1039 @@
+// Copyright 1998-2014 Epic Games, Inc. All Rights Reserved.
+
+/*=============================================================================
+	D3D12RHI.cpp: Unreal D3D RHI library implementation.
+	=============================================================================*/
+
+#include "D3D12RHIPrivate.h"
+#include "RHIStaticStates.h"
+
+#ifndef WITH_DX_PERF
+#define WITH_DX_PERF	1
+#endif
+
+#if WITH_DX_PERF
+// For perf events
+#include "AllowWindowsPlatformTypes.h"
+#include "d3d9.h"
+#include "HideWindowsPlatformTypes.h"
+#endif	//WITH_DX_PERF
+#include "OneColorShader.h"
+
+#if !UE_BUILD_SHIPPING
+#include "STaskGraph.h"
+#endif
+
+DEFINE_LOG_CATEGORY(LogD3D12RHI);
+
+namespace D3D12RHI
+{
+	extern void UniformBufferBeginFrame();
+}
+using namespace D3D12RHI;
+
+void FD3D12CommandContext::OpenCommandList(bool bRestoreState)
+{
+	FD3D12CommandListManager& CommandListManager = GetParentDevice()->GetCommandListManager();
+
+	// Get a new command list
+	CommandListHandle = CommandListManager.BeginCommandList(StateCache.GetPipelineStateObject());
+
+	CommandListHandle->SetDescriptorHeaps(DescriptorHeaps.Num(), DescriptorHeaps.GetData());
+	CommandListHandle->SetGraphicsRootSignature(GetParentDevice()->GetGraphicsRootSignature());
+	CommandListHandle->SetComputeRootSignature(GetParentDevice()->GetComputeRootSignature());
+
+	StateCache.GetDescriptorCache()->NotifyCurrentCommandList(CommandListHandle);
+
+	if (bRestoreState)
+	{
+		// The next time ApplyState is called, it will set all state on this new command list
+		StateCache.RestoreState();
+	}
+
+	numDraws = 0;
+	numDispatches = 0;
+	numClears = 0;
+	numBarriers = 0;
+	numCopies = 0;
+	otherWorkCounter = 0;
+	CommandListHandle.SetCurrentOwningContext(this);
+}
+
+FD3D12CommandListHandle FD3D12CommandContext::FlushCommands(bool WaitForCompletion)
+{
+	// We should only be flushing the default context
+	check(IsDefaultContext());
+
+	// Only submit a command list if it does meaningful work or the flush is expected to wait for completion.
+	if (WaitForCompletion || HasDoneWork())
+	{
+#if SUPPORTS_MEMORY_RESIDENCY
+		GetParentDevice()->GetOwningRHI()->GetResourceResidencyManager().MakeResident();
+#endif
+		// Close the current command list
+		CommandListHandle.Close();
+
+		CommandListHandle.Execute(WaitForCompletion);
+
+		// Get a new command list to replace the one we submitted for execution. 
+		// Restore the state from the previous command list.
+		OpenCommandList(true);
+	}
+
+	return CommandListHandle;
+}
+
+void FD3D12CommandContext::Finish(TArray<FD3D12CommandListHandle>& CommandLists)
+{
+	if (HasDoneWork())
+	{
+		// Close the current command list
+		CommandListHandle.Close();
+
+		CommandLists.Add(CommandListHandle);
+
+		// Get a new command list to replace the one we submitted for execution
+		OpenCommandList();
+	}
+}
+
+void FD3D12CommandContext::RHIBeginFrame()
+{
+	check(IsDefaultContext());
+	RHIPrivateBeginFrame();
+	UniformBufferBeginFrame();
+	OwningRHI.GPUProfilingData.BeginFrame(&OwningRHI);
+}
+
+void FD3D12CommandContext::ClearState()
+{
+	StateCache.ClearState();
+
+	for (int32 Frequency = 0; Frequency < SF_NumFrequencies; Frequency++)
+	{
+		DirtyUniformBuffers[Frequency] = 0;
+		for (int32 Index = 0; Index < MAX_UNIFORM_BUFFERS_PER_SHADER_STAGE; ++Index)
+		{
+			BoundUniformBuffers[Frequency][Index] = nullptr;
+		}
+	}
+
+	for (int32 Index = 0; Index < _countof(CurrentRenderTargets); ++Index)
+	{
+		CurrentRenderTargets[Index] = nullptr;
+	}
+	for (int32 Index = 0; Index < _countof(CurrentUAVs); ++Index)
+	{
+		CurrentUAVs[Index] = nullptr;
+	}
+	
+	CurrentDepthStencilTarget = nullptr;
+	CurrentDepthTexture = nullptr;
+
+	NumSimultaneousRenderTargets = 0;
+	NumUAVs = 0;
+
+	CurrentDSVAccessType = FExclusiveDepthStencil::DepthWrite_StencilWrite;
+
+	bDiscardSharedConstants = false;
+
+	bUsingTessellation = false;
+
+	// MSFT: Need to do anything with the constant buffers?
+	/*for (int32 BufferIndex = 0; BufferIndex < VSConstantBuffers.Num(); BufferIndex++)
+	{
+		VSConstantBuffers[BufferIndex]->ClearConstantBuffer();
+	}
+	for (int32 BufferIndex = 0; BufferIndex < PSConstantBuffers.Num(); BufferIndex++)
+	{
+		PSConstantBuffers[BufferIndex]->ClearConstantBuffer();
+	}
+	for (int32 BufferIndex = 0; BufferIndex < HSConstantBuffers.Num(); BufferIndex++)
+	{
+		HSConstantBuffers[BufferIndex]->ClearConstantBuffer();
+	}
+	for (int32 BufferIndex = 0; BufferIndex < DSConstantBuffers.Num(); BufferIndex++)
+	{
+		DSConstantBuffers[BufferIndex]->ClearConstantBuffer();
+	}
+	for (int32 BufferIndex = 0; BufferIndex < GSConstantBuffers.Num(); BufferIndex++)
+	{
+		GSConstantBuffers[BufferIndex]->ClearConstantBuffer();
+	}
+	for (int32 BufferIndex = 0; BufferIndex < CSConstantBuffers.Num(); BufferIndex++)
+	{
+		CSConstantBuffers[BufferIndex]->ClearConstantBuffer();
+	}*/
+
+	CurrentComputeShader = nullptr;
+	CurrentBoundShaderState = nullptr;
+
+	UploadHeapAllocator.CleanupFreeBlocks();
+}
+
+void FD3D12CommandContext::CheckIfSRVIsResolved(FD3D12ShaderResourceView* SRV)
+{
+	// MSFT: Seb: TODO: Make this work on 12
+#if CHECK_SRV_TRANSITIONS
+	if (GRHIThread || !SRV)
+	{
+		return;
+	}
+
+	static const auto CheckSRVCVar = IConsoleManager::Get().FindTConsoleVariableDataInt(TEXT("r.CheckSRVTransitions"));
+	if (!CheckSRVCVar->GetValueOnRenderThread())
+	{
+		return;
+	}
+
+	ID3D11Resource* SRVResource = nullptr;
+	SRV->GetResource(&SRVResource);
+
+	int32 MipLevel;
+	int32 NumMips;
+	int32 ArraySlice;
+	int32 NumSlices;
+	GetMipAndSliceInfoFromSRV(SRV, MipLevel, NumMips, ArraySlice, NumSlices);
+
+	//d3d uses -1 to mean 'all mips'
+	int32 LastMip = MipLevel + NumMips - 1;
+	int32 LastSlice = ArraySlice + NumSlices - 1;
+
+	TArray<FUnresolvedRTInfo> RTInfoArray;
+	UnresolvedTargets.MultiFind(SRVResource, RTInfoArray);
+
+	for (int32 InfoIndex = 0; InfoIndex < RTInfoArray.Num(); ++InfoIndex)
+	{
+		const FUnresolvedRTInfo& RTInfo = RTInfoArray[InfoIndex];
+		int32 RTLastMip = RTInfo.MipLevel + RTInfo.NumMips - 1;
+		ensureMsgf((MipLevel == -1 || NumMips == -1) || (LastMip < RTInfo.MipLevel || MipLevel > RTLastMip), TEXT("SRV is set to read mips in range %i to %i.  Target %s is unresolved for mip %i"), MipLevel, LastMip, *RTInfo.ResourceName.ToString(), RTInfo.MipLevel);
+		ensureMsgf(NumMips != -1, TEXT("SRV is set to read all mips.  Target %s is unresolved for mip %i"), *RTInfo.ResourceName.ToString(), RTInfo.MipLevel);
+
+		int32 RTLastSlice = RTInfo.ArraySlice + RTInfo.ArraySize - 1;
+		ensureMsgf((ArraySlice == -1 || LastSlice == -1) || (LastSlice < RTInfo.ArraySlice || ArraySlice > RTLastSlice), TEXT("SRV is set to read slices in range %i to %i.  Target %s is unresolved for mip %i"), ArraySlice, LastSlice, *RTInfo.ResourceName.ToString(), RTInfo.ArraySlice);
+		ensureMsgf(ArraySlice == -1 || NumSlices != -1, TEXT("SRV is set to read all slices.  Target %s is unresolved for slice %i"));
+	}
+	SRVResource->Release();
+#endif
+}
+
+template <EShaderFrequency ShaderFrequency>
+void FD3D12CommandContext::InternalSetShaderResourceView(FD3D12ResourceLocation* Resource, FD3D12ShaderResourceView* SRV, int32 ResourceIndex, FD3D12StateCache::ESRV_Type SrvType)
+{
+	// Check either both are set, or both are null.
+	check((Resource && SRV) || (!Resource && !SRV));
+	CheckIfSRVIsResolved(SRV);
+
+	// Set the SRV we have been given (or null).
+	StateCache.SetShaderResourceView<ShaderFrequency>(SRV, ResourceIndex, SrvType);
+}
+
+template <EShaderFrequency ShaderFrequency>
+void FD3D12CommandContext::ClearShaderResourceViews(FD3D12ResourceLocation* Resource)
+{
+	StateCache.ClearShaderResourceViews<ShaderFrequency>(Resource);
+}
+
+void FD3D12CommandContext::ConditionalClearShaderResource(FD3D12ResourceLocation* Resource)
+{
+	SCOPE_CYCLE_COUNTER(STAT_D3D12ClearShaderResourceTime);
+	check(Resource);
+	ClearShaderResourceViews<SF_Vertex>(Resource);
+	ClearShaderResourceViews<SF_Hull>(Resource);
+	ClearShaderResourceViews<SF_Domain>(Resource);
+	ClearShaderResourceViews<SF_Pixel>(Resource);
+	ClearShaderResourceViews<SF_Geometry>(Resource);
+	ClearShaderResourceViews<SF_Compute>(Resource);
+}
+
+void FD3D12CommandContext::ClearAllShaderResources()
+{
+	StateCache.ClearSRVs();
+}
+
+void FD3D12DynamicRHI::IssueLongGPUTask()
+{
+	if (GMaxRHIFeatureLevel >= ERHIFeatureLevel::SM4)
+	{
+		int32 LargestViewportIndex = INDEX_NONE;
+		int32 LargestViewportPixels = 0;
+
+		for (int32 ViewportIndex = 0; ViewportIndex < GetRHIDevice()->GetViewports().Num(); ViewportIndex++)
+		{
+            FD3D12Viewport* Viewport = GetRHIDevice()->GetViewports()[ViewportIndex];
+
+			if (Viewport->GetSizeXY().X * Viewport->GetSizeXY().Y > LargestViewportPixels)
+			{
+				LargestViewportPixels = Viewport->GetSizeXY().X * Viewport->GetSizeXY().Y;
+				LargestViewportIndex = ViewportIndex;
+			}
+		}
+
+		if (LargestViewportIndex >= 0)
+		{
+            FD3D12Viewport* Viewport = GetRHIDevice()->GetViewports()[LargestViewportIndex];
+
+			FRHICommandList_RecursiveHazardous RHICmdList(RHIGetDefaultContext());
+
+			SetRenderTarget(RHICmdList, Viewport->GetBackBuffer(), FTextureRHIRef());
+			RHICmdList.SetBlendState(TStaticBlendState<CW_RGBA, BO_Add, BF_One, BF_One>::GetRHI(), FLinearColor::Black);
+			RHICmdList.SetDepthStencilState(TStaticDepthStencilState<false, CF_Always>::GetRHI(), 0);
+			RHICmdList.SetRasterizerState(TStaticRasterizerState<FM_Solid, CM_None>::GetRHI());
+
+			auto ShaderMap = GetGlobalShaderMap(GMaxRHIFeatureLevel);
+			TShaderMapRef<TOneColorVS<true> > VertexShader(ShaderMap);
+			TShaderMapRef<FLongGPUTaskPS> PixelShader(ShaderMap);
+
+			RHICmdList.SetLocalBoundShaderState(RHICmdList.BuildLocalBoundShaderState(GD3D11Vector4VertexDeclaration.VertexDeclarationRHI, VertexShader->GetVertexShader(), FHullShaderRHIRef(), FDomainShaderRHIRef(), PixelShader->GetPixelShader(), FGeometryShaderRHIRef()));
+
+			// Draw a fullscreen quad
+			FVector4 Vertices[4];
+			Vertices[0].Set(-1.0f, 1.0f, 0, 1.0f);
+			Vertices[1].Set(1.0f, 1.0f, 0, 1.0f);
+			Vertices[2].Set(-1.0f, -1.0f, 0, 1.0f);
+			Vertices[3].Set(1.0f, -1.0f, 0, 1.0f);
+			DrawPrimitiveUP(RHICmdList, PT_TriangleStrip, 2, Vertices, sizeof(Vertices[0]));
+
+			// Implicit flush. Always call flush when using a command list in RHI implementations before doing anything else. This is super hazardous.
+		}
+	}
+}
+
+namespace D3D12RHI
+{
+	void FD3DGPUProfiler::BeginFrame(FD3D12DynamicRHI* InRHI)
+	{
+		CurrentEventNode = NULL;
+		check(!bTrackingEvents);
+		check(!CurrentEventNodeFrame); // this should have already been cleaned up and the end of the previous frame
+
+		// latch the bools from the game thread into our private copy
+		bLatchedGProfilingGPU = GTriggerGPUProfile;
+		bLatchedGProfilingGPUHitches = GTriggerGPUHitchProfile;
+		if (bLatchedGProfilingGPUHitches)
+		{
+			bLatchedGProfilingGPU = false; // we do NOT permit an ordinary GPU profile during hitch profiles
+		}
+
+		if (bLatchedGProfilingGPU)
+		{
+			// Issue a bunch of GPU work at the beginning of the frame, to make sure that we are GPU bound
+			// We can't isolate idle time from GPU timestamps
+			InRHI->IssueLongGPUTask();
+		}
+
+		// if we are starting a hitch profile or this frame is a gpu profile, then save off the state of the draw events
+		if (bLatchedGProfilingGPU || (!bPreviousLatchedGProfilingGPUHitches && bLatchedGProfilingGPUHitches))
+		{
+			bOriginalGEmitDrawEvents = GEmitDrawEvents;
+		}
+
+		if (bLatchedGProfilingGPU || bLatchedGProfilingGPUHitches)
+		{
+			if (bLatchedGProfilingGPUHitches && GPUHitchDebounce)
+			{
+				// if we are doing hitches and we had a recent hitch, wait to recover
+				// the reasoning is that collecting the hitch report may itself hitch the GPU
+				GPUHitchDebounce--;
+			}
+			else
+			{
+				GEmitDrawEvents = true;  // thwart an attempt to turn this off on the game side
+				bTrackingEvents = true;
+				CurrentEventNodeFrame = new FD3D12EventNodeFrame(GetParentDevice());
+				CurrentEventNodeFrame->StartFrame();
+			}
+		}
+		else if (bPreviousLatchedGProfilingGPUHitches)
+		{
+			// hitch profiler is turning off, clear history and restore draw events
+			GPUHitchEventNodeFrames.Empty();
+			GEmitDrawEvents = bOriginalGEmitDrawEvents;
+		}
+		bPreviousLatchedGProfilingGPUHitches = bLatchedGProfilingGPUHitches;
+
+		// Skip timing events when using SLI, they will not be accurate anyway
+		if (GNumActiveGPUsForRendering == 1)
+		{
+			FrameTiming.StartTiming();
+		}
+
+		if (GEmitDrawEvents)
+		{
+			PushEvent(TEXT("FRAME"));
+		}
+	}
+}
+
+void FD3D12CommandContext::RHIEndFrame()
+{
+	check(IsDefaultContext());
+	OwningRHI.GPUProfilingData.EndFrame();
+
+	FD3D12Device* Device = GetParentDevice();
+
+	// MSFT: Seb: Is any of this needed? Perhaps there is now a better place to clean up this stuff
+	//TODO: do for all devices?
+	//TODO: move this into the device
+	Device->GetDeferredDeletionQueue().ReleaseResources();
+	Device->GetDefaultBufferAllocator().CleanupFreeBlocks();
+
+	for (uint32 i = 0; i < Device->GetNumContexts(); ++i)
+    {
+		Device->GetCommandContext(i).StateCache.GetDescriptorCache()->EndFrame();
+		Device->GetCommandContext(i).UploadHeapAllocator.CleanupFreeBlocks();
+	}
+
+	for (uint32 i = 0; i < FD3D12DynamicRHI::GetD3DRHI()->NumThreadDynamicHeapAllocators; ++i)
+    {
+		FD3D12DynamicHeapAllocator* pCurrentHelperThreadDynamicHeapAllocator = FD3D12DynamicRHI::GetD3DRHI()->ThreadDynamicHeapAllocatorArray[i];
+		check(pCurrentHelperThreadDynamicHeapAllocator != nullptr);
+		pCurrentHelperThreadDynamicHeapAllocator->CleanupFreeBlocks();
+    }
+
+#if SUPPORTS_MEMORY_RESIDENCY
+	GetResourceResidencyManager().Process();
+#endif
+	DXGI_QUERY_VIDEO_MEMORY_INFO LocalVideoMemoryInfo;
+
+	FD3D12DynamicRHI::GetD3DRHI()->GetLocalVideoMemoryInfo(&LocalVideoMemoryInfo);
+
+	int64 budget = LocalVideoMemoryInfo.Budget;
+
+	int64 AvailableSpace = budget - int64(LocalVideoMemoryInfo.CurrentUsage);
+
+	SET_MEMORY_STAT(STAT_D3D12UsedVideoMemory, LocalVideoMemoryInfo.CurrentUsage);
+	SET_MEMORY_STAT(STAT_D3D12AvailableVideoMemory, AvailableSpace);
+	SET_MEMORY_STAT(STAT_D3D12TotalVideoMemory, budget);
+}
+
+void FD3DGPUProfiler::EndFrame()
+{
+	if (GEmitDrawEvents)
+	{
+		PopEvent();
+	}
+
+	// Skip timing events when using SLI, they will not be accurate anyway
+	if (GNumActiveGPUsForRendering == 1)
+	{
+		FrameTiming.EndTiming();
+	}
+
+	// Skip timing events when using SLI, as they will block the GPU and we want maximum throughput
+	// Stat unit GPU time is not accurate anyway with SLI
+	if (FrameTiming.IsSupported() && GNumActiveGPUsForRendering == 1)
+	{
+		uint64 GPUTiming = FrameTiming.GetTiming();
+		uint64 GPUFreq = FrameTiming.GetTimingFrequency();
+		GGPUFrameTime = FMath::TruncToInt(double(GPUTiming) / double(GPUFreq) / FPlatformTime::GetSecondsPerCycle());
+	}
+	else
+	{
+		GGPUFrameTime = 0;
+	}
+
+	// if we have a frame open, close it now.
+	if (CurrentEventNodeFrame)
+	{
+		CurrentEventNodeFrame->EndFrame();
+	}
+
+	check(!bTrackingEvents || bLatchedGProfilingGPU || bLatchedGProfilingGPUHitches);
+	check(!bTrackingEvents || CurrentEventNodeFrame);
+	if (bLatchedGProfilingGPU)
+	{
+		if (bTrackingEvents)
+		{
+			GEmitDrawEvents = bOriginalGEmitDrawEvents;
+			UE_LOG(LogD3D12RHI, Warning, TEXT(""));
+			UE_LOG(LogD3D12RHI, Warning, TEXT(""));
+			CurrentEventNodeFrame->DumpEventTree();
+			GTriggerGPUProfile = false;
+			bLatchedGProfilingGPU = false;
+
+			if (RHIConfig::ShouldSaveScreenshotAfterProfilingGPU()
+				&& GEngine->GameViewport)
+			{
+				GEngine->GameViewport->Exec(NULL, TEXT("SCREENSHOT"), *GLog);
+			}
+		}
+	}
+	else if (bLatchedGProfilingGPUHitches)
+	{
+		//@todo this really detects any hitch, even one on the game thread.
+		// it would be nice to restrict the test to stalls on D3D, but for now...
+		// this needs to be out here because bTrackingEvents is false during the hitch debounce
+		static double LastTime = -1.0;
+		double Now = FPlatformTime::Seconds();
+		if (bTrackingEvents)
+		{
+			/** How long, in seconds a frame much be to be considered a hitch **/
+			const float HitchThreshold = RHIConfig::GetGPUHitchThreshold();
+			float ThisTime = Now - LastTime;
+			bool bHitched = (ThisTime > HitchThreshold) && LastTime > 0.0 && CurrentEventNodeFrame;
+			if (bHitched)
+			{
+				UE_LOG(LogD3D12RHI, Warning, TEXT("*******************************************************************************"));
+				UE_LOG(LogD3D12RHI, Warning, TEXT("********** Hitch detected on CPU, frametime = %6.1fms"), ThisTime * 1000.0f);
+				UE_LOG(LogD3D12RHI, Warning, TEXT("*******************************************************************************"));
+
+				for (int32 Frame = 0; Frame < GPUHitchEventNodeFrames.Num(); Frame++)
+				{
+					UE_LOG(LogD3D12RHI, Warning, TEXT(""));
+					UE_LOG(LogD3D12RHI, Warning, TEXT(""));
+					UE_LOG(LogD3D12RHI, Warning, TEXT("********** GPU Frame: Current - %d"), GPUHitchEventNodeFrames.Num() - Frame);
+					GPUHitchEventNodeFrames[Frame].DumpEventTree();
+				}
+				UE_LOG(LogD3D12RHI, Warning, TEXT(""));
+				UE_LOG(LogD3D12RHI, Warning, TEXT(""));
+				UE_LOG(LogD3D12RHI, Warning, TEXT("********** GPU Frame: Current"));
+				CurrentEventNodeFrame->DumpEventTree();
+
+				UE_LOG(LogD3D12RHI, Warning, TEXT("*******************************************************************************"));
+				UE_LOG(LogD3D12RHI, Warning, TEXT("********** End Hitch GPU Profile"));
+				UE_LOG(LogD3D12RHI, Warning, TEXT("*******************************************************************************"));
+				if (GEngine->GameViewport)
+				{
+					GEngine->GameViewport->Exec(NULL, TEXT("SCREENSHOT"), *GLog);
+				}
+
+				GPUHitchDebounce = 5; // don't trigger this again for a while
+				GPUHitchEventNodeFrames.Empty(); // clear history
+			}
+			else if (CurrentEventNodeFrame) // this will be null for discarded frames while recovering from a recent hitch
+			{
+				/** How many old frames to buffer for hitch reports **/
+				static const int32 HitchHistorySize = 4;
+
+				if (GPUHitchEventNodeFrames.Num() >= HitchHistorySize)
+				{
+					GPUHitchEventNodeFrames.RemoveAt(0);
+				}
+				GPUHitchEventNodeFrames.Add((FD3D12EventNodeFrame*)CurrentEventNodeFrame);
+				CurrentEventNodeFrame = NULL;  // prevent deletion of this below; ke kept it in the history
+			}
+		}
+		LastTime = Now;
+	}
+	bTrackingEvents = false;
+	delete CurrentEventNodeFrame;
+	CurrentEventNodeFrame = NULL;
+}
+
+float FD3D12EventNode::GetTiming()
+{
+	float Result = 0;
+
+	if (Timing.IsSupported())
+	{
+		// Get the timing result and block the CPU until it is ready
+		const uint64 GPUTiming = Timing.GetTiming(true);
+		const uint64 GPUFreq = Timing.GetTimingFrequency();
+
+		Result = double(GPUTiming) / double(GPUFreq);
+	}
+
+	return Result;
+}
+
+void FD3D12CommandContext::RHIBeginScene()
+{
+	check(IsDefaultContext());
+	FD3D12DynamicRHI& RHI = OwningRHI;
+	// Increment the frame counter. INDEX_NONE is a special value meaning "uninitialized", so if
+	// we hit it just wrap around to zero.
+	OwningRHI.SceneFrameCounter++;
+	if (OwningRHI.SceneFrameCounter == INDEX_NONE)
+	{
+		OwningRHI.SceneFrameCounter++;
+	}
+
+	static auto* ResourceTableCachingCvar = IConsoleManager::Get().FindTConsoleVariableDataInt(TEXT("rhi.ResourceTableCaching"));
+	if (ResourceTableCachingCvar == NULL || ResourceTableCachingCvar->GetValueOnAnyThread() == 1)
+	{
+		OwningRHI.ResourceTableFrameCounter = OwningRHI.SceneFrameCounter;
+	}
+}
+
+void FD3D12CommandContext::RHIEndScene()
+{
+	OwningRHI.ResourceTableFrameCounter = INDEX_NONE;
+}
+
+void FD3DGPUProfiler::PushEvent(const TCHAR* Name)
+{
+#if WITH_DX_PERF
+	D3DPERF_BeginEvent(FColor(0, 0, 0).DWColor(), Name);
+#endif
+
+	FGPUProfiler::PushEvent(Name);
+}
+
+void FD3DGPUProfiler::PopEvent()
+{
+#if WITH_DX_PERF
+	D3DPERF_EndEvent();
+#endif
+
+	FGPUProfiler::PopEvent();
+}
+
+/** Start this frame of per tracking */
+void FD3D12EventNodeFrame::StartFrame()
+{
+	EventTree.Reset();
+	RootEventTiming.StartTiming();
+}
+
+/** End this frame of per tracking, but do not block yet */
+void FD3D12EventNodeFrame::EndFrame()
+{
+	RootEventTiming.EndTiming();
+}
+
+float FD3D12EventNodeFrame::GetRootTimingResults()
+{
+	double RootResult = 0.0f;
+	if (RootEventTiming.IsSupported())
+	{
+		const uint64 GPUTiming = RootEventTiming.GetTiming(true);
+		const uint64 GPUFreq = RootEventTiming.GetTimingFrequency();
+
+		RootResult = double(GPUTiming) / double(GPUFreq);
+	}
+
+	return (float)RootResult;
+}
+
+void FD3D12EventNodeFrame::LogDisjointQuery()
+{
+}
+
+void UpdateBufferStats(FD3D12ResourceLocation* ResourceLocation, bool bAllocating, uint32 BufferType)
+{
+	uint64 RequestedSize = ResourceLocation->GetEffectiveBufferSize();
+
+	const bool bUniformBuffer = BufferType == D3D12_BUFFER_TYPE_CONSTANT;
+	const bool bIndexBuffer = BufferType == D3D12_BUFFER_TYPE_INDEX;
+	const bool bVertexBuffer = BufferType == D3D12_BUFFER_TYPE_VERTEX;
+
+	if (bAllocating)
+	{
+		if (bUniformBuffer)
+		{
+			INC_MEMORY_STAT_BY(STAT_UniformBufferMemory, RequestedSize);
+		}
+		else if (bIndexBuffer)
+		{
+			INC_MEMORY_STAT_BY(STAT_IndexBufferMemory, RequestedSize);
+		}
+		else if (bVertexBuffer)
+		{
+			INC_MEMORY_STAT_BY(STAT_VertexBufferMemory, RequestedSize);
+		}
+		else
+		{
+			INC_MEMORY_STAT_BY(STAT_StructuredBufferMemory, RequestedSize);
+		}
+	}
+	else
+	{
+		if (bUniformBuffer)
+		{
+			DEC_MEMORY_STAT_BY(STAT_UniformBufferMemory, RequestedSize);
+		}
+		else if (bIndexBuffer)
+		{
+			DEC_MEMORY_STAT_BY(STAT_IndexBufferMemory, RequestedSize);
+		}
+		else if (bVertexBuffer)
+		{
+			DEC_MEMORY_STAT_BY(STAT_VertexBufferMemory, RequestedSize);
+		}
+		else
+		{
+			DEC_MEMORY_STAT_BY(STAT_StructuredBufferMemory, RequestedSize);
+		}
+	}
+}
+
+FD3D12ResourceLocation::~FD3D12ResourceLocation()
+{
+	InternalReleaseResource();
+}
+
+void FD3D12ResourceLocation::InternalReleaseResource()
+{
+	if (LinkedResourceLocation.GetReference() != nullptr)
+	{
+		// This is a FastAlloc resource location object
+		check(BlockInfo.GetReference() == nullptr);
+
+		// Dereference linked resource
+		LinkedResourceLocation = nullptr;
+		return;
+	}
+
+	if (BlockInfo.GetReference() != nullptr)
+	{
+		Resource = nullptr;
+
+		FD3D12ResourceAllocator *Allocator = BlockInfo->Allocator;
+		check(!!Allocator);
+
+		Allocator->ExpireBlock(BlockInfo);
+		BlockInfo = nullptr;
+	}
+	else
+	{
+		FD3D12Resource *pResource = Resource.GetReference();
+		if (pResource)
+		{
+			GetParentDevice()->GetDeferredDeletionQueue().EnqueueResource(pResource);
+
+			Resource = nullptr;
+		}
+	}
+}
+
+void FD3D12ResourceLocation::SetFromD3DResource(FD3D12Resource *InResource, uint64 InOffset, uint32 InEffectiveBufferSize)
+{
+	InternalReleaseResource();
+
+	EffectiveBufferSize = InEffectiveBufferSize;
+	Resource = InResource;
+	Offset = InOffset;
+}
+
+void FD3D12ResourceLocation::Clear()
+{
+	InternalReleaseResource();
+
+	EffectiveBufferSize = 0;
+	Resource = nullptr;
+	Offset = 0;
+}
+
+void FD3D12ResourceLocation::UpdateStateCache(FD3D12StateCache& StateCache)
+{
+
+	// Rename SRV?
+	if (!StateCache.GetForceSetSRVs() && StateCache.IsShaderResource(this))
+	{
+		StateCache.ForceSetSRVs();
+	}
+
+	// Rename vertex buffer?
+	if (!StateCache.GetForceSetVB() && StateCache.IsStreamSource(this))
+	{
+		StateCache.ForceSetVB();
+	}
+
+	// Rename index buffer?
+	if (!StateCache.GetForceSetIB() && StateCache.IsIndexBuffer(this))
+	{
+		// If this resource matches the current index buffer then force the state 
+		// cache to update the index buffer view
+		StateCache.ForceSetIB();
+	}
+}
+
+inline void FD3D12ResourceLocation::UpdateDefaultStateCache()
+{
+    UpdateStateCache(GetParentDevice()->GetDefaultCommandContext().StateCache);
+}
+
+void FD3D12DeferredDeletionQueue::EnqueueResource(FD3D12Resource *pResource)
+{
+	const uint64 CurrentFrameFence = GetParentDevice()->GetCommandListManager().GetFence(EFenceType::FT_Frame).GetCurrentFence();
+
+	pResource->AddRef();
+
+	FencedObjectType FencedObject;
+	FencedObject.Key = pResource;
+	FencedObject.Value = CurrentFrameFence;
+	DeferredReleaseQueue.Enqueue(FencedObject);
+}
+
+LONGLONG D3D12ShaderResourceViewSequenceNumber = 1;
+
+#ifndef PLATFORM_IMPLEMENTS_FASTVRAMALLOCATOR
+#define PLATFORM_IMPLEMENTS_FASTVRAMALLOCATOR		0
+#endif
+
+#if !PLATFORM_IMPLEMENTS_FASTVRAMALLOCATOR
+FFastVRAMAllocator* FFastVRAMAllocator::GetFastVRAMAllocator()
+{
+	static FFastVRAMAllocator FastVRAMAllocatorSingleton;
+	return &FastVRAMAllocatorSingleton;
+}
+#endif
+
+void FD3D12DynamicRHI::GetLocalVideoMemoryInfo (DXGI_QUERY_VIDEO_MEMORY_INFO* LocalVideoMemoryInfo)
+{
+	VERIFYD3D11RESULT(MainDevice->GetAdapter3()->QueryVideoMemoryInfo(0, DXGI_MEMORY_SEGMENT_GROUP_LOCAL, LocalVideoMemoryInfo));
+}	
+
+#if SUPPORTS_MEMORY_RESIDENCY
+int32 bLogMemoryResidency = 1;
+static FAutoConsoleVariableRef CVarLogMemoryResidency(
+	TEXT("r.D3D12LogMemoryResidency"),
+	bLogMemoryResidency,
+	TEXT("Print out a log of the memory residency stats.")
+	TEXT("  0: Off completely\n")
+	TEXT("  1: Off unless under memory pressure")
+	TEXT("  2: On"),
+	ECVF_RenderThreadSafe
+	);
+
+const uint32 FD3D12ResourceResidencyManager::MEMORY_PRESSURE_FENCE_THRESHOLD[FD3D12ResourceResidencyManager::MEMORY_PRESSURE_LEVELS] = { 60, 30, 15, 7, 3, 1, 0, 0 };
+
+void FD3D12ResourceResidencyManager::ResourceFreed (FD3D12Resource* Resource)
+{
+	check (Resource);
+
+	FScopeLock Lock (&CS);
+
+	FD3D12ResidencyHandle& Handle = Resource->ResidencyHandle;
+
+	if (!Handle.IsResident)
+		UpdateResidency (Resource);
+
+	// Clear any pending resident calls.
+	MakeResident();
+	
+	check (Handle.IsResident && Handle.Index != FD3D12ResidencyHandle::INVALID_INDEX);
+	check (Handle.Index >= 0 && Handle.Index < MAX_RESOURCES)
+
+	Resources[Handle.Index].LastUsedFence = 0;
+	Resources[Handle.Index].Resource      = nullptr;
+
+	// The resource was already resident, so update the stats.
+	ResidentStats.Memory -= Resource->GetResourceSize();
+	ResidentStats.ResourceCount--;
+}
+
+void FD3D12ResourceResidencyManager::UpdateResidency (FD3D12Resource* Resource)
+{
+	check (Resource);
+
+	FScopeLock Lock (&CS);
+
+    FD3D12ResidencyHandle& Handle = Resource->ResidencyHandle;
+
+	uint64 CurrentFence = GetParentDevice()->GetCommandListManager().GetCurrentFence();
+    
+	// Only update once an command list.
+	if (Handle.LastUpdatedFence == CurrentFence)
+		return;
+
+    // Check if the resource is not resident in memory. If it is
+    // not, make it resident as it will be used soon.
+    if (Handle.IsResident == false)
+    {
+		MakeResidentResources.Enqueue(Resource->GetResourceUntracked());
+
+		ResidentStats.PendingMemory += Resource->GetResourceSize();
+		ResidentStats.PendingResouceCount++;
+    }
+       
+    // Clear the residency flag to get the previous index
+    // and clear the element.
+    else if (Handle.Index != FD3D12ResidencyHandle::INVALID_INDEX)
+    {
+		check (Handle.Index >= 0 && Handle.Index < MAX_RESOURCES);
+
+        // Free the current element.
+        Resources[Handle.Index].LastUsedFence = 0;
+        Resources[Handle.Index].Resource      = nullptr;
+    }
+
+	// First entry into the manager. Keep track of it's memory.
+	else
+	{
+		ResidentStats.Memory += Resource->GetResourceSize();
+		ResidentStats.ResourceCount++;
+	}
+    
+	Handle.IsResident       = true;
+	Handle.LastUpdatedFence = CurrentFence;
+    Handle.Index            = HeadIndex;
+
+	uint32 NewHeadIndex = (HeadIndex + 1) % MAX_RESOURCES;
+
+	// The head reached the tail. Attempt a process in case there is space and the tail
+	// just hasn't been updated in a while.
+	if (NewHeadIndex == TailIndex)
+	{
+		Process();
+
+		check(NewHeadIndex != TailIndex);
+	}
+
+	HeadIndex = NewHeadIndex;
+    
+    Resources[Handle.Index].LastUsedFence = CurrentFence;
+    Resources[Handle.Index].Resource      = Resource;
+}
+
+void FD3D12ResourceResidencyManager::Process(uint32 MemoryPressureLevel, uint32 BytesToFree)
+{
+	FScopeLock Lock (&CS);
+
+	TArray<ID3D12Resource*> EvictedResources;
+	
+	// Any memory pressure levels higher than the stall point should execute the current command list
+	// and wait for it to finish before attempting to evict the memory to ensure it's not in use.
+	if (MemoryPressureLevel >= STALL_MEMORY_PRESSURE_LEVEL || BytesToFree != 0)
+	{
+		// TODO: Needs to wait for the GPU to idle.
+		//FD3D12DynamicRHI::GetD3DRHI()->COMMANDL>ExecuteCommandListAndWaitForCompletion();
+	}
+	
+	uint64 LastCompletedFence = GetParentDevice()->GetCommandListManager().GetLastCompletedFence();
+
+	uint32 MemoryPressureFenceThreshold = (BytesToFree != 0) ? 0 : MEMORY_PRESSURE_FENCE_THRESHOLD[MemoryPressureLevel];
+
+	uint64 FenceThreshold = (LastCompletedFence <= MemoryPressureFenceThreshold) ? 0 : LastCompletedFence - MemoryPressureFenceThreshold;
+
+    while (HeadIndex != TailIndex)
+    {
+        FD3D12Element& Element = Resources[TailIndex];
+
+		// Check if it's been used too recently.
+        if (Element.LastUsedFence > FenceThreshold)
+            break;
+           
+        if (Element.Resource != nullptr)
+        {
+            // Add to the list of resources to evict.
+			EvictedResources.Add (Element.Resource->GetResourceUntracked());
+
+			EvictedStats.PendingMemory += Element.Resource->GetResourceSize();
+			EvictedStats.PendingResouceCount++;
+               
+            // Clear the handle.
+            Element.Resource->ResidencyHandle.IsResident       = false;
+			Element.Resource->ResidencyHandle.LastUpdatedFence = 0;
+			Element.Resource->ResidencyHandle.Index            = FD3D12ResidencyHandle::INVALID_INDEX;
+               
+            // Clear the element.
+            Element.LastUsedFence = 0;
+            Element.Resource      = nullptr;
+
+			// Check if there is a specific amount of memory to be freed and break after it's been reached.
+			if (BytesToFree != 0 && EvictedStats.PendingMemory > BytesToFree)
+				break;
+        }
+           
+        // Move the tail up.
+        TailIndex = (TailIndex + 1) % MAX_RESOURCES;
+    }
+	
+	if (EvictedResources.Num())
+	{
+		VERIFYD3D11RESULT (GetParentDevice()->GetDevice()->Evict (EvictedResources.Num(), (ID3D12Pageable**)EvictedResources.GetData()));
+
+		EvictedResources.Empty();
+
+		EvictedStats.MemoryChurn		+= EvictedStats.PendingMemory;
+		EvictedStats.ResourceCountChurn	+= EvictedStats.PendingResouceCount;
+
+		EvictedStats.Memory				+= EvictedStats.PendingMemory;
+		EvictedStats.ResourceCount		+= EvictedStats.PendingResouceCount;
+		ResidentStats.Memory			-= EvictedStats.PendingMemory;
+		ResidentStats.ResourceCount		-= EvictedStats.PendingResouceCount;
+
+		EvictedStats.PendingMemory		 = 0;
+		EvictedStats.PendingResouceCount = 0;
+	}
+
+	static uint32 DebugDisplayFrameCounter = 0;
+	DebugDisplayFrameCounter++;
+
+	if (bLogMemoryResidency > 0 && ((DebugDisplayFrameCounter % 30 == 0 && bLogMemoryResidency == 2) || MemoryPressureLevel > 0 || BytesToFree != 0))
+	{
+		FOutputDeviceRedirector* pOutputDevice = FOutputDeviceRedirector::Get();
+
+		FBufferedOutputDevice BufferedOutput;
+		{
+			FName categoryName(L"ResourceResidencyManager");
+
+			BufferedOutput.CategorizedLogf(categoryName, ELogVerbosity::Log, TEXT("Type             Memory      Resources   Memory Churn  Resources Churn"));
+			BufferedOutput.CategorizedLogf(categoryName, ELogVerbosity::Log, TEXT("---------------  ----------  ----------  ------------  ---------------"));
+
+			BufferedOutput.CategorizedLogf(categoryName, ELogVerbosity::Log, TEXT("Resident (VRAM)  % 7i MB  % 10i  % 9i MB  % 15i"), ResidentStats.Memory / 1024ll / 1024ll, ResidentStats.ResourceCount, ResidentStats.MemoryChurn / 1024ll / 1024ll, ResidentStats.ResourceCountChurn);
+			BufferedOutput.CategorizedLogf(categoryName, ELogVerbosity::Log, TEXT("Evicted  (RAM)   % 7i MB  % 10i  % 9i MB  % 15i"), EvictedStats.Memory / 1024ll / 1024ll, EvictedStats.ResourceCount, EvictedStats.MemoryChurn / 1024ll / 1024ll, EvictedStats.ResourceCountChurn);
+			BufferedOutput.CategorizedLogf(categoryName, ELogVerbosity::Log, TEXT("Total            % 7i MB  % 10i  % 9i MB  % 15i"), (ResidentStats.Memory + EvictedStats.Memory) / 1024ll / 1024ll, ResidentStats.ResourceCount + EvictedStats.ResourceCount, (ResidentStats.MemoryChurn + EvictedStats.MemoryChurn) / 1024ll / 1024ll, ResidentStats.ResourceCountChurn + EvictedStats.ResourceCountChurn);
+			BufferedOutput.CategorizedLogf(categoryName, ELogVerbosity::Log, TEXT(""));
+
+			ResidentStats.MemoryChurn        = 0;
+			ResidentStats.ResourceCountChurn = 0;
+			EvictedStats.MemoryChurn         = 0;
+			EvictedStats.ResourceCountChurn  = 0;
+		}
+		BufferedOutput.RedirectTo(*pOutputDevice);
+	}
+}
+
+void FD3D12ResourceResidencyManager::FreeMemory (uint32 BytesToFree)
+{
+	DXGI_QUERY_VIDEO_MEMORY_INFO LocalVideoMemoryInfo;
+	
+	FD3D12DynamicRHI::GetD3DRHI()->GetLocalVideoMemoryInfo (&LocalVideoMemoryInfo);
+
+	int64 budget = LocalVideoMemoryInfo.Budget;
+
+	int64 AvailableSpace = budget - int64(LocalVideoMemoryInfo.CurrentUsage);
+
+	if (AvailableSpace < BytesToFree)
+		Process (0, BytesToFree - AvailableSpace);
+}
+
+void FD3D12ResourceResidencyManager::MakeResident()
+{
+	TArray<ID3D12Pageable*> PageableResources;
+	PageableResources.Reserve (128);
+
+	ID3D12Resource* Resource;
+
+	while (MakeResidentResources.Dequeue(Resource))
+		PageableResources.Add (Resource);
+
+	if (PageableResources.Num())
+	{
+		HRESULT hresult = GetParentDevice()->GetDevice()->MakeResident (PageableResources.Num(), PageableResources.GetData());
+
+		if (hresult == E_OUTOFMEMORY)
+		{
+			for (uint32 MemoryPressureLevel = 1; MemoryPressureLevel < FD3D12ResourceResidencyManager::MEMORY_PRESSURE_LEVELS; ++MemoryPressureLevel)
+			{
+				FD3D12DynamicRHI::GetD3DRHI()->GetResourceResidencyManager().Process(MemoryPressureLevel);
+
+				hresult = GetParentDevice()->GetDevice()->MakeResident (PageableResources.Num(), PageableResources.GetData());
+
+				if (hresult != E_OUTOFMEMORY)
+					break;
+			}
+		}
+
+		{
+			FScopeLock Lock(&CS);
+
+			VERIFYD3D11RESULT(hresult);
+		
+			ResidentStats.MemoryChurn			+= ResidentStats.PendingMemory;
+			ResidentStats.ResourceCountChurn	+= ResidentStats.PendingResouceCount;
+
+			ResidentStats.Memory				+= ResidentStats.PendingMemory;
+			ResidentStats.ResourceCount			+= ResidentStats.PendingResouceCount;
+			EvictedStats.Memory					-= ResidentStats.PendingMemory;
+			EvictedStats.ResourceCount			-= ResidentStats.PendingResouceCount;
+
+			ResidentStats.PendingMemory			 = 0;
+			ResidentStats.PendingResouceCount	 = 0;
+		}
+	}
+}
+#endif
