UE4 - D3D12RHI (not working)

diff --git a/Engine/Source/Runtime/Windows/D3D12RHI/Private/D3D12RHIPrivate.h b/Engine/Source/Runtime/Windows/D3D12RHI/Private/D3D12RHIPrivate.h
new file mode 100644
--- /dev/null
+++ b/Engine/Source/Runtime/Windows/D3D12RHI/Private/D3D12RHIPrivate.h
@@ -0,0 +1,2249 @@
+// Copyright 1998-2014 Epic Games, Inc. All Rights Reserved.
+
+/*=============================================================================
+	D3D12RHIPrivate.h: Private D3D RHI definitions.
+	=============================================================================*/
+
+#ifndef __D3D12RHIPRIVATE_H__
+#define __D3D12RHIPRIVATE_H__
+
+#define D3D12_SUPPORTS_PARALLEL_RHI_EXECUTE				1
+
+#include "D3D12RHI.h"
+
+// Dependencies.
+#include "Core.h"
+#include "RHI.h"
+#include "GPUProfiler.h"
+#include "ShaderCore.h"
+#include "Engine.h"
+
+DECLARE_LOG_CATEGORY_EXTERN(LogD3D12RHI, Log, All);
+
+#include "D3D12RHIBasePrivate.h"
+#include "StaticArray.h"
+
+#include "AllowWindowsPlatformTypes.h"
+#include "dxgi1_4.h"
+#include "HideWindowsPlatformTypes.h"
+
+// D3D RHI public headers.
+#include "../Public/D3D12Util.h"
+#include "../Public/D3D12State.h"
+#include "../Public/D3D12Resources.h"
+#include "../Public/D3D12Viewport.h"
+#include "../Public/D3D12ConstantBuffer.h"
+#include "D3D12StateCache.h"
+#include "D3D12DirectCommandListManager.h"
+
+#if UE_BUILD_SHIPPING || UE_BUILD_TEST
+#define CHECK_SRV_TRANSITIONS 0
+#else
+#define CHECK_SRV_TRANSITIONS 0	// MSFT: Seb: TODO: ENable
+#endif
+
+// DX11 doesn't support higher MSAA count
+#define DX_MAX_MSAA_COUNT 8
+
+// Definitions.
+#define EXECUTE_DEBUG_COMMAND_LISTS 0
+#define DEBUG_RESOURCE_STATES 0
+#define ENABLE_PLACED_RESOURCES 0 // Disabled due to a couple of NVidia bugs related to placed resources. Works fine on Intel
+#define REMOVE_OLD_QUERY_BATCHES 1  // D3D12: MSFT: TODO: Works around a suspected UE4 InfiltratorDemo bug where a query is never released
+
+#if DEBUG_RESOURCE_STATES
+#define LOG_EXECUTE_COMMAND_LISTS 1
+#else
+#define LOG_EXECUTE_COMMAND_LISTS 0
+#endif
+
+#if EXECUTE_DEBUG_COMMAND_LISTS
+bool GIsDoingQuery = false;
+#define DEBUG_EXECUTE_COMMAND_LIST(scope) if (!GIsDoingQuery) { scope##->FlushCommands(true); }
+#define DEBUG_EXECUTE_COMMAND_CONTEXT(context) if (!GIsDoingQuery) { context##.FlushCommands(true); }
+#define DEBUG_RHI_EXECUTE_COMMAND_LIST(scope) if (!GIsDoingQuery) { scope##->GetRHIDevice()->GetDefaultCommandContext().FlushCommands(true); }
+#else
+#define DEBUG_EXECUTE_COMMAND_LIST(scope) 
+#define DEBUG_EXECUTE_COMMAND_CONTEXT(context) 
+#define DEBUG_RHI_EXECUTE_COMMAND_LIST(scope) 
+#endif
+
+/**
+ * The D3D RHI stats.
+ */
+DECLARE_CYCLE_STAT_EXTERN(TEXT("Present time"), STAT_D3D12PresentTime, STATGROUP_D3D12RHI, );
+DECLARE_CYCLE_STAT_EXTERN(TEXT("CreateTexture time"), STAT_D3D12CreateTextureTime, STATGROUP_D3D12RHI, );
+DECLARE_CYCLE_STAT_EXTERN(TEXT("LockTexture time"), STAT_D3D12LockTextureTime, STATGROUP_D3D12RHI, );
+DECLARE_CYCLE_STAT_EXTERN(TEXT("UnlockTexture time"), STAT_D3D12UnlockTextureTime, STATGROUP_D3D12RHI, );
+DECLARE_CYCLE_STAT_EXTERN(TEXT("CopyTexture time"), STAT_D3D12CopyTextureTime, STATGROUP_D3D12RHI, );
+DECLARE_CYCLE_STAT_EXTERN(TEXT("CreateBoundShaderState time"), STAT_D3D12CreateBoundShaderStateTime, STATGROUP_D3D12RHI, );
+DECLARE_CYCLE_STAT_EXTERN(TEXT("New bound shader state time"), STAT_D3D12NewBoundShaderStateTime, STATGROUP_D3D12RHI, );
+DECLARE_CYCLE_STAT_EXTERN(TEXT("Clean uniform buffer pool"), STAT_D3D12CleanUniformBufferTime, STATGROUP_D3D12RHI, );
+DECLARE_CYCLE_STAT_EXTERN(TEXT("Clear shader resources"), STAT_D3D12ClearShaderResourceTime, STATGROUP_D3D12RHI, );
+DECLARE_DWORD_ACCUMULATOR_STAT_EXTERN(TEXT("Uniform buffer pool num free"), STAT_D3D12NumFreeUniformBuffers, STATGROUP_D3D12RHI, );
+DECLARE_DWORD_ACCUMULATOR_STAT_EXTERN(TEXT("Num Bound Shader State"), STAT_D3D12NumBoundShaderState, STATGROUP_D3D12RHI, );
+DECLARE_MEMORY_STAT_EXTERN(TEXT("Uniform buffer pool memory"), STAT_D3D12FreeUniformBufferMemory, STATGROUP_D3D12RHI, );
+DECLARE_CYCLE_STAT_EXTERN(TEXT("Update uniform buffer"), STAT_D3D12UpdateUniformBufferTime, STATGROUP_D3D12RHI, );
+DECLARE_DWORD_COUNTER_STAT_EXTERN(TEXT("Textures Allocated"), STAT_D3D12TexturesAllocated, STATGROUP_D3D12RHI, );
+DECLARE_DWORD_COUNTER_STAT_EXTERN(TEXT("Textures Released"), STAT_D3D12TexturesReleased, STATGROUP_D3D12RHI, );
+DECLARE_MEMORY_STAT_EXTERN(TEXT("Texture object pool memory"), STAT_D3D12TexturePoolMemory, STATGROUP_D3D12RHI, );
+
+DECLARE_CYCLE_STAT_EXTERN(TEXT("Commit resource tables"), STAT_D3D12CommitResourceTables, STATGROUP_D3D12RHI, );
+DECLARE_CYCLE_STAT_EXTERN(TEXT("Cache resource tables"), STAT_D3D12CacheResourceTables, STATGROUP_D3D12RHI, );
+DECLARE_DWORD_COUNTER_STAT_EXTERN(TEXT("Num cached resource tables"), STAT_D3D12CacheResourceTableCalls, STATGROUP_D3D12RHI, );
+DECLARE_DWORD_COUNTER_STAT_EXTERN(TEXT("Num textures in tables"), STAT_D3D12SetTextureInTableCalls, STATGROUP_D3D12RHI, );
+DECLARE_CYCLE_STAT_EXTERN(TEXT("SetShaderTexture time"), STAT_D3D12SetShaderTextureTime, STATGROUP_D3D12RHI, );
+DECLARE_DWORD_COUNTER_STAT_EXTERN(TEXT("SetShaderTexture calls"), STAT_D3D12SetShaderTextureCalls, STATGROUP_D3D12RHI, );
+
+DECLARE_CYCLE_STAT_EXTERN(TEXT("ApplyState time"), STAT_D3D12ApplyStateTime, STATGROUP_D3D12RHI, );
+DECLARE_CYCLE_STAT_EXTERN(TEXT("ApplyState: Rebuild PSO time"), STAT_D3D12ApplyStateRebuildPSOTime, STATGROUP_D3D12RHI, );
+DECLARE_CYCLE_STAT_EXTERN(TEXT("ApplyState: Set SRV time"), STAT_D3D12ApplyStateSetSRVTime, STATGROUP_D3D12RHI, );
+DECLARE_CYCLE_STAT_EXTERN(TEXT("ApplyState: Set UAV time"), STAT_D3D12ApplyStateSetUAVTime, STATGROUP_D3D12RHI, );
+DECLARE_CYCLE_STAT_EXTERN(TEXT("ApplyState: Set Vertex Buffer time"), STAT_D3D12ApplyStateSetVertexBufferTime, STATGROUP_D3D12RHI, );
+DECLARE_CYCLE_STAT_EXTERN(TEXT("ApplyState: Set Constant Buffer time"), STAT_D3D12ApplyStateSetConstantBufferTime, STATGROUP_D3D12RHI, );
+DECLARE_CYCLE_STAT_EXTERN(TEXT("ApplyState: PSO Create time"), STAT_D3D12ApplyStatePSOCreateTime, STATGROUP_D3D12RHI, );
+DECLARE_CYCLE_STAT_EXTERN(TEXT("Clear MRT time"), STAT_D3D12ClearMRT, STATGROUP_D3D12RHI, );
+DECLARE_CYCLE_STAT_EXTERN(TEXT("Commit graphics constants"), STAT_D3D12CommitGraphicsConstants, STATGROUP_D3D12RHI, );
+DECLARE_CYCLE_STAT_EXTERN(TEXT("Set shader uniform buffer"), STAT_D3D12SetShaderUniformBuffer, STATGROUP_D3D12RHI, );
+DECLARE_CYCLE_STAT_EXTERN(TEXT("Set bound shader state"), STAT_D3D12SetBoundShaderState, STATGROUP_D3D12RHI, );
+DECLARE_CYCLE_STAT_EXTERN(TEXT("ExecuteCommandList time"), STAT_D3D12ExecuteCommandListTime, STATGROUP_D3D12RHI, );
+DECLARE_CYCLE_STAT_EXTERN(TEXT("WaitForFence time"), STAT_D3D12WaitForFenceTime, STATGROUP_D3D12RHI, );
+DECLARE_CYCLE_STAT_EXTERN(TEXT("Resource Clean Up time"), STAT_D3D12ResourceCleanUpTime, STATGROUP_D3D12RHI, );
+DECLARE_MEMORY_STAT_EXTERN(TEXT("Used Video Memory"), STAT_D3D12UsedVideoMemory, STATGROUP_D3D12RHI, );
+DECLARE_MEMORY_STAT_EXTERN(TEXT("Available Video Memory"), STAT_D3D12AvailableVideoMemory, STATGROUP_D3D12RHI, );
+DECLARE_MEMORY_STAT_EXTERN(TEXT("Total Video Memory"), STAT_D3D12TotalVideoMemory, STATGROUP_D3D12RHI, );
+
+// This class handles query heaps
+class FD3D12QueryHeap : public FD3D12DeviceChild
+{
+private:
+    typedef TPair<FD3D12CLSyncPoint, int64> FenceBatchIDPair;
+
+    struct QueryBatch
+    {
+    private:
+        int64 BatchID;            // The unique ID for the batch
+
+        int64 GenerateID()
+        {
+            static int64 ID = 0;
+            return FPlatformAtomics::_InterlockedIncrement(&ID);
+        }
+
+    public:
+        uint32 StartElement;    // The first element in the batch (inclusive)
+        uint32 EndElement;      // The last element in the batch (inclusive)
+        uint32 ElementCount;    // The number of elements in the batch
+        bool bOpen;             // Is the batch still open for more begin/end queries?
+
+        void Clear()
+        {
+            StartElement = 0;
+            EndElement = 0;
+            ElementCount = 0;
+            BatchID = GenerateID();
+            bOpen = false;
+        }
+
+        int64 GetBatchID() const { return BatchID; }
+
+        bool IsValidElement(uint32 Element)
+        {
+            return ((Element >= StartElement) && (Element <= EndElement));
+        }
+    };
+
+public:
+    FD3D12QueryHeap(class FD3D12Device* InParent, const D3D12_QUERY_HEAP_TYPE &InQueryHeapType, uint32 InQueryHeapCount);
+    ~FD3D12QueryHeap();
+
+    void Init();
+    uint32 ReserveElement();            // Reserve a single element in the query heap. Ensure you have room in the heap.
+
+    void StartQueryBatch(FD3D12CommandContext& CmdContext);           // Start tracking a new batch of begin/end query calls that will be resolved together
+    void EndQueryBatchAndResolveQueryData(FD3D12CommandContext& CmdContext, D3D12_QUERY_TYPE InQueryType);  // Stop tracking the current batch of begin/end query calls that will be resolved together
+
+    uint32 BeginQuery(FD3D12CommandContext& CmdContext, D3D12_QUERY_TYPE InQueryType); // Obtain a query from the store of available queries
+    void EndQuery(FD3D12CommandContext& CmdContext, D3D12_QUERY_TYPE InQueryType, uint32 InElement);
+
+	HRESULT MapResultBufferRange(const D3D12_RANGE &InRange)
+	{
+		return ResultBuffer->GetResource()->Map(0, &InRange, &pResultData);	// Note: The pointer is NOT offset by the range
+	}
+
+	void* GetResultBufferData(uint32 InElement) { return reinterpret_cast<uint8*>(pResultData)+(InElement * ResultSize); }  // Get a pointer to the result data for the specified element
+	void* GetResultBufferData(const D3D12_RANGE &InRange) { return reinterpret_cast<uint8*>(pResultData)+InRange.Begin; }	// Get a pointer to the result data at the start of the range
+
+	void UnmapResultBufferRange()
+	{
+		// Unmap an empty range because we didn't write any data (Note: this is different from a null range)
+		const CD3DX12_RANGE emptyRange(0, 0);
+		ResultBuffer->GetResource()->Unmap(0, &emptyRange);
+	}
+
+    uint32 GetQueryHeapCount() const { return QueryHeapDesc.Count; }
+	uint32 GetResultSize() const { return ResultSize; }
+
+	FD3D12CommandContext& GetOwningContext() { check(OwningContext); return *OwningContext; }
+	FD3D12CLSyncPoint GetSyncPoint() { return SyncPoint; }
+
+private:
+    uint32 GetNextElement(uint32 InElement);  // Get the next element, after the specified element. Handles overflow.
+    uint32 GetPreviousElement(uint32 InElement);  // Get the previous element, before the specified element. Handles underflow.
+    bool IsHeapFull();
+    bool IsHeapEmpty() const { return ActiveAllocatedElementCount == 0; }
+
+    uint32 GetNextBatchElement(uint32 InBatchElement);
+    uint32 GetPreviousBatchElement(uint32 InBatchElement);
+
+    uint32 AllocQuery(FD3D12CommandContext& CmdContext, D3D12_QUERY_TYPE InQueryType);
+    void CreateQueryHeap();
+    void CreateResultBuffer();
+
+    uint64 GetResultBufferOffsetForElement(uint32 InElement) const { return ResultSize * InElement; };
+
+private:
+    QueryBatch CurrentQueryBatch;                       // The current recording batch.
+
+    TArray<QueryBatch> ActiveQueryBatches;              // List of active query batches. The data for these is in use.
+
+    static const uint32 MAX_ACTIVE_BATCHES = 10;        // The max number of query batches that will be held.
+    uint32 LastBatch;                                   // The index of the newest batch.
+
+    uint32 HeadActiveElement;                   // The oldest element that is in use (Active). The data for this element is being used.
+    uint32 TailActiveElement;                   // The most recenet element that is in use (Active). The data for this element is being used.
+    uint32 ActiveAllocatedElementCount;         // The number of elements that are in use (Active). Between the head and the tail.
+
+    uint32 LastAllocatedElement;                // The last element that was allocated for BeginQuery
+    uint32 ReservedElementCount;                // The number of queries that are reserved in the pool. Used to know when we need to grow the heap.
+    uint32 ResultSize;                          // The byte size of a result for a single query
+    D3D12_QUERY_HEAP_DESC QueryHeapDesc;        // The description of the current query heap
+    D3D12_QUERY_TYPE QueryType;
+    TRefCountPtr<ID3D12QueryHeap> QueryHeap;    // The query heap where all elements reside
+    TRefCountPtr<FD3D12Resource> ResultBuffer;  // The buffer where all query results are stored
+	FD3D12CommandContext* OwningContext;		// The context containing the queries
+	FD3D12CLSyncPoint SyncPoint;				// The actual command list to which the queries were written
+    void *pResultData;
+};
+
+// This class has multiple inheritance but really FGPUTiming is a static class
+class FD3D12BufferedGPUTiming : public FRenderResource, public FGPUTiming, public FD3D12DeviceChild
+{
+public:
+	/**
+	 * Constructor.
+	 *
+	 * @param InD3DRHI			RHI interface
+	 * @param InBufferSize		Number of buffered measurements
+	 */
+	FD3D12BufferedGPUTiming(class FD3D12Device* InParent, int32 BufferSize);
+
+    FD3D12BufferedGPUTiming()
+    {
+    }
+
+	/**
+	 * Start a GPU timing measurement.
+	 */
+	void	StartTiming();
+
+	/**
+	 * End a GPU timing measurement.
+	 * The timing for this particular measurement will be resolved at a later time by the GPU.
+	 */
+	void	EndTiming();
+
+	/**
+	 * Retrieves the most recently resolved timing measurement.
+	 * The unit is the same as for FPlatformTime::Cycles(). Returns 0 if there are no resolved measurements.
+	 *
+	 * @return	Value of the most recently resolved timing, or 0 if no measurements have been resolved by the GPU yet.
+	 */
+	uint64	GetTiming(bool bGetCurrentResultsAndBlock = false);
+
+	/**
+	 * Initializes all D3D resources.
+	 */
+	virtual void InitDynamicRHI() override;
+
+	/**
+	 * Releases all D3D resources.
+	 */
+	virtual void ReleaseDynamicRHI() override;
+
+private:
+	/**
+	 * Initializes the static variables, if necessary.
+	 */
+	static void PlatformStaticInitialize(void* UserData);
+
+	/** Number of timestamps created in 'StartTimestamps' and 'EndTimestamps'. */
+	int32						BufferSize;
+	/** Current timing being measured on the CPU. */
+	int32						CurrentTimestamp;
+	/** Number of measurements in the buffers (0 - BufferSize). */
+	int32						NumIssuedTimestamps;
+	/** Timestamps for all StartTimings. */
+	TRefCountPtr<ID3D12QueryHeap>	StartTimestampQueryHeap;
+	FD3D12CLSyncPoint*				StartTimestampListHandles;
+	TRefCountPtr<ID3D12Resource>	StartTimestampQueryHeapBuffer;
+	uint64*							StartTimestampQueryHeapBufferData;
+	/** Timestamps for all EndTimings. */
+	TRefCountPtr<ID3D12QueryHeap>	EndTimestampQueryHeap;
+	FD3D12CLSyncPoint*				EndTimestampListHandles;
+	TRefCountPtr<ID3D12Resource>	EndTimestampQueryHeapBuffer;
+	uint64*							EndTimestampQueryHeapBufferData;
+	/** Whether we are currently timing the GPU: between StartTiming() and EndTiming(). */
+	bool						bIsTiming;
+    /** Whether stable power state is currently enabled */
+    bool                        bStablePowerState;
+};
+
+/** A single perf event node, which tracks information about a appBeginDrawEvent/appEndDrawEvent range. */
+class FD3D12EventNode : public FGPUProfilerEventNode, public FD3D12DeviceChild
+{
+public:
+	FD3D12EventNode(const TCHAR* InName, FGPUProfilerEventNode* InParent, class FD3D12Device* InParentDevice) :
+		FGPUProfilerEventNode(InName, InParent),
+        Timing(InParentDevice, 1)
+	{
+		// Initialize Buffered timestamp queries 
+		Timing.InitDynamicRHI();
+	}
+
+	virtual ~FD3D12EventNode()
+	{
+		Timing.ReleaseDynamicRHI();
+	}
+
+	/**
+	 * Returns the time in ms that the GPU spent in this draw event.
+	 * This blocks the CPU if necessary, so can cause hitching.
+	 */
+	virtual float GetTiming() override;
+
+	virtual void StartTiming() override
+	{
+		Timing.StartTiming();
+	}
+
+	virtual void StopTiming() override
+	{
+		Timing.EndTiming();
+	}
+
+	FD3D12BufferedGPUTiming Timing;
+};
+
+/** An entire frame of perf event nodes, including ancillary timers. */
+class FD3D12EventNodeFrame : public FGPUProfilerEventNodeFrame, public FD3D12DeviceChild
+{
+public:
+
+	FD3D12EventNodeFrame(class FD3D12Device* InParent) :
+		FGPUProfilerEventNodeFrame(),
+        RootEventTiming(InParent, 1)
+	{
+		RootEventTiming.InitDynamicRHI();
+	}
+
+	~FD3D12EventNodeFrame()
+	{
+		RootEventTiming.ReleaseDynamicRHI();
+	}
+
+	/** Start this frame of per tracking */
+	virtual void StartFrame() override;
+
+	/** End this frame of per tracking, but do not block yet */
+	virtual void EndFrame() override;
+
+	/** Calculates root timing base frequency (if needed by this RHI) */
+	virtual float GetRootTimingResults() override;
+
+	virtual void LogDisjointQuery() override;
+
+	/** Timer tracking inclusive time spent in the root nodes. */
+	FD3D12BufferedGPUTiming RootEventTiming;
+};
+
+namespace D3D12RHI
+{
+	/**
+	 * Encapsulates GPU profiling logic and data.
+	 * There's only one global instance of this struct so it should only contain global data, nothing specific to a frame.
+	 */
+	struct FD3DGPUProfiler : public FGPUProfiler, public FD3D12DeviceChild
+	{
+		/** Used to measure GPU time per frame. */
+		FD3D12BufferedGPUTiming FrameTiming;
+
+		/** GPU hitch profile histories */
+		TIndirectArray<FD3D12EventNodeFrame> GPUHitchEventNodeFrames;
+
+        FD3DGPUProfiler()
+        {}
+
+		//FD3DGPUProfiler(class FD3D12Device* InParent) :
+		//	FGPUProfiler(),
+        //    FrameTiming(InParent, 4),
+        //    FD3D12DeviceChild(InParent)
+		//{
+		//	// Initialize Buffered timestamp queries 
+		//	FrameTiming.InitResource();
+		//}
+
+        void Init(class FD3D12Device* InParent)
+        {
+            Parent = InParent;
+            FrameTiming = FD3D12BufferedGPUTiming(InParent, 4);
+            // Initialize Buffered timestamp queries 
+            FrameTiming.InitResource();
+        }
+
+		virtual FGPUProfilerEventNode* CreateEventNode(const TCHAR* InName, FGPUProfilerEventNode* InParent) override
+		{
+			FD3D12EventNode* EventNode = new FD3D12EventNode(InName, InParent, GetParentDevice());
+			return EventNode;
+		}
+
+		virtual void PushEvent(const TCHAR* Name) override;
+		virtual void PopEvent() override;
+
+		void BeginFrame(class FD3D12DynamicRHI* InRHI);
+
+		void EndFrame();
+	};
+}
+using namespace D3D12RHI;
+
+class FD3D12CommandContext : public IRHICommandContext, public FD3D12DeviceChild
+{
+public:
+
+	FD3D12CommandContext(class FD3D12Device* InParent);
+	virtual ~FD3D12CommandContext();
+
+	template<typename TRHIType>
+	static FORCEINLINE typename TD3D12ResourceTraits<TRHIType>::TConcreteType* ResourceCast(TRHIType* Resource)
+	{
+		return static_cast<typename TD3D12ResourceTraits<TRHIType>::TConcreteType*>(Resource);
+	}
+
+	template <EShaderFrequency ShaderFrequency>
+	void ClearShaderResourceViews(FD3D12ResourceLocation* Resource);
+
+	void CheckIfSRVIsResolved(FD3D12ShaderResourceView* SRV);
+
+	template <EShaderFrequency ShaderFrequency>
+	void InternalSetShaderResourceView(FD3D12ResourceLocation* Resource, FD3D12ShaderResourceView* SRV, int32 ResourceIndex, FD3D12StateCache::ESRV_Type SrvType = FD3D12StateCache::SRV_Unknown);
+
+	void SetCurrentComputeShader(FComputeShaderRHIParamRef ComputeShader)
+	{
+		CurrentComputeShader = ComputeShader;
+	}
+
+	const FComputeShaderRHIRef& GetCurrentComputeShader() const
+	{
+		return CurrentComputeShader;
+	}
+
+	template <EShaderFrequency ShaderFrequency>
+	void SetShaderResourceView(FD3D12ResourceLocation* Resource, FD3D12ShaderResourceView* SRV, int32 ResourceIndex, FD3D12StateCache::ESRV_Type SrvType = FD3D12StateCache::SRV_Unknown)
+	{
+		InternalSetShaderResourceView<ShaderFrequency>(Resource, SRV, ResourceIndex, SrvType);
+	}
+
+	// Cycle to a new command list, but don't execute the current one yet.
+	void OpenCommandList(bool bRestoreState = false);
+
+	// Close the D3D command list and execute it.  Optionally wait for the GPU to finish. Returns the handle to the command list so you can wait for it later.
+	FD3D12CommandListHandle FlushCommands(bool WaitForCompletion = false);
+
+	void Finish (TArray<FD3D12CommandListHandle>& CommandLists);
+
+	void ClearState();
+	void ConditionalClearShaderResource(FD3D12ResourceLocation* Resource);
+	void ClearAllShaderResources();
+
+    FD3D12DynamicHeapAllocator	UploadHeapAllocator;
+    FD3D12DynamicHeapAllocator& GetUploadHeapAllocator() { return UploadHeapAllocator; }
+
+
+	// A handle to the command list this context owns (granted by the command list manager), and a direct pointer to the D3D command list itself.
+	FD3D12CommandListHandle		CommandListHandle;
+
+	FD3D12StateCache StateCache;
+
+	FD3D12DynamicRHI&			OwningRHI;
+
+	// Tracks the currently set state blocks.
+
+	TRefCountPtr<FD3D12RenderTargetView> CurrentRenderTargets[D3D12_SIMULTANEOUS_RENDER_TARGET_COUNT];
+	TRefCountPtr<FD3D12UnorderedAccessView> CurrentUAVs[D3D12_PS_CS_UAV_REGISTER_COUNT];
+	TRefCountPtr<FD3D12DepthStencilView> CurrentDepthStencilTarget;
+	TRefCountPtr<FD3D12TextureBase> CurrentDepthTexture;
+	uint32 NumSimultaneousRenderTargets;
+	uint32 NumUAVs;
+
+	/** D3D11 defines a maximum of 14 constant buffers per shader stage. */
+	enum { MAX_UNIFORM_BUFFERS_PER_SHADER_STAGE = 14 };
+
+	/** Track the currently bound uniform buffers. */
+	FUniformBufferRHIRef BoundUniformBuffers[SF_NumFrequencies][MAX_UNIFORM_BUFFERS_PER_SHADER_STAGE];
+
+	/** Bit array to track which uniform buffers have changed since the last draw call. */
+	uint16 DirtyUniformBuffers[SF_NumFrequencies];
+
+	/** Tracks the current depth stencil access type. */
+	FExclusiveDepthStencil CurrentDSVAccessType;
+
+	/** When a new shader is set, we discard all old constants set for the previous shader. */
+	bool bDiscardSharedConstants;
+
+	/** Set to true when the current shading setup uses tessellation */
+	bool bUsingTessellation;
+
+	uint32 numDraws;
+	uint32 numDispatches;
+	uint32 numClears;
+	uint32 numBarriers;
+	uint32 numCopies;
+	uint32 otherWorkCounter;
+
+	bool HasDoneWork()
+	{
+		return (numDraws + numDispatches + numClears + numBarriers + numCopies + otherWorkCounter) > 0;
+	}
+
+	/** Dynamic vertex and index buffers. */
+	TRefCountPtr<FD3D12DynamicBuffer> DynamicVB;
+	TRefCountPtr<FD3D12DynamicBuffer> DynamicIB;
+
+	// State for begin/end draw primitive UP interface.
+	uint32 PendingNumVertices;
+	uint32 PendingVertexDataStride;
+	uint32 PendingPrimitiveType;
+	uint32 PendingNumPrimitives;
+	uint32 PendingMinVertexIndex;
+	uint32 PendingNumIndices;
+	uint32 PendingIndexDataStride;
+
+	/** A list of all D3D constant buffers RHIs that have been created. */
+	TArray<TRefCountPtr<FD3D12ConstantBuffer> > VSConstantBuffers;
+	TArray<TRefCountPtr<FD3D12ConstantBuffer> > HSConstantBuffers;
+	TArray<TRefCountPtr<FD3D12ConstantBuffer> > DSConstantBuffers;
+	TArray<TRefCountPtr<FD3D12ConstantBuffer> > PSConstantBuffers;
+	TArray<TRefCountPtr<FD3D12ConstantBuffer> > GSConstantBuffers;
+	TArray<TRefCountPtr<FD3D12ConstantBuffer> > CSConstantBuffers;
+
+	FComputeShaderRHIRef CurrentComputeShader;
+
+#if CHECK_SRV_TRANSITIONS
+	/*
+	* Rendertargets must be explicitly 'resolved' to manage their transition to an SRV on some platforms and DX12
+	* We keep track of targets that need 'resolving' to provide safety asserts at SRV binding time.
+	*/
+	struct FUnresolvedRTInfo
+	{
+		FUnresolvedRTInfo(FName InResourceName, int32 InMipLevel, int32 InNumMips, int32 InArraySlice, int32 InArraySize)
+			: ResourceName(InResourceName)
+			, MipLevel(InMipLevel)
+			, NumMips(InNumMips)
+			, ArraySlice(InArraySlice)
+			, ArraySize(InArraySize)
+		{
+		}
+
+		bool operator==(const FUnresolvedRTInfo& Other) const
+		{
+			return MipLevel == Other.MipLevel &&
+				NumMips == Other.NumMips &&
+				ArraySlice == Other.ArraySlice &&
+				ArraySize == Other.ArraySize;
+		}
+
+		FName ResourceName;
+		int32 MipLevel;
+		int32 NumMips;
+		int32 ArraySlice;
+		int32 ArraySize;
+	};
+	TMultiMap<ID3D12Resource*, FUnresolvedRTInfo> UnresolvedTargets;
+#endif
+
+	TRefCountPtr<FD3D12BoundShaderState> CurrentBoundShaderState;
+
+	/** Initializes the constant buffers.  Called once at RHI initialization time. */
+	void InitConstantBuffers();
+
+	TArray<ID3D12DescriptorHeap*> DescriptorHeaps;
+	void SetDescriptorHeaps( TArray<ID3D12DescriptorHeap*>& InHeaps )
+    { 
+        DescriptorHeaps = InHeaps;
+
+        // Need to set the descriptor heaps on the underlying command list because they can change mid command list
+        if (CommandListHandle != nullptr)
+        {
+			CommandListHandle->SetDescriptorHeaps(DescriptorHeaps.Num(), DescriptorHeaps.GetData());
+        }
+    }
+
+	/** needs to be called before each draw call */
+	virtual void CommitNonComputeShaderConstants();
+
+	/** needs to be called before each dispatch call */
+	virtual void CommitComputeShaderConstants();
+
+	template <class ShaderType> void SetResourcesFromTables(const ShaderType* RESTRICT);
+	void CommitGraphicsResourceTables();
+	void CommitComputeResourceTables(FD3D12ComputeShader* ComputeShader);
+
+	void ValidateExclusiveDepthStencilAccess(FExclusiveDepthStencil Src) const;
+
+	void CommitRenderTargetsAndUAVs();
+
+	template<typename TPixelShader>
+	void ResolveTextureUsingShader(
+		FRHICommandList_RecursiveHazardous& RHICmdList,
+		FD3D12Texture2D* SourceTexture,
+		FD3D12Texture2D* DestTexture,
+		FD3D12RenderTargetView* DestSurfaceRTV,
+		FD3D12DepthStencilView* DestSurfaceDSV,
+		const D3D12_RESOURCE_DESC& ResolveTargetDesc,
+		const FResolveRect& SourceRect,
+		const FResolveRect& DestRect,
+		typename TPixelShader::FParameter PixelShaderParameter
+		);
+
+	// Some platforms might want to override this
+	virtual void SetScissorRectIfRequiredWhenSettingViewport(uint32 MinX, uint32 MinY, uint32 MaxX, uint32 MaxY)
+	{
+		RHISetScissorRect(false, 0, 0, 0, 0);
+	}
+
+	inline bool IsDefaultContext() const;
+
+	virtual void RHISetComputeShader(FComputeShaderRHIParamRef ComputeShader) final override;
+	virtual void RHIDispatchComputeShader(uint32 ThreadGroupCountX, uint32 ThreadGroupCountY, uint32 ThreadGroupCountZ) final override;
+	virtual void RHIDispatchIndirectComputeShader(FVertexBufferRHIParamRef ArgumentBuffer, uint32 ArgumentOffset) final override;
+	virtual void RHIAutomaticCacheFlushAfterComputeShader(bool bEnable) final override;
+	virtual void RHIFlushComputeShaderCache() final override;
+	virtual void RHISetMultipleViewports(uint32 Count, const FViewportBounds* Data) final override;
+	virtual void RHIClearUAV(FUnorderedAccessViewRHIParamRef UnorderedAccessViewRHI, const uint32* Values) final override;
+	virtual void RHICopyToResolveTarget(FTextureRHIParamRef SourceTexture, FTextureRHIParamRef DestTexture, bool bKeepOriginalSurface, const FResolveParams& ResolveParams) final override;
+	virtual void RHIBeginRenderQuery(FRenderQueryRHIParamRef RenderQuery) final override;
+	virtual void RHIEndRenderQuery(FRenderQueryRHIParamRef RenderQuery) final override;
+	virtual void RHIBeginOcclusionQueryBatch() final override;
+	virtual void RHIEndOcclusionQueryBatch() final override;
+	virtual void RHISubmitCommandsHint() final override;
+	virtual void RHIBeginDrawingViewport(FViewportRHIParamRef Viewport, FTextureRHIParamRef RenderTargetRHI) final override;
+	virtual void RHIEndDrawingViewport(FViewportRHIParamRef Viewport, bool bPresent, bool bLockToVsync) final override;
+	virtual void RHIBeginFrame() final override;
+	virtual void RHIEndFrame() final override;
+	virtual void RHIBeginScene() final override;
+	virtual void RHIEndScene() final override;
+	virtual void RHISetStreamSource(uint32 StreamIndex, FVertexBufferRHIParamRef VertexBuffer, uint32 Stride, uint32 Offset) final override;
+	virtual void RHISetRasterizerState(FRasterizerStateRHIParamRef NewState) final override;
+	virtual void RHISetViewport(uint32 MinX, uint32 MinY, float MinZ, uint32 MaxX, uint32 MaxY, float MaxZ) final override;
+	virtual void RHISetScissorRect(bool bEnable, uint32 MinX, uint32 MinY, uint32 MaxX, uint32 MaxY) final override;
+	virtual void RHISetBoundShaderState(FBoundShaderStateRHIParamRef BoundShaderState) final override;
+	virtual void RHISetShaderTexture(FVertexShaderRHIParamRef VertexShader, uint32 TextureIndex, FTextureRHIParamRef NewTexture) final override;
+	virtual void RHISetShaderTexture(FHullShaderRHIParamRef HullShader, uint32 TextureIndex, FTextureRHIParamRef NewTexture) final override;
+	virtual void RHISetShaderTexture(FDomainShaderRHIParamRef DomainShader, uint32 TextureIndex, FTextureRHIParamRef NewTexture) final override;
+	virtual void RHISetShaderTexture(FGeometryShaderRHIParamRef GeometryShader, uint32 TextureIndex, FTextureRHIParamRef NewTexture) final override;
+	virtual void RHISetShaderTexture(FPixelShaderRHIParamRef PixelShader, uint32 TextureIndex, FTextureRHIParamRef NewTexture) final override;
+	virtual void RHISetShaderTexture(FComputeShaderRHIParamRef PixelShader, uint32 TextureIndex, FTextureRHIParamRef NewTexture) final override;
+	virtual void RHISetShaderSampler(FComputeShaderRHIParamRef ComputeShader, uint32 SamplerIndex, FSamplerStateRHIParamRef NewState) final override;
+	virtual void RHISetShaderSampler(FVertexShaderRHIParamRef VertexShader, uint32 SamplerIndex, FSamplerStateRHIParamRef NewState) final override;
+	virtual void RHISetShaderSampler(FGeometryShaderRHIParamRef GeometryShader, uint32 SamplerIndex, FSamplerStateRHIParamRef NewState) final override;
+	virtual void RHISetShaderSampler(FDomainShaderRHIParamRef DomainShader, uint32 SamplerIndex, FSamplerStateRHIParamRef NewState) final override;
+	virtual void RHISetShaderSampler(FHullShaderRHIParamRef HullShader, uint32 SamplerIndex, FSamplerStateRHIParamRef NewState) final override;
+	virtual void RHISetShaderSampler(FPixelShaderRHIParamRef PixelShader, uint32 SamplerIndex, FSamplerStateRHIParamRef NewState) final override;
+	virtual void RHISetUAVParameter(FComputeShaderRHIParamRef ComputeShader, uint32 UAVIndex, FUnorderedAccessViewRHIParamRef UAV) final override;
+	virtual void RHISetUAVParameter(FComputeShaderRHIParamRef ComputeShader, uint32 UAVIndex, FUnorderedAccessViewRHIParamRef UAV, uint32 InitialCount) final override;
+	virtual void RHISetShaderResourceViewParameter(FPixelShaderRHIParamRef PixelShader, uint32 SamplerIndex, FShaderResourceViewRHIParamRef SRV) final override;
+	virtual void RHISetShaderResourceViewParameter(FVertexShaderRHIParamRef VertexShader, uint32 SamplerIndex, FShaderResourceViewRHIParamRef SRV) final override;
+	virtual void RHISetShaderResourceViewParameter(FComputeShaderRHIParamRef ComputeShader, uint32 SamplerIndex, FShaderResourceViewRHIParamRef SRV) final override;
+	virtual void RHISetShaderResourceViewParameter(FHullShaderRHIParamRef HullShader, uint32 SamplerIndex, FShaderResourceViewRHIParamRef SRV) final override;
+	virtual void RHISetShaderResourceViewParameter(FDomainShaderRHIParamRef DomainShader, uint32 SamplerIndex, FShaderResourceViewRHIParamRef SRV) final override;
+	virtual void RHISetShaderResourceViewParameter(FGeometryShaderRHIParamRef GeometryShader, uint32 SamplerIndex, FShaderResourceViewRHIParamRef SRV) final override;
+	virtual void RHISetShaderUniformBuffer(FVertexShaderRHIParamRef VertexShader, uint32 BufferIndex, FUniformBufferRHIParamRef Buffer) final override;
+	virtual void RHISetShaderUniformBuffer(FHullShaderRHIParamRef HullShader, uint32 BufferIndex, FUniformBufferRHIParamRef Buffer) final override;
+	virtual void RHISetShaderUniformBuffer(FDomainShaderRHIParamRef DomainShader, uint32 BufferIndex, FUniformBufferRHIParamRef Buffer) final override;
+	virtual void RHISetShaderUniformBuffer(FGeometryShaderRHIParamRef GeometryShader, uint32 BufferIndex, FUniformBufferRHIParamRef Buffer) final override;
+	virtual void RHISetShaderUniformBuffer(FPixelShaderRHIParamRef PixelShader, uint32 BufferIndex, FUniformBufferRHIParamRef Buffer) final override;
+	virtual void RHISetShaderUniformBuffer(FComputeShaderRHIParamRef ComputeShader, uint32 BufferIndex, FUniformBufferRHIParamRef Buffer) final override;
+	virtual void RHISetShaderParameter(FVertexShaderRHIParamRef VertexShader, uint32 BufferIndex, uint32 BaseIndex, uint32 NumBytes, const void* NewValue) final override;
+	virtual void RHISetShaderParameter(FPixelShaderRHIParamRef PixelShader, uint32 BufferIndex, uint32 BaseIndex, uint32 NumBytes, const void* NewValue) final override;
+	virtual void RHISetShaderParameter(FHullShaderRHIParamRef HullShader, uint32 BufferIndex, uint32 BaseIndex, uint32 NumBytes, const void* NewValue) final override;
+	virtual void RHISetShaderParameter(FDomainShaderRHIParamRef DomainShader, uint32 BufferIndex, uint32 BaseIndex, uint32 NumBytes, const void* NewValue) final override;
+	virtual void RHISetShaderParameter(FGeometryShaderRHIParamRef GeometryShader, uint32 BufferIndex, uint32 BaseIndex, uint32 NumBytes, const void* NewValue) final override;
+	virtual void RHISetShaderParameter(FComputeShaderRHIParamRef ComputeShader, uint32 BufferIndex, uint32 BaseIndex, uint32 NumBytes, const void* NewValue) final override;
+	virtual void RHISetDepthStencilState(FDepthStencilStateRHIParamRef NewState, uint32 StencilRef) final override;
+	virtual void RHISetBlendState(FBlendStateRHIParamRef NewState, const FLinearColor& BlendFactor) final override;
+	virtual void RHISetRenderTargets(uint32 NumSimultaneousRenderTargets, const FRHIRenderTargetView* NewRenderTargets, const FRHIDepthRenderTargetView* NewDepthStencilTarget, uint32 NumUAVs, const FUnorderedAccessViewRHIParamRef* UAVs) final override;
+	virtual void RHISetRenderTargetsAndClear(const FRHISetRenderTargetsInfo& RenderTargetsInfo) final override;
+	//virtual void RHIBindClearMRTValues(bool bClearColor, bool bClearDepth, bool bClearStencil) final override;
+	virtual void RHIDrawPrimitive(uint32 PrimitiveType, uint32 BaseVertexIndex, uint32 NumPrimitives, uint32 NumInstances) final override;
+	virtual void RHIDrawPrimitiveIndirect(uint32 PrimitiveType, FVertexBufferRHIParamRef ArgumentBuffer, uint32 ArgumentOffset) final override;
+	virtual void RHIDrawIndexedIndirect(FIndexBufferRHIParamRef IndexBufferRHI, uint32 PrimitiveType, FStructuredBufferRHIParamRef ArgumentsBufferRHI, int32 DrawArgumentsIndex, uint32 NumInstances) final override;
+	virtual void RHIDrawIndexedPrimitive(FIndexBufferRHIParamRef IndexBuffer, uint32 PrimitiveType, int32 BaseVertexIndex, uint32 FirstInstance, uint32 NumVertices, uint32 StartIndex, uint32 NumPrimitives, uint32 NumInstances) final override;
+	virtual void RHIDrawIndexedPrimitiveIndirect(uint32 PrimitiveType, FIndexBufferRHIParamRef IndexBuffer, FVertexBufferRHIParamRef ArgumentBuffer, uint32 ArgumentOffset) final override;
+	virtual void RHIBeginDrawPrimitiveUP(uint32 PrimitiveType, uint32 NumPrimitives, uint32 NumVertices, uint32 VertexDataStride, void*& OutVertexData) final override;
+	virtual void RHIEndDrawPrimitiveUP() final override;
+	virtual void RHIBeginDrawIndexedPrimitiveUP(uint32 PrimitiveType, uint32 NumPrimitives, uint32 NumVertices, uint32 VertexDataStride, void*& OutVertexData, uint32 MinVertexIndex, uint32 NumIndices, uint32 IndexDataStride, void*& OutIndexData) final override;
+	virtual void RHIEndDrawIndexedPrimitiveUP() final override;
+	virtual void RHIClear(bool bClearColor, const FLinearColor& Color, bool bClearDepth, float Depth, bool bClearStencil, uint32 Stencil, FIntRect ExcludeRect) final override;
+	virtual void RHIClearMRT(bool bClearColor, int32 NumClearColors, const FLinearColor* ColorArray, bool bClearDepth, float Depth, bool bClearStencil, uint32 Stencil, FIntRect ExcludeRect) final override;
+	virtual void RHIEnableDepthBoundsTest(bool bEnable, float MinDepth, float MaxDepth) final override;
+	virtual void RHIPushEvent(const TCHAR* Name) final override;
+	virtual void RHIPopEvent() final override;
+	virtual void RHIUpdateTextureReference(FTextureReferenceRHIParamRef TextureRef, FTextureRHIParamRef NewTexture) final override;
+	virtual void RHIBeginAsyncComputeJob_DrawThread(EAsyncComputePriority Priority) override;
+	virtual void RHIEndAsyncComputeJob_DrawThread(uint32 FenceIndex) override;	
+	virtual void RHIGraphicsWaitOnAsyncComputeJob(uint32 FenceIndex) override;
+
+	virtual void RHIClearMRTImpl(bool bClearColor, int32 NumClearColors, const FLinearColor* ColorArray, bool bClearDepth, float Depth, bool bClearStencil, uint32 Stencil, FIntRect ExcludeRect, bool bForceShaderClear);
+};
+
+struct FD3D12Adapter
+{
+    /** -1 if not supported or FindAdpater() wasn't called. Ideally we would store a pointer to IDXGIAdapter but it's unlikely the adpaters change during engine init. */
+    int32 AdapterIndex;
+    /** The maximum D3D11 feature level supported. 0 if not supported or FindAdpater() wasn't called */
+    D3D_FEATURE_LEVEL MaxSupportedFeatureLevel;
+
+    // constructor
+    FD3D12Adapter(int32 InAdapterIndex = -1, D3D_FEATURE_LEVEL InMaxSupportedFeatureLevel = (D3D_FEATURE_LEVEL)0)
+        : AdapterIndex(InAdapterIndex)
+        , MaxSupportedFeatureLevel(InMaxSupportedFeatureLevel)
+    {
+    }
+
+    bool IsValid() const
+    {
+        return MaxSupportedFeatureLevel != (D3D_FEATURE_LEVEL)0 && AdapterIndex >= 0;
+    }
+};
+
+class FD3D12DynamicRHI;
+
+class FD3D12Device
+{
+
+private:
+    FD3D12Adapter               DeviceAdapter;
+    DXGI_ADAPTER_DESC           AdapterDesc;
+    TRefCountPtr<IDXGIAdapter3> DxgiAdapter3;
+    TRefCountPtr<IDXGIFactory4> DXGIFactory;
+    TRefCountPtr<ID3D12Device>  Direct3DDevice;
+    FD3D12DynamicRHI*           OwningRHI;
+	D3D12_RESOURCE_HEAP_TIER    ResourceHeapTier;
+
+    /** True if the device being used has been removed. */
+    bool bDeviceRemoved;
+
+public:
+
+    FD3D12Device(FD3D12DynamicRHI* InOwningRHI, IDXGIFactory4* InDXGIFactory, FD3D12Adapter& InAdapter);
+
+    /** If it hasn't been initialized yet, initializes the D3D device. */
+    virtual void InitD3DDevice();
+
+    void CreateCommandContexts();
+
+    /**
+    * Cleanup the D3D device.
+    * This function must be called from the main game thread.
+    */
+    virtual void CleanupD3DDevice();
+
+    /**
+    * Populates a D3D query's data buffer.
+    * @param Query - The occlusion query to read data from.
+    * @param bWait - If true, it will wait for the query to finish.
+    * @return true if the query finished.
+    */
+    bool GetQueryData(FD3D12OcclusionQuery& Query, bool bWait);
+
+    FSamplerStateRHIRef CreateSamplerState(const FSamplerStateInitializerRHI& Initializer);
+
+    inline ID3D12Device*                GetDevice()                                 { return Direct3DDevice; }
+    inline uint32                       GetAdapterIndex()                           { return DeviceAdapter.AdapterIndex; }
+    inline D3D_FEATURE_LEVEL            GetFeatureLevel()                           { return DeviceAdapter.MaxSupportedFeatureLevel; }
+    inline DXGI_ADAPTER_DESC*           GetD3DAdapterDesc()                         { return &AdapterDesc; }
+    inline void                         SetDeviceRemoved(bool value)                { bDeviceRemoved = value; }
+    inline bool                         IsDeviceRemoved()                           { return bDeviceRemoved; }
+    inline FD3D12DynamicRHI*            GetOwningRHI()                              { return OwningRHI; }
+    inline void                         SetAdapter3(IDXGIAdapter3* Adapter3)        { DxgiAdapter3 = Adapter3; }
+    inline IDXGIAdapter3*               GetAdapter3()                               { return DxgiAdapter3; }
+
+    inline TArray<FD3D12Viewport*>&     GetViewports()                              { return Viewports; }
+    inline FD3D12Viewport*              GetDrawingViewport()                        { return DrawingViewport; }
+    inline void                         SetDrawingViewport(FD3D12Viewport* InViewport)  { DrawingViewport = InViewport; }
+
+    inline FD3D12PipelineStateCache&    GetPSOCache()                               { return PipelineStateCache; }
+    inline ID3D12RootSignature*         GetGraphicsRootSignature()                  { return GraphicsRS; }
+    inline ID3D12RootSignature*         GetComputeRootSignature()                   { return ComputeRS; }
+    inline ID3D12CommandSignature*      GetDrawIndirectCommandSignature()           { return DrawIndirectCommandSignature; }
+    inline ID3D12CommandSignature*      GetDrawIndexedIndirectCommandSignature()    { return DrawIndexedIndirectCommandSignature; }
+    inline ID3D12CommandSignature*      GetDispatchIndirectCommandSignature()       { return DispatchIndirectCommandSignature; }
+    inline FD3D12QueryHeap*             GetQueryHeap()                              { return &OcclusionQueryHeap; }
+
+	TRefCountPtr<ID3D12Resource>		GetCounterUploadHeap()						{ return CounterUploadHeap; }
+	uint32&								GetCounterUploadHeapIndex()					{ return CounterUploadHeapIndex; }
+	void*								GetCounterUploadHeapData()					{ return CounterUploadHeapData; }
+
+    template <typename TViewDesc> FDescriptorHeapManager& GetViewDescriptorAllocator();
+    template<> FDescriptorHeapManager& GetViewDescriptorAllocator<D3D12_SHADER_RESOURCE_VIEW_DESC>()    { return SRVAllocator; }
+    template<> FDescriptorHeapManager& GetViewDescriptorAllocator<D3D12_RENDER_TARGET_VIEW_DESC>()      { return RTVAllocator; }
+    template<> FDescriptorHeapManager& GetViewDescriptorAllocator<D3D12_DEPTH_STENCIL_VIEW_DESC>()      { return DSVAllocator; }
+    template<> FDescriptorHeapManager& GetViewDescriptorAllocator<D3D12_UNORDERED_ACCESS_VIEW_DESC>()   { return UAVAllocator; }
+    template<> FDescriptorHeapManager& GetViewDescriptorAllocator<D3D12_CONSTANT_BUFFER_VIEW_DESC>()    { return CBVAllocator; }
+
+    inline FDescriptorHeapManager&          GetSamplerDescriptorAllocator() { return SamplerAllocator; }
+    inline FD3D12CommandListManager&        GetCommandListManager()         { return CommandListManager; }
+    inline FD3D12CommandListManager&        GetCopyCommandListManager()     { return CopyCommandListManager; }
+    inline FD3D12ResourceHelper&            GetResourceHelper()             { return ResourceHelper; }
+    inline FD3D12DeferredDeletionQueue&     GetDeferredDeletionQueue()      { return DeferredDeletionQueue; }
+    inline FD3D12DefaultBufferAllocator&    GetDefaultBufferAllocator()     { return DefaultBufferAllocator; }
+
+	inline uint32 GetNumContexts() { return CommandContextArray.Num(); }
+	inline FD3D12CommandContext& GetCommandContext(uint32 i = 0) const { return *CommandContextArray[i];	}
+
+	inline FD3D12CommandContext* ObtainCommandContext() { 
+		FScopeLock Lock(&FreeContextsLock);
+		return FreeCommandContexts.Pop();
+	}
+	inline void ReleaseCommandContext(FD3D12CommandContext* CmdContext) {
+		FScopeLock Lock(&FreeContextsLock);
+		FreeCommandContexts.Add(CmdContext);
+	}
+
+	inline FD3D12CommandContext& GetDefaultCommandContext() const { return GetCommandContext(0); }
+	inline FD3D12DynamicHeapAllocator& GetDefaultUploadHeapAllocator() { return GetDefaultCommandContext().GetUploadHeapAllocator();	}
+
+	inline D3D12_RESOURCE_HEAP_TIER   GetResourceHeapTier() { return ResourceHeapTier; }
+
+	TArray<FD3D12CommandListHandle> PendingCommandLists;
+    uint32 PendingCommandListsTotalWorkCommands;
+
+protected:
+
+
+    /** A pool of command lists and allocators we can cycle through for the global D3D device */
+    FD3D12CommandListManager CommandListManager;
+    FD3D12CommandListManager CopyCommandListManager;
+
+    TRefCountPtr<ID3D12RootSignature> GraphicsRS;
+    TRefCountPtr<ID3D12RootSignature> ComputeRS;
+
+    TRefCountPtr<ID3D12CommandSignature> DrawIndirectCommandSignature;
+    TRefCountPtr<ID3D12CommandSignature> DrawIndexedIndirectCommandSignature;
+    TRefCountPtr<ID3D12CommandSignature> DispatchIndirectCommandSignature;
+
+	TRefCountPtr<ID3D12Resource> CounterUploadHeap;
+	uint32                       CounterUploadHeapIndex;
+	void*                        CounterUploadHeapData;
+
+    FD3D12PipelineStateCache PipelineStateCache;
+
+    // Must be before the StateCache so that destructor ordering is valid
+    FDescriptorHeapManager RTVAllocator;
+    FDescriptorHeapManager DSVAllocator;
+    FDescriptorHeapManager SRVAllocator;
+    FDescriptorHeapManager UAVAllocator;
+    FDescriptorHeapManager SamplerAllocator;
+    FDescriptorHeapManager CBVAllocator;
+
+    /** The resource manager allocates resources and tracks when to release them */
+    FD3D12ResourceHelper ResourceHelper;
+
+    FD3D12DeferredDeletionQueue DeferredDeletionQueue;
+
+    FD3D12QueryHeap OcclusionQueryHeap;
+
+    FD3D12DefaultBufferAllocator DefaultBufferAllocator;
+
+	TArray<FD3D12CommandContext*> CommandContextArray;
+	TArray<FD3D12CommandContext*> FreeCommandContexts;
+	FCriticalSection FreeContextsLock;
+
+    /** A list of all viewport RHIs that have been created. */
+    TArray<FD3D12Viewport*> Viewports;
+
+    /** The viewport which is currently being drawn. */
+    TRefCountPtr<FD3D12Viewport> DrawingViewport;
+
+    TMap< D3D12_SAMPLER_DESC, TRefCountPtr<FD3D12SamplerState> > SamplerMap;
+    uint32 SamplerID;
+
+    // set by UpdateMSAASettings(), get by GetMSAAQuality()
+    // [SampleCount] = Quality, 0xffffffff if not supported
+    uint32 AvailableMSAAQualities[DX_MAX_MSAA_COUNT + 1];
+
+    // Creates default root and execute indirect signatures
+    void CreateSignatures();
+
+    // shared code for different D3D11 devices (e.g. PC DirectX11 and XboxOne) called
+    // after device creation and GRHISupportsAsyncTextureCreation was set and before resource init
+    void SetupAfterDeviceCreation();
+
+    // called by SetupAfterDeviceCreation() when the device gets initialized
+
+    void UpdateMSAASettings();
+
+    void ReleasePooledUniformBuffers();
+
+};
+
+/** The interface which is implemented by the dynamically bound RHI. */
+class FD3D12DynamicRHI : public FDynamicRHI
+{
+	friend class FD3D12CommandContext;
+
+	static FD3D12DynamicRHI *SingleD3DRHI;
+
+public:
+
+	static FD3D12DynamicRHI *GetD3DRHI() { return SingleD3DRHI; }
+
+	/** Global D3D12 lock list */
+	TMap<FD3D12LockedKey, FD3D12LockedData> OutstandingLocks;
+
+	/** Initialization constructor. */
+    FD3D12DynamicRHI(IDXGIFactory4* InDXGIFactory, FD3D12Adapter& InAdapter);
+
+	/** Destructor */
+	virtual ~FD3D12DynamicRHI() {}
+
+	// FDynamicRHI interface.
+	virtual void Init() override;
+	virtual void Shutdown() override;
+
+    template<typename TRHIType>
+    static FORCEINLINE typename TD3D12ResourceTraits<TRHIType>::TConcreteType* ResourceCast(TRHIType* Resource)
+    {
+        return static_cast<typename TD3D12ResourceTraits<TRHIType>::TConcreteType*>(Resource);
+    }
+
+	virtual void RHIGpuTimeBegin(uint32 Hash,bool bCompute) final override;
+	virtual void RHIGpuTimeEnd(uint32 Hash,bool bCompute) final override;
+	virtual FSamplerStateRHIRef RHICreateSamplerState(const FSamplerStateInitializerRHI& Initializer) final override;
+	virtual FRasterizerStateRHIRef RHICreateRasterizerState(const FRasterizerStateInitializerRHI& Initializer) final override;
+	virtual FDepthStencilStateRHIRef RHICreateDepthStencilState(const FDepthStencilStateInitializerRHI& Initializer) final override;
+	virtual FBlendStateRHIRef RHICreateBlendState(const FBlendStateInitializerRHI& Initializer) final override;
+	virtual FVertexDeclarationRHIRef RHICreateVertexDeclaration(const FVertexDeclarationElementList& Elements) final override;
+	virtual FPixelShaderRHIRef RHICreatePixelShader(const TArray<uint8>& Code) final override;
+	virtual FVertexShaderRHIRef RHICreateVertexShader(const TArray<uint8>& Code) final override;
+	virtual FHullShaderRHIRef RHICreateHullShader(const TArray<uint8>& Code) final override;
+	virtual FDomainShaderRHIRef RHICreateDomainShader(const TArray<uint8>& Code) final override;
+	virtual FGeometryShaderRHIRef RHICreateGeometryShader(const TArray<uint8>& Code) final override;
+	virtual FGeometryShaderRHIRef RHICreateGeometryShaderWithStreamOutput(const TArray<uint8>& Code, const FStreamOutElementList& ElementList, uint32 NumStrides, const uint32* Strides, int32 RasterizedStream) final override;
+	virtual FComputeShaderRHIRef RHICreateComputeShader(const TArray<uint8>& Code) final override;
+	virtual FBoundShaderStateRHIRef RHICreateBoundShaderState(FVertexDeclarationRHIParamRef VertexDeclaration, FVertexShaderRHIParamRef VertexShader, FHullShaderRHIParamRef HullShader, FDomainShaderRHIParamRef DomainShader, FPixelShaderRHIParamRef PixelShader, FGeometryShaderRHIParamRef GeometryShader) final override;
+	virtual FUniformBufferRHIRef RHICreateUniformBuffer(const void* Contents, const FRHIUniformBufferLayout& Layout, EUniformBufferUsage Usage) final override;
+	virtual FIndexBufferRHIRef RHICreateIndexBuffer(uint32 Stride, uint32 Size, uint32 InUsage, FRHIResourceCreateInfo& CreateInfo) final override;
+	virtual void* RHILockIndexBuffer(FIndexBufferRHIParamRef IndexBuffer, uint32 Offset, uint32 Size, EResourceLockMode LockMode) final override;
+	virtual void RHIUnlockIndexBuffer(FIndexBufferRHIParamRef IndexBuffer) final override;
+	virtual FVertexBufferRHIRef RHICreateVertexBuffer(uint32 Size, uint32 InUsage, FRHIResourceCreateInfo& CreateInfo) final override;
+	virtual void* RHILockVertexBuffer(FVertexBufferRHIParamRef VertexBuffer, uint32 Offset, uint32 SizeRHI, EResourceLockMode LockMode) final override;
+	virtual void RHIUnlockVertexBuffer(FVertexBufferRHIParamRef VertexBuffer) final override;
+	virtual void RHICopyVertexBuffer(FVertexBufferRHIParamRef SourceBuffer,FVertexBufferRHIParamRef DestBuffer) final override;
+	virtual FStructuredBufferRHIRef RHICreateStructuredBuffer(uint32 Stride, uint32 Size, uint32 InUsage, FRHIResourceCreateInfo& CreateInfo) final override;
+	virtual void* RHILockStructuredBuffer(FStructuredBufferRHIParamRef StructuredBuffer, uint32 Offset, uint32 SizeRHI, EResourceLockMode LockMode) final override;
+	virtual void RHIUnlockStructuredBuffer(FStructuredBufferRHIParamRef StructuredBuffer) final override;
+	virtual FUnorderedAccessViewRHIRef RHICreateUnorderedAccessView(FStructuredBufferRHIParamRef StructuredBuffer, bool bUseUAVCounter, bool bAppendBuffer) final override;
+	virtual FUnorderedAccessViewRHIRef RHICreateUnorderedAccessView(FTextureRHIParamRef Texture, uint32 MipLevel) final override;
+	virtual FUnorderedAccessViewRHIRef RHICreateUnorderedAccessView(FVertexBufferRHIParamRef VertexBuffer, uint8 Format) final override;
+	virtual FShaderResourceViewRHIRef RHICreateShaderResourceView(FStructuredBufferRHIParamRef StructuredBuffer) final override;
+	virtual FShaderResourceViewRHIRef RHICreateShaderResourceView(FVertexBufferRHIParamRef VertexBuffer, uint32 Stride, uint8 Format) final override;
+	virtual uint64 RHICalcTexture2DPlatformSize(uint32 SizeX, uint32 SizeY, uint8 Format, uint32 NumMips, uint32 NumSamples, uint32 Flags, uint32& OutAlign) final override;
+	virtual uint64 RHICalcTexture3DPlatformSize(uint32 SizeX, uint32 SizeY, uint32 SizeZ, uint8 Format, uint32 NumMips, uint32 Flags, uint32& OutAlign) final override;
+	virtual uint64 RHICalcTextureCubePlatformSize(uint32 Size, uint8 Format, uint32 NumMips, uint32 Flags, uint32& OutAlign) final override;
+	virtual void RHIGetTextureMemoryStats(FTextureMemoryStats& OutStats) final override;
+	virtual bool RHIGetTextureMemoryVisualizeData(FColor* TextureData,int32 SizeX,int32 SizeY,int32 Pitch,int32 PixelSize) final override;
+	virtual FTextureReferenceRHIRef RHICreateTextureReference(FLastRenderTimeContainer* LastRenderTime) final override;
+	virtual FTexture2DRHIRef RHICreateTexture2D(uint32 SizeX, uint32 SizeY, uint8 Format, uint32 NumMips, uint32 NumSamples, uint32 Flags, FRHIResourceCreateInfo& CreateInfo) final override;
+	virtual FTexture2DRHIRef RHIAsyncCreateTexture2D(uint32 SizeX, uint32 SizeY, uint8 Format, uint32 NumMips, uint32 Flags, void** InitialMipData, uint32 NumInitialMips) final override;
+	virtual void RHICopySharedMips(FTexture2DRHIParamRef DestTexture2D, FTexture2DRHIParamRef SrcTexture2D) final override;
+	virtual FTexture2DArrayRHIRef RHICreateTexture2DArray(uint32 SizeX, uint32 SizeY, uint32 SizeZ, uint8 Format, uint32 NumMips, uint32 Flags, FRHIResourceCreateInfo& CreateInfo) final override;
+	virtual FTexture3DRHIRef RHICreateTexture3D(uint32 SizeX, uint32 SizeY, uint32 SizeZ, uint8 Format, uint32 NumMips, uint32 Flags, FRHIResourceCreateInfo& CreateInfo) final override;
+	virtual void RHIGetResourceInfo(FTextureRHIParamRef Ref, FRHIResourceInfo& OutInfo) final override;
+	virtual FShaderResourceViewRHIRef RHICreateShaderResourceView(FTexture2DRHIParamRef Texture2DRHI, uint8 MipLevel) final override;
+	virtual FShaderResourceViewRHIRef RHICreateShaderResourceView(FTexture2DRHIParamRef Texture2DRHI, uint8 MipLevel, uint8 NumMipLevels, uint8 Format) final override;
+	virtual FShaderResourceViewRHIRef RHICreateShaderResourceView(FTexture3DRHIParamRef Texture3DRHI, uint8 MipLevel) final override;
+	virtual FShaderResourceViewRHIRef RHICreateShaderResourceView(FTexture2DArrayRHIParamRef Texture2DArrayRHI, uint8 MipLevel) final override;
+	virtual FShaderResourceViewRHIRef RHICreateShaderResourceView(FTextureCubeRHIParamRef TextureCubeRHI, uint8 MipLevel) final override;
+	virtual void RHIGenerateMips(FTextureRHIParamRef Texture) final override;
+	virtual uint32 RHIComputeMemorySize(FTextureRHIParamRef TextureRHI) final override;
+	virtual FTexture2DRHIRef RHIAsyncReallocateTexture2D(FTexture2DRHIParamRef Texture2D, int32 NewMipCount, int32 NewSizeX, int32 NewSizeY, FThreadSafeCounter* RequestStatus) final override;
+	virtual ETextureReallocationStatus RHIFinalizeAsyncReallocateTexture2D(FTexture2DRHIParamRef Texture2D, bool bBlockUntilCompleted) final override;
+	virtual ETextureReallocationStatus RHICancelAsyncReallocateTexture2D(FTexture2DRHIParamRef Texture2D, bool bBlockUntilCompleted) final override;
+	virtual void* RHILockTexture2D(FTexture2DRHIParamRef Texture, uint32 MipIndex, EResourceLockMode LockMode, uint32& DestStride, bool bLockWithinMiptail) final override;
+	virtual void RHIUnlockTexture2D(FTexture2DRHIParamRef Texture, uint32 MipIndex, bool bLockWithinMiptail) final override;
+	virtual void* RHILockTexture2DArray(FTexture2DArrayRHIParamRef Texture, uint32 TextureIndex, uint32 MipIndex, EResourceLockMode LockMode, uint32& DestStride, bool bLockWithinMiptail) final override;
+	virtual void RHIUnlockTexture2DArray(FTexture2DArrayRHIParamRef Texture, uint32 TextureIndex, uint32 MipIndex, bool bLockWithinMiptail) final override;
+	virtual void RHIUpdateTexture2D(FTexture2DRHIParamRef Texture, uint32 MipIndex, const struct FUpdateTextureRegion2D& UpdateRegion, uint32 SourcePitch, const uint8* SourceData) final override;
+	virtual void RHIUpdateTexture3D(FTexture3DRHIParamRef Texture, uint32 MipIndex, const struct FUpdateTextureRegion3D& UpdateRegion, uint32 SourceRowPitch, uint32 SourceDepthPitch, const uint8* SourceData) final override;
+	virtual FTextureCubeRHIRef RHICreateTextureCube(uint32 Size, uint8 Format, uint32 NumMips, uint32 Flags, FRHIResourceCreateInfo& CreateInfo) final override;
+	virtual FTextureCubeRHIRef RHICreateTextureCubeArray(uint32 Size, uint32 ArraySize, uint8 Format, uint32 NumMips, uint32 Flags, FRHIResourceCreateInfo& CreateInfo) final override;
+	virtual void* RHILockTextureCubeFace(FTextureCubeRHIParamRef Texture, uint32 FaceIndex, uint32 ArrayIndex, uint32 MipIndex, EResourceLockMode LockMode, uint32& DestStride, bool bLockWithinMiptail) final override;
+	virtual void RHIUnlockTextureCubeFace(FTextureCubeRHIParamRef Texture, uint32 FaceIndex, uint32 ArrayIndex, uint32 MipIndex, bool bLockWithinMiptail) final override;
+	virtual void RHIBindDebugLabelName(FTextureRHIParamRef Texture, const TCHAR* Name) final override;
+	virtual void RHIReadSurfaceData(FTextureRHIParamRef Texture,FIntRect Rect,TArray<FColor>& OutData,FReadSurfaceDataFlags InFlags) final override;
+	virtual void RHIMapStagingSurface(FTextureRHIParamRef Texture,void*& OutData,int32& OutWidth,int32& OutHeight) final override;
+	virtual void RHIUnmapStagingSurface(FTextureRHIParamRef Texture) final override;
+	virtual void RHIReadSurfaceFloatData(FTextureRHIParamRef Texture,FIntRect Rect,TArray<FFloat16Color>& OutData,ECubeFace CubeFace,int32 ArrayIndex,int32 MipIndex) final override;
+	virtual void RHIRead3DSurfaceFloatData(FTextureRHIParamRef Texture,FIntRect Rect,FIntPoint ZMinMax,TArray<FFloat16Color>& OutData) final override;
+	virtual FRenderQueryRHIRef RHICreateRenderQuery(ERenderQueryType QueryType) final override;
+	virtual bool RHIGetRenderQueryResult(FRenderQueryRHIParamRef RenderQuery, uint64& OutResult, bool bWait) final override;
+	virtual FTexture2DRHIRef RHIGetViewportBackBuffer(FViewportRHIParamRef Viewport) final override;
+	virtual void RHIAdvanceFrameForGetViewportBackBuffer() final override;
+	virtual bool RHIIsDrawingViewport() final override;
+	virtual void RHIAcquireThreadOwnership() final override;
+	virtual void RHIReleaseThreadOwnership() final override;
+	virtual void RHIFlushResources() final override;
+	virtual uint32 RHIGetGPUFrameCycles() final override;
+	virtual FViewportRHIRef RHICreateViewport(void* WindowHandle, uint32 SizeX, uint32 SizeY, bool bIsFullscreen, EPixelFormat PreferredPixelFormat) final override;
+	virtual void RHIResizeViewport(FViewportRHIParamRef Viewport, uint32 SizeX, uint32 SizeY, bool bIsFullscreen) final override;
+	virtual void RHITick(float DeltaTime) final override;
+	virtual void RHISetStreamOutTargets(uint32 NumTargets,const FVertexBufferRHIParamRef* VertexBuffers,const uint32* Offsets) final override;
+	virtual void RHIDiscardRenderTargets(bool Depth,bool Stencil,uint32 ColorBitMask) final override;
+	virtual void RHIBlockUntilGPUIdle() final override;
+	virtual void RHISuspendRendering() final override;
+	virtual void RHIResumeRendering() final override;
+	virtual bool RHIIsRenderingSuspended() final override;
+	virtual bool RHIGetAvailableResolutions(FScreenResolutionArray& Resolutions, bool bIgnoreRefreshRate) final override;
+	virtual void RHIGetSupportedResolution(uint32& Width, uint32& Height) final override;
+	virtual void RHIVirtualTextureSetFirstMipInMemory(FTexture2DRHIParamRef Texture, uint32 FirstMip) final override;
+	virtual void RHIVirtualTextureSetFirstMipVisible(FTexture2DRHIParamRef Texture, uint32 FirstMip) final override;
+	virtual void RHIExecuteCommandList(FRHICommandList* CmdList) final override;
+	virtual void* RHIGetNativeDevice() final override;
+	virtual class IRHICommandContext* RHIGetDefaultContext() final override;
+	virtual class IRHICommandContextContainer* RHIGetCommandContextContainer() final override;
+
+#if UE_BUILD_DEBUG	
+	uint32 SubmissionLockStalls;
+	uint32 DrawCount;
+	uint64 PresentCount;
+#endif
+
+	inline IDXGIFactory4* GetFactory() const
+	{
+		return DXGIFactory;
+	}
+
+    inline void InitD3DDevices()
+    {
+		MainDevice->InitD3DDevice();
+        PerRHISetup(MainDevice);
+    }
+
+	/** Determine if an two views intersect */
+	template <class LeftT, class RightT>
+	static inline bool ResourceViewsIntersect(FD3D12View<LeftT>* pLeftView, FD3D12View<RightT>* pRightView)
+	{
+		if (pLeftView == nullptr || pRightView == nullptr)
+		{
+			// Cannot intersect if at least one is null
+			return false;
+		}
+
+		if ((void*)pLeftView == (void*)pRightView)
+		{
+			// Cannot intersect with itself
+			return false;
+		}
+
+		FD3D12Resource* pRTVResource = pLeftView->GetResource();
+		FD3D12Resource* pSRVResource = pRightView->GetResource();
+		if (pRTVResource != pSRVResource)
+		{
+			// Not the same resource
+			return false;
+		}
+
+		// Same resource, so see if their subresources overlap
+		return !pLeftView->DoesNotOverlap(*pRightView);
+	}
+
+	static inline bool IsTransitionNeeded(uint32 before, uint32 after)
+	{
+		check(before != D3D12_RESOURCE_STATE_CORRUPT && after != D3D12_RESOURCE_STATE_CORRUPT);
+		check(before != D3D12_RESOURCE_STATE_TBD && after != D3D12_RESOURCE_STATE_TBD);
+
+		// If 'after' is a subset of 'before', then there's no need for a transition
+		// Note: COMMON is an oddball state that doesn't follow the RESOURE_STATE pattern of 
+		// having exactly one bit set so we need to special case these
+		return before != after &&
+			((before | after) != before ||
+			after == D3D12_RESOURCE_STATE_COMMON);
+	}
+
+	/** Transition a resource's state based on a Render target view */
+	static inline void TransitionResource(FD3D12CommandListHandle& hCommandList, FD3D12RenderTargetView* pView, D3D12_RESOURCE_STATES after)
+	{
+		FD3D12Resource* pResource = pView->GetResource();
+
+		const D3D12_RENDER_TARGET_VIEW_DESC &desc = pView->GetDesc();
+		switch (desc.ViewDimension)
+		{
+		case D3D12_RTV_DIMENSION_TEXTURE3D:
+			// Note: For volume (3D) textures, all slices for a given mipmap level are a single subresource index.
+			// Fall-through
+		case D3D12_RTV_DIMENSION_TEXTURE2D:
+		case D3D12_RTV_DIMENSION_TEXTURE2DMS:
+			// Only one subresource to transition
+			TransitionResource(hCommandList, pResource, after, desc.Texture2D.MipSlice);
+			break;
+
+		case D3D12_RTV_DIMENSION_TEXTURE2DARRAY:
+		{
+			// Multiple subresources to transition
+			TransitionResource(hCommandList, pResource, after, pView->GetViewSubresourceSubset());
+			break;
+		}
+
+		default:
+			check(false);	// Need to update this code to include the view type
+			break;
+		}
+	}
+	
+	/** Transition a resource's state based on a Depth stencil view's desc flags */
+	static inline void TransitionResource(FD3D12CommandListHandle& hCommandList, FD3D12DepthStencilView* pView)
+	{
+		// Determine the required subresource states from the view desc
+		const D3D12_DEPTH_STENCIL_VIEW_DESC &desc = pView->GetDesc();
+		const D3D12_RESOURCE_STATES depthState =	(desc.Flags & D3D12_DSV_FLAG_READ_ONLY_DEPTH)	? D3D12_RESOURCE_STATE_DEPTH_READ : D3D12_RESOURCE_STATE_DEPTH_WRITE;
+		const D3D12_RESOURCE_STATES stencilState =	(desc.Flags & D3D12_DSV_FLAG_READ_ONLY_STENCIL)	? D3D12_RESOURCE_STATE_DEPTH_READ : D3D12_RESOURCE_STATE_DEPTH_WRITE;
+		const bool bSameState = depthState == stencilState;
+
+		FD3D12Resource* pResource = pView->GetResource();
+		const bool bHasDepth = pView->HasDepth();
+		const bool bHasStencil = pView->HasStencil();
+
+		// This code assumes that the DSV always contains the depth plane
+		check(bHasDepth);
+		if (!bHasStencil)
+		{
+			// Only one plane, depth, to transition
+			TransitionResource(hCommandList, pResource, depthState, pView->GetViewSubresourceSubset());
+			return;
+		}
+		else if (bSameState)
+		{
+			// Transition multiple planes to the same state
+			check(bHasStencil);
+			TransitionResource(hCommandList, pResource, depthState, pView->GetViewSubresourceSubset());
+		}
+		else
+		{
+			// Transition each plane separately because they need to be in different states
+			check(bHasStencil && !bSameState);
+			TransitionResource(hCommandList, pResource, depthState, pView->GetDepthOnlyViewSubresourceSubset());
+			TransitionResource(hCommandList, pResource, stencilState, pView->GetStencilOnlyViewSubresourceSubset());
+		}
+	}
+
+	/** Transition a resource's state based on a Depth stencil view */
+	static inline void TransitionResource(FD3D12CommandListHandle& hCommandList, FD3D12DepthStencilView* pView, D3D12_RESOURCE_STATES after)
+	{
+		FD3D12Resource* pResource = pView->GetResource();
+
+		const D3D12_DEPTH_STENCIL_VIEW_DESC &desc = pView->GetDesc();
+		switch (desc.ViewDimension)
+		{
+		case D3D12_DSV_DIMENSION_TEXTURE2D:
+		case D3D12_DSV_DIMENSION_TEXTURE2DMS:
+			if (pResource->GetPlaneCount() > 1)
+			{
+				// Multiple subresources to transtion
+				TransitionResource(hCommandList, pResource, after, pView->GetViewSubresourceSubset());
+				break;
+			}
+			else
+			{
+				// Only one subresource to transition
+				check(pResource->GetPlaneCount() == 1);
+				TransitionResource(hCommandList, pResource, after, desc.Texture2D.MipSlice);
+			}
+			break;
+
+		case D3D12_DSV_DIMENSION_TEXTURE2DARRAY:
+		case D3D12_DSV_DIMENSION_TEXTURE2DMSARRAY:
+		{
+			// Multiple subresources to transtion
+			TransitionResource(hCommandList, pResource, after, pView->GetViewSubresourceSubset());
+			break;
+		}
+
+		default:
+			check(false);	// Need to update this code to include the view type
+			break;
+		}
+	}
+
+	/** Transition a resource's state based on a Unordered access view */
+	static inline void TransitionResource(FD3D12CommandListHandle& hCommandList, FD3D12UnorderedAccessView* pView, D3D12_RESOURCE_STATES after)
+	{
+		FD3D12Resource* pResource = pView->GetResource();
+
+		const D3D12_UNORDERED_ACCESS_VIEW_DESC &desc = pView->GetDesc();
+		switch (desc.ViewDimension)
+		{
+		case D3D12_UAV_DIMENSION_BUFFER:
+			TransitionResource(hCommandList, pResource, after, 0);
+			break;
+
+		case D3D12_UAV_DIMENSION_TEXTURE2D:
+			// Only one subresource to transition
+			TransitionResource(hCommandList, pResource, after, desc.Texture2D.MipSlice);
+			break;
+
+		case D3D12_UAV_DIMENSION_TEXTURE2DARRAY:
+		{
+			// Multiple subresources to transtion
+			TransitionResource(hCommandList, pResource, after, pView->GetViewSubresourceSubset());
+			break;
+		}
+		case D3D12_UAV_DIMENSION_TEXTURE3D:
+		{
+			// Multiple subresources to transtion
+			TransitionResource(hCommandList, pResource, after, pView->GetViewSubresourceSubset());
+			break;
+		}
+
+		default:
+			check(false);	// Need to update this code to include the view type
+			break;
+		}
+	}
+
+	/** Transition a resource's state based on a Shader resource view */
+	static inline void TransitionResource(FD3D12CommandListHandle& hCommandList, FD3D12ShaderResourceView* pSRView, D3D12_RESOURCE_STATES after)
+	{
+		FD3D12Resource* pResource = pSRView->GetResource();
+
+		if (!pResource || !pResource->RequiresResourceStateTracking())
+		{
+			// Early out if we never need to do state tracking, the resource should always be in an SRV state.
+			return;
+		}
+
+		const D3D12_RESOURCE_DESC &ResDesc = pResource->GetDesc();
+		const CViewSubresourceSubset &subresourceSubset = pSRView->GetViewSubresourceSubset();
+
+		const D3D12_SHADER_RESOURCE_VIEW_DESC &desc = pSRView->GetDesc();
+		switch (desc.ViewDimension)
+		{
+		default:
+		{
+			// Transition the resource
+			TransitionResource(hCommandList, pResource, after, subresourceSubset);
+			break;
+		}
+
+		case D3D12_SRV_DIMENSION_BUFFER:
+		{
+			if (pResource->GetHeapType() == D3D12_HEAP_TYPE_DEFAULT)
+			{
+				// Transition the resource
+				TransitionResource(hCommandList, pResource, after, subresourceSubset);
+			}
+			break;
+		}
+		}
+	}
+
+	// Transition a specific subresource to the after state.
+	static inline void TransitionResource(FD3D12CommandListHandle& hCommandList, FD3D12Resource* pResource, D3D12_RESOURCE_STATES after, uint32 subresource)
+	{
+		TransitionResourceWithTracking(hCommandList, pResource, after, subresource);
+	}
+
+	// Transition a subset of subresources to the after state.
+	static inline void TransitionResource(FD3D12CommandListHandle& hCommandList, FD3D12Resource* pResource, D3D12_RESOURCE_STATES after, const CViewSubresourceSubset& subresourceSubset)
+	{
+		TransitionResourceWithTracking(hCommandList, pResource, after, subresourceSubset);
+	}
+
+	// Transition subresources from one state to another
+	static inline void ResourceBarrier(FD3D12CommandListHandle& hCommandList, FD3D12Resource* pResource,D3D12_RESOURCE_BARRIER* BarrierDescs, uint32 numBarriers)
+	{
+#if SUPPORTS_MEMORY_RESIDENCY
+		pResource->UpdateResidency();
+#endif
+
+		check(numBarriers > 0);
+
+#if DEBUG_RESOURCE_STATES
+		if(pResource->RequiresResourceStateTracking())
+		{
+			LogResourceBarriers(numBarriers, BarrierDescs, hCommandList.CommandList());
+		}
+#endif
+
+		hCommandList.GetCurrentOwningContext()->numBarriers += numBarriers;
+		hCommandList->ResourceBarrier(numBarriers, BarrierDescs);
+	};
+
+	// Transition a subresource from one state to another
+	static inline void ResourceBarrier(FD3D12CommandListHandle& hCommandList, FD3D12Resource* pResource, D3D12_RESOURCE_STATES before, D3D12_RESOURCE_STATES after, uint32 subresource)
+	{
+#if SUPPORTS_MEMORY_RESIDENCY
+		pResource->UpdateResidency();
+#endif
+
+		check(IsTransitionNeeded(before, after));
+
+		// Fill out the desc
+		D3D12_RESOURCE_BARRIER desc = {};
+		desc.Type = D3D12_RESOURCE_BARRIER_TYPE_TRANSITION;
+		desc.Transition.pResource = pResource->GetResource();
+		desc.Transition.StateBefore = before;
+		desc.Transition.StateAfter = after;
+		desc.Transition.Subresource = subresource;
+
+#if DEBUG_RESOURCE_STATES
+		if (pResource->RequiresResourceStateTracking())
+		{
+			LogResourceBarriers(1, &desc, hCommandList.CommandList());
+		}
+#endif
+				
+		hCommandList.GetCurrentOwningContext()->numBarriers++;
+		hCommandList->ResourceBarrier(1, &desc);
+	};
+
+	// Transition a subresource from current to a new state, using resource state tracking.
+	static void TransitionResourceWithTracking(FD3D12CommandListHandle& hCommandList, FD3D12Resource* pResource, D3D12_RESOURCE_STATES after, uint32 subresource)
+	{
+		check(pResource->RequiresResourceStateTracking());
+
+		CResourceState& ResourceState = hCommandList.GetResourceState(pResource);
+		if (subresource == D3D12_RESOURCE_BARRIER_ALL_SUBRESOURCES && !ResourceState.AreAllSubresourcesSame())
+		{
+			// Slow path. Want to transition the entire resource (with multiple subresources). But they aren't in the same state.
+
+			// Partially fill out the desc
+			D3D12_RESOURCE_BARRIER desc = {};
+			desc.Type = D3D12_RESOURCE_BARRIER_TYPE_TRANSITION;
+			desc.Transition.pResource = pResource->GetResource();
+			desc.Transition.StateAfter = after;
+
+#if SUPPORTS_MEMORY_RESIDENCY
+			pResource->UpdateResidency();
+#endif
+			
+			D3D12_RESOURCE_BARRIER* descs = hCommandList.GetResourceBarrierScratchSpace();
+			uint32 numBarriersNeeded = 0;
+			const uint8 SubresourceCount = pResource->GetSubresourceCount();
+			for (uint32 SubresourceIndex = 0; SubresourceIndex < SubresourceCount; SubresourceIndex++)
+			{
+				const D3D12_RESOURCE_STATES before = ResourceState.GetSubresourceState(SubresourceIndex);
+				if (before == D3D12_RESOURCE_STATE_TBD)
+				{
+					// We need a pending resource barrier so we can setup the state before this command list executes
+					hCommandList.AddPendingResourceBarrier(pResource, after, SubresourceIndex);
+				}
+				// We're not using IsTransitionNeeded() because we do want to transition even if 'after' is a subset of 'before'
+				// This is so that we can ensure all subresources are in the same state, simplifying future barriers
+				else if (before != after)
+				{
+					desc.Transition.StateBefore = before;
+					desc.Transition.Subresource = SubresourceIndex;
+
+					check(numBarriersNeeded < FD3D12CommandListHandle::ResourceBarrierScratchSpaceSize);
+					descs[numBarriersNeeded] = desc;
+					numBarriersNeeded++;
+				}
+			}
+
+			if (numBarriersNeeded > 0)
+			{
+				ResourceBarrier(hCommandList, pResource, descs, numBarriersNeeded);
+			}
+
+			// The entire resource should now be in the after state on this command list (even if all barriers are pending)
+			ResourceState.SetResourceState(after);
+		}
+		else
+		{
+			const D3D12_RESOURCE_STATES before = ResourceState.GetSubresourceState(subresource);
+			if (before == D3D12_RESOURCE_STATE_TBD)
+			{
+#if SUPPORTS_MEMORY_RESIDENCY
+				pResource->UpdateResidency();
+#endif
+
+				// We need a pending resource barrier so we can setup the state before this command list executes
+				hCommandList.AddPendingResourceBarrier(pResource, after, subresource);
+				ResourceState.SetSubresourceState(subresource, after);
+			}
+			else if (IsTransitionNeeded(before, after))
+			{
+				ResourceBarrier(hCommandList, pResource, before, after, subresource);
+				ResourceState.SetSubresourceState(subresource, after);
+			}
+		}
+	}
+
+	// Transition subresources from current to a new state, using resource state tracking.
+	static void TransitionResourceWithTracking(FD3D12CommandListHandle& hCommandList, FD3D12Resource* pResource, D3D12_RESOURCE_STATES after, const CViewSubresourceSubset& subresourceSubset)
+	{
+		check(pResource->RequiresResourceStateTracking());
+
+		D3D12_RESOURCE_STATES before;
+		const bool bIsWholeResource = subresourceSubset.IsWholeResource();
+		CResourceState& ResourceState = hCommandList.GetResourceState(pResource);
+		if (bIsWholeResource && ResourceState.AreAllSubresourcesSame())
+		{
+			// Fast path. Transition the entire resource from one state to another.
+			before = ResourceState.GetSubresourceState(D3D12_RESOURCE_BARRIER_ALL_SUBRESOURCES);
+			if (before == D3D12_RESOURCE_STATE_TBD)
+			{
+#if SUPPORTS_MEMORY_RESIDENCY
+				pResource->UpdateResidency();
+#endif
+
+				// We need a pending resource barrier so we can setup the state before this command list executes
+				hCommandList.AddPendingResourceBarrier(pResource, after, D3D12_RESOURCE_BARRIER_ALL_SUBRESOURCES);
+				ResourceState.SetResourceState(after);
+			}
+			else if (IsTransitionNeeded(before, after))
+			{
+				ResourceBarrier(hCommandList, pResource, before, after, D3D12_RESOURCE_BARRIER_ALL_SUBRESOURCES);
+				ResourceState.SetResourceState(after);
+			}
+		}
+		else
+		{
+			// Slower path. Either the subresources are in more than 1 state, or the view only partially covers the resource.
+			// Either way, we'll need to loop over each subresource in the view...
+
+			// Partially fill out the desc
+			D3D12_RESOURCE_BARRIER desc = {};
+			desc.Type = D3D12_RESOURCE_BARRIER_TYPE_TRANSITION;
+			desc.Transition.pResource = pResource->GetResource();
+			desc.Transition.StateAfter = after;
+
+#if SUPPORTS_MEMORY_RESIDENCY
+			pResource->UpdateResidency();
+#endif
+			D3D12_RESOURCE_BARRIER* descs = hCommandList.GetResourceBarrierScratchSpace();
+			uint32 numBarriersNeeded = 0;
+			for (CViewSubresourceSubset::CViewSubresourceIterator it = subresourceSubset.begin(); it != subresourceSubset.end(); ++it)
+			{
+				for (uint32 SubresourceIndex = it.StartSubresource(); SubresourceIndex < it.EndSubresource(); SubresourceIndex++)
+				{
+					before = ResourceState.GetSubresourceState(SubresourceIndex);
+					if (before == D3D12_RESOURCE_STATE_TBD)
+					{
+						// We need a pending resource barrier so we can setup the state before this command list executes
+						hCommandList.AddPendingResourceBarrier(pResource, after, SubresourceIndex);
+						ResourceState.SetSubresourceState(SubresourceIndex, after);
+					}
+					else if (IsTransitionNeeded(before, after))
+					{
+						// Update the desc for the current subresource and add it to the array
+						desc.Transition.StateBefore = before;
+						desc.Transition.Subresource = SubresourceIndex;
+
+						check(numBarriersNeeded < FD3D12CommandListHandle::ResourceBarrierScratchSpaceSize);
+						descs[numBarriersNeeded] = desc;
+						numBarriersNeeded++;
+
+						ResourceState.SetSubresourceState(SubresourceIndex, after);
+					}
+				}
+			}
+
+			if (numBarriersNeeded > 0)
+			{
+				ResourceBarrier(hCommandList, pResource, descs, numBarriersNeeded);
+			}
+
+			// If we just transtioned every subresource to the same state, lets update it's tracking so it's on a per-resource level
+			if (bIsWholeResource)
+			{
+				ResourceState.SetResourceState(after);
+			}
+		}
+	}
+
+	void GetLocalVideoMemoryInfo(DXGI_QUERY_VIDEO_MEMORY_INFO* LocalVideoMemoryInfo);
+
+	static HRESULT CreateCommittedResource(const D3D12_RESOURCE_DESC& Desc, const D3D12_HEAP_PROPERTIES& HeapProps, const D3D12_RESOURCE_STATES& InitialUsage, FD3D12Resource** ppResource)
+	{
+		HRESULT hresult = S_OK;
+		D3D12_RESOURCE_HEAP_TIER ResourceHeapTier = SingleD3DRHI->GetRHIDevice()->GetResourceHeapTier();
+		TRefCountPtr<ID3D12Resource> pResource;
+
+		ID3D12Device* pD3DDevice = SingleD3DRHI->GetRHIDevice()->GetDevice();
+
+		if (!ppResource)
+		{
+			return E_POINTER;
+		}
+
+		// If Placed Resources are supported, back the resource with a heap and create a placed resource on top of the heap.
+		// This is equivalent to what the Runtime would do for CreateCommitedResource but we also get the added benefit of 
+		// being able to "recycle" these resources via aliased placed resources.
+#if ENABLE_PLACED_RESOURCES
+		TRefCountPtr<ID3D12Heap> pHeap;
+		D3D12_RESOURCE_ALLOCATION_INFO AllocationInfo = pD3DDevice->GetResourceAllocationInfo(0, 1, &Desc);
+		D3D12_HEAP_DESC HeapDesc = {};
+		HeapDesc.SizeInBytes = AllocationInfo.SizeInBytes;
+		HeapDesc.Alignment = AllocationInfo.Alignment;
+		HeapDesc.Properties = HeapProps;
+		if (HeapProps.Type == D3D12_HEAP_TYPE_READBACK || HeapProps.Type == D3D12_HEAP_TYPE_UPLOAD)
+		{
+			HeapDesc.Flags = D3D12_HEAP_FLAG_ALLOW_ONLY_BUFFERS;
+		}
+		// Tier 1 requires heaps to be much more restrictive
+		else if (ResourceHeapTier == D3D12_RESOURCE_HEAP_TIER_1)
+		{
+			if (Desc.Dimension == D3D12_RESOURCE_DIMENSION_BUFFER)
+			{
+				HeapDesc.Flags = D3D12_HEAP_FLAG_ALLOW_ONLY_BUFFERS;
+			}
+			else if (Desc.Flags & (D3D12_RESOURCE_FLAG_ALLOW_RENDER_TARGET | D3D12_RESOURCE_FLAG_ALLOW_DEPTH_STENCIL))
+			{
+				HeapDesc.Flags = D3D12_HEAP_FLAG_ALLOW_ONLY_RT_DS_TEXTURES;
+			}
+			else
+			{
+				HeapDesc.Flags = D3D12_HEAP_FLAG_ALLOW_ONLY_NON_RT_DS_TEXTURES;
+			}
+		}
+
+		hresult = pD3DDevice->CreateHeap(
+			&HeapDesc,
+			IID_PPV_ARGS(pHeap.GetInitReference()));
+
+		VERIFYD3D11RESULT(hresult);
+
+		hresult = pD3DDevice->CreatePlacedResource(pHeap, 0, &Desc, InitialUsage, nullptr, IID_PPV_ARGS(pResource.GetInitReference()));
+		VERIFYD3D11RESULT(hresult);
+
+		// Set the output pointer
+		*ppResource = new FD3D12Resource(pResource, Desc, pHeap, HeapProps.Type);
+#else
+		hresult = pD3DDevice->CreateCommittedResource(&HeapProps, D3D12_HEAP_FLAG_NONE, &Desc, InitialUsage, nullptr, IID_PPV_ARGS(pResource.GetInitReference()));
+#if SUPPORTS_MEMORY_RESIDENCY
+		if (hresult == E_OUTOFMEMORY)
+		{
+			for (uint32 MemoryPressureLevel = 1; MemoryPressureLevel < FD3D12ResourceResidencyManager::MEMORY_PRESSURE_LEVELS; ++MemoryPressureLevel)
+			{
+				SingleD3DRHI->GetResourceResidencyManager().Process(MemoryPressureLevel);
+
+				ID3D12Resource* pD3DResource = pResource.GetReference();
+				hresult = pD3DDevice->CreateCommittedResource(&HeapProps, D3D12_HEAP_FLAG_NONE, &Desc, InitialUsage, nullptr, IID_PPV_ARGS(&pD3DResource));
+
+				if (hresult != E_OUTOFMEMORY)
+					break;
+			}
+		}
+#endif
+		VERIFYD3D11RESULT(hresult);
+
+		// Set the output pointer
+		*ppResource = new FD3D12Resource(pResource, Desc, nullptr, HeapProps.Type);
+#endif
+		(*ppResource)->AddRef();
+
+		return S_OK;
+	}
+
+#if SUPPORTS_MEMORY_RESIDENCY
+	FD3D12ResourceResidencyManager& GetResourceResidencyManager() { return ResourceResidencyManager; }
+
+private:
+
+	FD3D12ResourceResidencyManager ResourceResidencyManager;
+#endif
+
+public:
+
+	inline FD3D12DynamicHeapAllocator& GetHelperThreadDynamicUploadHeapAllocator()
+	{
+		check(!IsInActualRenderingThread());
+
+		if (HelperThreadDynamicHeapAllocator == nullptr)
+		{
+			uint32 NextIndex = InterlockedIncrement(&NumThreadDynamicHeapAllocators) - 1;
+			check(NextIndex < _countof(ThreadDynamicHeapAllocatorArray));
+			HelperThreadDynamicHeapAllocator = new FD3D12DynamicHeapAllocator(GetRHIDevice(), D3D12_HEAP_TYPE_UPLOAD);
+			ThreadDynamicHeapAllocatorArray[NextIndex] = HelperThreadDynamicHeapAllocator;
+		}
+
+		return *HelperThreadDynamicHeapAllocator;
+	}
+	FD3D12DynamicHeapAllocator* ThreadDynamicHeapAllocatorArray[8];
+	uint32 NumThreadDynamicHeapAllocators;
+	static __declspec(thread) FD3D12DynamicHeapAllocator* HelperThreadDynamicHeapAllocator;
+
+	static DXGI_FORMAT GetPlatformTextureResourceFormat(DXGI_FORMAT InFormat, uint32 InFlags);
+
+#if	PLATFORM_SUPPORTS_VIRTUAL_TEXTURES
+	virtual void* CreateVirtualTexture(uint32 SizeX, uint32 SizeY, uint32 SizeZ, uint32 NumMips, bool bCubeTexture, uint32 Flags, void* D3DTextureDesc, void* D3DTextureResource) = 0;
+	virtual void DestroyVirtualTexture(uint32 Flags, void* RawTextureMemory) = 0;
+	virtual bool HandleSpecialLock(FD3D12LockedData& LockedData, uint32 MipIndex, uint32 ArrayIndex, uint32 Flags, EResourceLockMode LockMode,
+		void* D3DTextureResource, void* RawTextureMemory, uint32 NumMips, uint32& DestStride) = 0;
+	virtual bool HandleSpecialUnlock(uint32 MipIndex, uint32 Flags, void* D3DTextureResource, void* RawTextureMemory) = 0;
+#endif
+
+    inline void RegisterGPUWork(uint32 NumPrimitives = 0, uint32 NumVertices = 0)
+    {
+        GPUProfilingData.RegisterGPUWork(NumPrimitives, NumVertices);
+    }
+
+    inline void PushGPUEvent(const TCHAR* Name)
+    {
+        GPUProfilingData.PushEvent(Name);
+    }
+
+    inline void PopGPUEvent()
+    {
+        GPUProfilingData.PopEvent();
+    }
+
+    inline void IncrementSetShaderTextureCalls(uint32 value = 1)
+    {
+        SetShaderTextureCalls += value;
+    }
+
+    inline void IncrementSetShaderTextureCycles(uint32 value)
+    {
+        SetShaderTextureCycles += value;
+    }
+
+    inline void IncrementCacheResourceTableCalls(uint32 value = 1)
+    {
+        CacheResourceTableCalls += value;
+    }
+
+    inline void IncrementCacheResourceTableCycles(uint32 value)
+    {
+        CacheResourceTableCycles += value;
+    }
+
+    inline void IncrementCommitComputeResourceTables(uint32 value = 1)
+    {
+        CommitResourceTableCycles += value;
+    }
+
+    inline void AddBoundShaderState(FD3D12BoundShaderState* state)
+    {
+        BoundShaderStateHistory.Add(state);
+    }
+
+    inline void IncrementSetTextureInTableCalls(uint32 value = 1)
+    {
+        SetTextureInTableCalls += value;
+    }
+
+    inline uint32 GetResourceTableFrameCounter() { return ResourceTableFrameCounter; }
+
+    /** Consumes about 100ms of GPU time (depending on resolution and GPU), useful for making sure we're not CPU bound when GPU profiling. */
+    void IssueLongGPUTask();
+
+	void ResetViewportFrameCounter() { ViewportFrameCounter = 0; }
+
+private:
+    inline void PerRHISetup(FD3D12Device* MainDevice);
+
+	/** The global D3D interface. */
+	TRefCountPtr<IDXGIFactory4> DXGIFactory;
+
+    FD3D12Device* MainDevice;
+    FD3D12Adapter MainAdapter;
+
+    /** The feature level of the device. */
+    D3D_FEATURE_LEVEL FeatureLevel;
+
+	/** A buffer in system memory containing all zeroes of the specified size. */
+	void* ZeroBuffer;
+	uint32 ZeroBufferSize;
+
+	uint32 CommitResourceTableCycles;
+	uint32 CacheResourceTableCalls;
+	uint32 CacheResourceTableCycles;
+	uint32 SetShaderTextureCycles;
+	uint32 SetShaderTextureCalls;
+	uint32 SetTextureInTableCalls;
+
+	uint32 ViewportFrameCounter;
+
+	/** Internal frame counter, incremented on each call to RHIBeginScene. */
+	uint32 SceneFrameCounter;
+
+	/**
+	* Internal counter used for resource table caching.
+	* INDEX_NONE means caching is not allowed.
+	*/
+	uint32 ResourceTableFrameCounter;
+
+	/** A history of the most recently used bound shader states, used to keep transient bound shader states from being recreated for each use. */
+	TGlobalResource< TBoundShaderStateHistory<10000> > BoundShaderStateHistory;
+
+	FD3DGPUProfiler GPUProfilingData;
+
+    /** Disables multi-queue support for the engine. */
+    bool bForceSingleQueueGPU;
+
+	template<typename BaseResourceType>
+	TD3D12Texture2D<BaseResourceType>* CreateD3D11Texture2D(uint32 SizeX, uint32 SizeY, uint32 SizeZ, bool bTextureArray, bool CubeTexture, uint8 Format,
+		uint32 NumMips, uint32 NumSamples, uint32 Flags, FRHIResourceCreateInfo& CreateInfo);
+
+	FD3D12Texture3D* CreateD3D11Texture3D(uint32 SizeX, uint32 SizeY, uint32 SizeZ, uint8 Format, uint32 NumMips, uint32 Flags, FRHIResourceCreateInfo& CreateInfo);
+
+	/**
+	 * Gets the best supported MSAA settings from the provided MSAA count to check against.
+	 *
+	 * @param PlatformFormat		The format of the texture being created
+	 * @param MSAACount				The MSAA count to check against.
+	 * @param OutBestMSAACount		The best MSAA count that is suppored.  Could be smaller than MSAACount if it is not supported
+	 * @param OutMSAAQualityLevels	The number MSAA quality levels for the best msaa count supported
+	 */
+	void GetBestSupportedMSAASetting(DXGI_FORMAT PlatformFormat, uint32 MSAACount, uint32& OutBestMSAACount, uint32& OutMSAAQualityLevels);
+
+
+	// @return 0xffffffff if not not supported
+	uint32 GetMaxMSAAQuality(uint32 SampleCount);
+
+	/**
+	* Returns a pointer to a texture resource that can be used for CPU reads.
+	* Note: the returned resource could be the original texture or a new temporary texture.
+	* @param TextureRHI - Source texture to create a staging texture from.
+	* @param InRect - rectangle to 'stage'.
+	* @param StagingRectOUT - parameter is filled with the rectangle to read from the returned texture.
+	* @return The CPU readable Texture object.
+	*/
+	TRefCountPtr<FD3D12Resource> GetStagingTexture(FTextureRHIParamRef TextureRHI, FIntRect InRect, FIntRect& OutRect, FReadSurfaceDataFlags InFlags, D3D12_PLACED_SUBRESOURCE_FOOTPRINT &readBackHeapDesc);
+
+	void ReadSurfaceDataNoMSAARaw(FTextureRHIParamRef TextureRHI, FIntRect Rect, TArray<uint8>& OutData, FReadSurfaceDataFlags InFlags);
+
+	void ReadSurfaceDataMSAARaw(FRHICommandList_RecursiveHazardous& RHICmdList, FTextureRHIParamRef TextureRHI, FIntRect Rect, TArray<uint8>& OutData, FReadSurfaceDataFlags InFlags);
+
+	void SetupRecursiveResources();
+
+    // This should only be called by Dynamic RHI member functions
+    inline FD3D12Device* GetRHIDevice() const
+    {
+		return MainDevice;
+    }
+};
+
+inline bool FD3D12CommandContext::IsDefaultContext() const
+{
+	return this == &GetParentDevice()->GetDefaultCommandContext();
+}
+
+/** Implements the D3D12RHI module as a dynamic RHI providing module. */
+class FD3D12DynamicRHIModule : public IDynamicRHIModule
+{
+public:
+	// IModuleInterface
+	virtual bool SupportsDynamicReloading() override { return false; }
+
+	// IDynamicRHIModule
+	virtual bool IsSupported() override;
+	virtual FDynamicRHI* CreateRHI() override;
+
+private:
+	FD3D12Adapter ChosenAdapter;
+
+	// set MaxSupportedFeatureLevel and ChosenAdapter
+	void FindAdapter();
+};
+
+/** Find an appropriate DXGI format for the input format and SRGB setting. */
+inline DXGI_FORMAT FindShaderResourceDXGIFormat(DXGI_FORMAT InFormat, bool bSRGB)
+{
+	if (bSRGB)
+	{
+		switch (InFormat)
+		{
+		case DXGI_FORMAT_B8G8R8A8_TYPELESS:    return DXGI_FORMAT_B8G8R8A8_UNORM_SRGB;
+		case DXGI_FORMAT_R8G8B8A8_TYPELESS:    return DXGI_FORMAT_R8G8B8A8_UNORM_SRGB;
+		case DXGI_FORMAT_BC1_TYPELESS:         return DXGI_FORMAT_BC1_UNORM_SRGB;
+		case DXGI_FORMAT_BC2_TYPELESS:         return DXGI_FORMAT_BC2_UNORM_SRGB;
+		case DXGI_FORMAT_BC3_TYPELESS:         return DXGI_FORMAT_BC3_UNORM_SRGB;
+		case DXGI_FORMAT_BC7_TYPELESS:         return DXGI_FORMAT_BC7_UNORM_SRGB;
+		};
+	}
+	else
+	{
+		switch (InFormat)
+		{
+		case DXGI_FORMAT_B8G8R8A8_TYPELESS: return DXGI_FORMAT_B8G8R8A8_UNORM;
+		case DXGI_FORMAT_R8G8B8A8_TYPELESS: return DXGI_FORMAT_R8G8B8A8_UNORM;
+		case DXGI_FORMAT_BC1_TYPELESS:      return DXGI_FORMAT_BC1_UNORM;
+		case DXGI_FORMAT_BC2_TYPELESS:      return DXGI_FORMAT_BC2_UNORM;
+		case DXGI_FORMAT_BC3_TYPELESS:      return DXGI_FORMAT_BC3_UNORM;
+		case DXGI_FORMAT_BC7_TYPELESS:      return DXGI_FORMAT_BC7_UNORM;
+		};
+	}
+	switch (InFormat)
+	{
+	case DXGI_FORMAT_R24G8_TYPELESS: return DXGI_FORMAT_R24_UNORM_X8_TYPELESS;
+	case DXGI_FORMAT_R32_TYPELESS: return DXGI_FORMAT_R32_FLOAT;
+	case DXGI_FORMAT_R16_TYPELESS: return DXGI_FORMAT_R16_UNORM;
+#if DEPTH_32_BIT_CONVERSION
+		// Changing Depth Buffers to 32 bit on Dingo as D24S8 is actually implemented as a 32 bit buffer in the hardware
+	case DXGI_FORMAT_R32G8X24_TYPELESS: return DXGI_FORMAT_R32_FLOAT_X8X24_TYPELESS;
+#endif
+	}
+	return InFormat;
+}
+
+/** Find an appropriate DXGI format unordered access of the raw format. */
+inline DXGI_FORMAT FindUnorderedAccessDXGIFormat(DXGI_FORMAT InFormat)
+{
+	switch (InFormat)
+	{
+	case DXGI_FORMAT_B8G8R8A8_TYPELESS: return DXGI_FORMAT_B8G8R8A8_UNORM;
+	case DXGI_FORMAT_R8G8B8A8_TYPELESS: return DXGI_FORMAT_R8G8B8A8_UNORM;
+	}
+	return InFormat;
+}
+
+/** Find the appropriate depth-stencil targetable DXGI format for the given format. */
+inline DXGI_FORMAT FindDepthStencilDXGIFormat(DXGI_FORMAT InFormat)
+{
+	switch (InFormat)
+	{
+	case DXGI_FORMAT_R24G8_TYPELESS:
+		return DXGI_FORMAT_D24_UNORM_S8_UINT;
+#if DEPTH_32_BIT_CONVERSION
+		// Changing Depth Buffers to 32 bit on Dingo as D24S8 is actually implemented as a 32 bit buffer in the hardware
+	case DXGI_FORMAT_R32G8X24_TYPELESS:
+		return DXGI_FORMAT_D32_FLOAT_S8X24_UINT;
+#endif
+	case DXGI_FORMAT_R32_TYPELESS:
+		return DXGI_FORMAT_D32_FLOAT;
+	case DXGI_FORMAT_R16_TYPELESS:
+		return DXGI_FORMAT_D16_UNORM;
+	};
+	return InFormat;
+}
+
+/**
+ * Returns whether the given format contains stencil information.
+ * Must be passed a format returned by FindDepthStencilDXGIFormat, so that typeless versions are converted to their corresponding depth stencil view format.
+ */
+inline bool HasStencilBits(DXGI_FORMAT InFormat)
+{
+	switch (InFormat)
+	{
+	case DXGI_FORMAT_D24_UNORM_S8_UINT:
+		return true;
+#if  DEPTH_32_BIT_CONVERSION
+		// Changing Depth Buffers to 32 bit on Dingo as D24S8 is actually implemented as a 32 bit buffer in the hardware
+	case DXGI_FORMAT_D32_FLOAT_S8X24_UINT:
+		return true;
+#endif
+	};
+	return false;
+}
+
+/** Vertex declaration for just one FVector4 position. */
+class FVector4VertexDeclaration : public FRenderResource
+{
+public:
+	FVertexDeclarationRHIRef VertexDeclarationRHI;
+	virtual void InitRHI() override
+	{
+		FVertexDeclarationElementList Elements;
+		Elements.Add(FVertexElement(0, 0, VET_Float4, 0, sizeof(FVector4)));
+		VertexDeclarationRHI = RHICreateVertexDeclaration(Elements);
+	}
+	virtual void ReleaseRHI() override
+	{
+		VertexDeclarationRHI.SafeRelease();
+	}
+};
+
+namespace D3D12RHI
+{
+	extern TGlobalResource<FVector4VertexDeclaration> GD3D11Vector4VertexDeclaration;
+
+	/**
+	 *	Default 'Fast VRAM' allocator...
+	 */
+	class FFastVRAMAllocator
+	{
+	public:
+		FFastVRAMAllocator() {}
+
+		virtual ~FFastVRAMAllocator() {}
+
+		/**
+		 *	IMPORTANT: This function CAN modify the TextureDesc!
+		 */
+		virtual FVRamAllocation AllocTexture2D(D3D12_RESOURCE_DESC& TextureDesc)
+		{
+			return FVRamAllocation();
+		}
+
+		/**
+		 *	IMPORTANT: This function CAN modify the TextureDesc!
+		 */
+		virtual FVRamAllocation AllocTexture3D(D3D12_RESOURCE_DESC& TextureDesc)
+		{
+			return FVRamAllocation();
+		}
+
+		/**
+		 *	IMPORTANT: This function CAN modify the BufferDesc!
+		 */
+		virtual FVRamAllocation AllocUAVBuffer(D3D12_RESOURCE_DESC& BufferDesc)
+		{
+			return FVRamAllocation();
+		}
+
+		template< typename t_A, typename t_B >
+		static t_A RoundUpToNextMultiple(const t_A& a, const t_B& b)
+		{
+			return ((a - 1) / b + 1) * b;
+		}
+
+		static FFastVRAMAllocator* GetFastVRAMAllocator();
+	};
+}
+
+// 1d, 31 bit (uses the sign bit for internal use), O(n) where n is the amount of elements stored
+// does not enforce any alignment
+// unoccupied regions get compacted but occupied don't get compacted
+class FRangeAllocator
+{
+public:
+
+	struct FRange
+	{
+		// not valid
+		FRange()
+			: Start(0)
+			, Size(0)
+		{
+			check(!IsValid());
+		}
+
+		void SetOccupied(int32 InStart, int32 InSize)
+		{
+			check(InStart >= 0);
+			check(InSize > 0);
+
+			Start = InStart;
+			Size = InSize;
+			check(IsOccupied());
+		}
+
+		void SetUnOccupied(int32 InStart, int32 InSize)
+		{
+			check(InStart >= 0);
+			check(InSize > 0);
+
+			Start = InStart;
+			Size = -InSize;
+			check(!IsOccupied());
+		}
+
+		bool IsValid() { return Size != 0; }
+
+		bool IsOccupied() const { return Size > 0; }
+		uint32 ComputeSize() const { return (Size > 0) ? Size : -Size; }
+
+		// @apram InSize can be <0 to remove from the size
+		void ExtendUnoccupied(int32 InSize) { check(!IsOccupied()); Size -= InSize; }
+
+		void MakeOccupied(int32 InSize) { check(InSize > 0); check(!IsOccupied()); Size = InSize; }
+		void MakeUnOccupied() { check(IsOccupied()); Size = -Size; }
+
+		bool operator==(const FRange& rhs) const { return Start == rhs.Start && Size == rhs.Size; }
+
+		int32 GetStart() { return Start; }
+		int32 GetEnd() { return Start + ComputeSize(); }
+
+	private:
+		// in bytes
+		int32 Start;
+		// in bytes, 0:not valid, <0:unoccupied, >0:occupied
+		int32 Size;
+	};
+public:
+
+	// constructor
+	FRangeAllocator(uint32 TotalSize)
+	{
+		FRange NewRange;
+
+		NewRange.SetUnOccupied(0, TotalSize);
+
+		Entries.Add(NewRange);
+	}
+
+	// specified range must be non occupied
+	void OccupyRange(FRange InRange)
+	{
+		check(InRange.IsValid());
+		check(InRange.IsOccupied());
+
+		for (uint32 i = 0, Num = Entries.Num(); i < Num; ++i)
+		{
+			FRange& ref = Entries[i];
+
+			if (!ref.IsOccupied())
+			{
+				int32 OverlapSize = ref.GetEnd() - InRange.GetStart();
+
+				if (OverlapSize > 0)
+				{
+					int32 FrontCutSize = InRange.GetStart() - ref.GetStart();
+
+					// there is some front part we cut off
+					if (FrontCutSize > 0)
+					{
+						FRange NewFrontRange;
+
+						NewFrontRange.SetUnOccupied(InRange.GetStart(), ref.ComputeSize() - FrontCutSize);
+
+						ref.SetUnOccupied(ref.GetStart(), FrontCutSize);
+
+						++i;
+
+						// remaining is added behind the found element
+						Entries.Insert(NewFrontRange, i);
+
+						// don't access ref or Num any more - Entries[] might be reallocated
+					}
+
+					check(Entries[i].GetStart() == InRange.GetStart());
+
+					int32 BackCutSize = Entries[i].ComputeSize() - InRange.ComputeSize();
+
+					// otherwise the range was already occupied or not enough space was left (internal error)
+					check(BackCutSize >= 0);
+
+					// there is some back part we cut off
+					if (BackCutSize > 0)
+					{
+						FRange NewBackRange;
+
+						NewBackRange.SetUnOccupied(Entries[i].GetStart() + InRange.ComputeSize(), BackCutSize);
+
+						Entries.Insert(NewBackRange, i + 1);
+					}
+
+					Entries[i] = InRange;
+					return;
+				}
+			}
+		}
+	}
+
+
+	// All resources in ESRAM must be 64KiB aligned
+	
+	// @param InSize >0
+	FRange AllocRange(uint32 InSize)
+	{
+		check(InSize > 0);
+
+		for (uint32 i = 0, Num = Entries.Num(); i < Num; ++i)
+		{
+			FRange& ref = Entries[i];
+
+			if (!ref.IsOccupied())
+			{
+				uint32 RefSize = ref.ComputeSize();
+
+				// take the first fitting one - later we could optimize for minimal fragmentation
+				if (RefSize >= InSize)
+				{
+					ref.MakeOccupied(InSize);
+
+					FRange Ret = ref;
+
+					if (RefSize > InSize)
+					{
+						FRange NewRange;
+
+						NewRange.SetUnOccupied(ref.GetEnd(), RefSize - InSize);
+
+						// remaining is added behind the found element
+						Entries.Insert(NewRange, i + 1);
+					}
+					return Ret;
+				}
+			}
+		}
+
+		// nothing found
+		return FRange();
+	}
+
+	// @param In needs to be what was returned by AllocRange()
+	void ReleaseRange(FRange In)
+	{
+		int32 Index = Entries.Find(In);
+
+		check(Index != INDEX_NONE);
+
+		FRange& refIndex = Entries[Index];
+
+		refIndex.MakeUnOccupied();
+
+		Compacten(Index);
+	}
+
+	// for debugging
+	uint32 GetNumEntries() const { return Entries.Num(); }
+
+	// for debugging
+	uint32 ComputeUnoccupiedSize() const
+	{
+		uint32 Ret = 0;
+
+		for (uint32 i = 0, Num = Entries.Num(); i < Num; ++i)
+		{
+			const FRange& ref = Entries[i];
+
+			if (!ref.IsOccupied())
+			{
+				uint32 RefSize = ref.ComputeSize();
+
+				Ret += RefSize;
+			}
+		}
+
+		return Ret;
+	}
+
+private:
+	// compact unoccupied ranges
+	void Compacten(uint32 StartIndex)
+	{
+		check(!Entries[StartIndex].IsOccupied());
+
+		if (StartIndex && !Entries[StartIndex - 1].IsOccupied())
+		{
+			// Seems we can combine with the element before,
+			// searching further is not needed as we assume the buffer was compact before the last change.
+			--StartIndex;
+		}
+
+		uint32 ElementsToRemove = 0;
+		uint32 SizeGained = 0;
+
+		for (uint32 i = StartIndex + 1, Num = Entries.Num(); i < Num; ++i)
+		{
+			FRange& ref = Entries[i];
+
+			if (!ref.IsOccupied())
+			{
+				++ElementsToRemove;
+				SizeGained += ref.ComputeSize();
+			}
+			else
+			{
+				break;
+			}
+		}
+
+		if (ElementsToRemove)
+		{
+			Entries.RemoveAt(StartIndex + 1, ElementsToRemove, false);
+			Entries[StartIndex].ExtendUnoccupied(SizeGained);
+		}
+	}
+
+private:
+
+	// ordered from small to large (for efficient compactening)
+	TArray<FRange> Entries;
+};
+
+/**
+*	Class of a scoped resource barrier.
+*	This class avoids resource state tracking because resources will be
+*	returned to their original state when the object leaves scope.
+*/
+class FScopeResourceBarrier
+{
+private:
+	FD3D12CommandListHandle &hCommandList;
+	FD3D12Resource *pResource;
+	D3D12_RESOURCE_STATES Current;
+	D3D12_RESOURCE_STATES Desired;
+	uint32 Subresource;
+
+public:
+	explicit FScopeResourceBarrier(FD3D12CommandListHandle &hInCommandList, FD3D12Resource *pInResource, D3D12_RESOURCE_STATES InCurrent, D3D12_RESOURCE_STATES InDesired, const uint32 InSubresource)
+		: hCommandList(hInCommandList)
+		, pResource(pInResource)
+		, Current(InCurrent)
+		, Desired(InDesired)
+		, Subresource(InSubresource)
+	{
+		check(!pResource->RequiresResourceStateTracking());
+		FD3D12DynamicRHI::ResourceBarrier(hCommandList, pResource, Current, Desired, Subresource);
+	}
+
+	~FScopeResourceBarrier()
+	{
+		FD3D12DynamicRHI::ResourceBarrier(hCommandList, pResource, Desired, Current, Subresource);
+	}
+};
+
+/**
+*	Class of a scoped resource barrier.
+*	This class conditionally uses resource state tracking.
+*	This should only be used with the Editor.
+*/
+class FConditionalScopeResourceBarrier
+{
+private:
+	FD3D12CommandListHandle &hCommandList;
+	FD3D12Resource *pResource;
+	D3D12_RESOURCE_STATES Current;
+	D3D12_RESOURCE_STATES Desired;
+	uint32 Subresource;
+	bool bUseTracking;
+
+public:
+	explicit FConditionalScopeResourceBarrier(FD3D12CommandListHandle &hInCommandList, FD3D12Resource *pInResource, D3D12_RESOURCE_STATES InDesired, const uint32 InSubresource)
+		: hCommandList(hInCommandList)
+		, pResource(pInResource)
+		, Current(D3D12_RESOURCE_STATE_TBD)
+		, Desired(InDesired)
+		, Subresource(InSubresource)
+		, bUseTracking(pResource->RequiresResourceStateTracking())
+	{
+		if (!bUseTracking)
+		{
+			Current = pResource->GetDefaultResourceState();
+			FD3D12DynamicRHI::ResourceBarrier(hCommandList, pResource, Current, Desired, Subresource);
+		}
+		else
+		{
+			FD3D12DynamicRHI::TransitionResource(hCommandList, pResource, Desired, Subresource);
+		}
+	}
+
+	~FConditionalScopeResourceBarrier()
+	{
+		// Return the resource to it's default state if it doesn't use tracking
+		if (!bUseTracking)
+		{
+			FD3D12DynamicRHI::ResourceBarrier(hCommandList, pResource, Desired, Current, Subresource);
+		}
+	}
+};
+
+
+#endif
