UE4 - D3D12RHI (not working)

diff --git a/Engine/Source/Runtime/Windows/D3D12RHI/Private/D3D12RenderTarget.cpp b/Engine/Source/Runtime/Windows/D3D12RHI/Private/D3D12RenderTarget.cpp
new file mode 100644
--- /dev/null
+++ b/Engine/Source/Runtime/Windows/D3D12RHI/Private/D3D12RenderTarget.cpp
@@ -0,0 +1,1520 @@
+// Copyright 1998-2014 Epic Games, Inc. All Rights Reserved.
+
+/*=============================================================================
+	D3D12RenderTarget.cpp: D3D render target implementation.
+	=============================================================================*/
+
+#include "D3D12RHIPrivate.h"
+#include "BatchedElements.h"
+#include "ScreenRendering.h"
+#include "RHIStaticStates.h"
+#include "ResolveShader.h"
+
+static inline DXGI_FORMAT ConvertTypelessToUnorm(DXGI_FORMAT Format)
+{
+	// required to prevent 
+	// D3D11: ERROR: ID3D11DeviceContext::ResolveSubresource: The Format (0x1b, R8G8B8A8_TYPELESS) is never able to resolve multisampled resources. [ RESOURCE_MANIPULATION ERROR #294: DEVICE_RESOLVESUBRESOURCE_FORMAT_INVALID ]
+	// D3D11: **BREAK** enabled for the previous D3D11 message, which was: [ RESOURCE_MANIPULATION ERROR #294: DEVICE_RESOLVESUBRESOURCE_FORMAT_INVALID ]
+	switch (Format)
+	{
+	case DXGI_FORMAT_R8G8B8A8_TYPELESS:
+		return DXGI_FORMAT_R8G8B8A8_UNORM;
+
+	case DXGI_FORMAT_B8G8R8A8_TYPELESS:
+		return DXGI_FORMAT_B8G8R8A8_UNORM;
+
+	default:
+		return Format;
+	}
+}
+
+static FResolveRect GetDefaultRect(const FResolveRect& Rect, uint32 DefaultWidth, uint32 DefaultHeight)
+{
+	if (Rect.X1 >= 0 && Rect.X2 >= 0 && Rect.Y1 >= 0 && Rect.Y2 >= 0)
+	{
+		return Rect;
+	}
+	else
+	{
+		return FResolveRect(0, 0, DefaultWidth, DefaultHeight);
+	}
+}
+
+template<typename TPixelShader>
+void FD3D12CommandContext::ResolveTextureUsingShader(
+	FRHICommandList_RecursiveHazardous& RHICmdList,
+	FD3D12Texture2D* SourceTexture,
+	FD3D12Texture2D* DestTexture,
+	FD3D12RenderTargetView* DestTextureRTV,
+	FD3D12DepthStencilView* DestTextureDSV,
+	const D3D12_RESOURCE_DESC& ResolveTargetDesc,
+	const FResolveRect& SourceRect,
+	const FResolveRect& DestRect,
+	typename TPixelShader::FParameter PixelShaderParameter
+	)
+{
+	// Save the current viewport so that it can be restored
+	D3D12_VIEWPORT SavedViewport;
+	uint32 NumSavedViewports = 1;
+	StateCache.GetViewports(&NumSavedViewports, &SavedViewport);
+
+	// No alpha blending, no depth tests or writes, no stencil tests or writes, no backface culling.
+	RHICmdList.SetBlendState(TStaticBlendState<>::GetRHI(), FLinearColor::White);
+	RHICmdList.SetRasterizerState(TStaticRasterizerState<FM_Solid, CM_None>::GetRHI());
+	RHICmdList.Flush(); // always call flush when using a command list in RHI implementations before doing anything else. This is super hazardous.
+
+	// Make sure the destination is not bound as a shader resource.
+	if (DestTexture)
+	{
+		ConditionalClearShaderResource(DestTexture->ResourceLocation);
+	}
+
+	// Determine if the entire destination surface is being resolved to.
+	// If the entire surface is being resolved to, then it means we can clear it and signal the driver that it can discard
+	// the surface's previous contents, which breaks dependencies between frames when using alternate-frame SLI.
+	const bool bClearDestTexture =
+		DestRect.X1 == 0
+		&& DestRect.Y1 == 0
+		&& (uint64)DestRect.X2 == ResolveTargetDesc.Width
+		&&	DestRect.Y2 == ResolveTargetDesc.Height;
+
+    if (ResolveTargetDesc.Flags & D3D12_RESOURCE_FLAG_ALLOW_DEPTH_STENCIL)
+	{
+		// Clear the destination texture.
+		if (bClearDestTexture)
+		{
+            OwningRHI.RegisterGPUWork(0);
+
+			FD3D12DynamicRHI::TransitionResource(CommandListHandle, DestTextureDSV, D3D12_RESOURCE_STATE_DEPTH_WRITE);
+
+			numClears++;
+			CommandListHandle->ClearDepthStencilView(DestTextureDSV->GetView(), D3D12_CLEAR_FLAG_DEPTH | D3D12_CLEAR_FLAG_STENCIL, 0, 0, 0, nullptr);
+		}
+
+		RHICmdList.SetDepthStencilState(TStaticDepthStencilState<true, CF_Always>::GetRHI(), 0);
+		RHICmdList.Flush(); // always call flush when using a command list in RHI implementations before doing anything else. This is super hazardous.
+
+		// Write to the dest texture as a depth-stencil target.
+		FD3D12RenderTargetView* NullRTV = NULL;
+
+		StateCache.SetRenderTargets(1, &NullRTV, DestTextureDSV);
+
+	}
+	else
+	{
+		// Clear the destination texture.
+		if (bClearDestTexture)
+		{
+            OwningRHI.RegisterGPUWork(0);
+
+			FD3D12DynamicRHI::TransitionResource(CommandListHandle, DestTextureRTV, D3D12_RESOURCE_STATE_RENDER_TARGET);
+
+			FLinearColor ClearColor(0, 0, 0, 0);
+			numClears++;
+			CommandListHandle->ClearRenderTargetView(DestTextureRTV->GetView(), (float*)&ClearColor, 0, nullptr);
+		}
+
+		RHICmdList.SetDepthStencilState(TStaticDepthStencilState<false, CF_Always>::GetRHI(), 0);
+		RHICmdList.Flush(); // always call flush when using a command list in RHI implementations before doing anything else. This is super hazardous.
+
+		// Write to the dest surface as a render target.
+		StateCache.SetRenderTargets(1, &DestTextureRTV, NULL);
+	}
+
+	RHICmdList.SetViewport(0.0f, 0.0f, 0.0f, (uint32)ResolveTargetDesc.Width, ResolveTargetDesc.Height, 1.0f);
+
+	// Generate the vertices used to copy from the source surface to the destination surface.
+	const float MinU = SourceRect.X1;
+	const float MinV = SourceRect.Y1;
+	const float MaxU = SourceRect.X2;
+	const float MaxV = SourceRect.Y2;
+	const float MinX = -1.f + DestRect.X1 / ((float)ResolveTargetDesc.Width * 0.5f);
+	const float MinY = +1.f - DestRect.Y1 / ((float)ResolveTargetDesc.Height * 0.5f);
+	const float MaxX = -1.f + DestRect.X2 / ((float)ResolveTargetDesc.Width * 0.5f);
+	const float MaxY = +1.f - DestRect.Y2 / ((float)ResolveTargetDesc.Height * 0.5f);
+
+	// MSFT: Seb: Does this need to be extern?
+	static FGlobalBoundShaderState ResolveBoundShaderState;
+
+	// Set the vertex and pixel shader
+	auto ShaderMap = GetGlobalShaderMap(GMaxRHIFeatureLevel);
+	TShaderMapRef<FResolveVS> ResolveVertexShader(ShaderMap);
+	TShaderMapRef<TPixelShader> ResolvePixelShader(ShaderMap);
+	SetGlobalBoundShaderState(RHICmdList, GMaxRHIFeatureLevel, ResolveBoundShaderState, GScreenVertexDeclaration.VertexDeclarationRHI, *ResolveVertexShader, *ResolvePixelShader);
+
+	ResolvePixelShader->SetParameters(RHICmdList, PixelShaderParameter);
+	RHICmdList.Flush(); // always call flush when using a command list in RHI implementations before doing anything else. This is super hazardous.
+
+	// Set the source texture.
+	const uint32 TextureIndex = ResolvePixelShader->UnresolvedSurface.GetBaseIndex();
+	SetShaderResourceView<SF_Pixel>(SourceTexture->ResourceLocation, SourceTexture->GetShaderResourceView(), TextureIndex);
+
+	// Generate the vertices used
+	FScreenVertex Vertices[4];
+
+	Vertices[0].Position.X = MaxX;
+	Vertices[0].Position.Y = MinY;
+	Vertices[0].UV.X = MaxU;
+	Vertices[0].UV.Y = MinV;
+
+	Vertices[1].Position.X = MaxX;
+	Vertices[1].Position.Y = MaxY;
+	Vertices[1].UV.X = MaxU;
+	Vertices[1].UV.Y = MaxV;
+
+	Vertices[2].Position.X = MinX;
+	Vertices[2].Position.Y = MinY;
+	Vertices[2].UV.X = MinU;
+	Vertices[2].UV.Y = MinV;
+
+	Vertices[3].Position.X = MinX;
+	Vertices[3].Position.Y = MaxY;
+	Vertices[3].UV.X = MinU;
+	Vertices[3].UV.Y = MaxV;
+
+	DrawPrimitiveUP(RHICmdList, PT_TriangleStrip, 2, Vertices, sizeof(Vertices[0]));
+	RHICmdList.Flush(); // always call flush when using a command list in RHI implementations before doing anything else. This is super hazardous.
+
+	if (SourceTexture)
+	{
+		ConditionalClearShaderResource(SourceTexture->ResourceLocation);
+	}
+
+	// Reset saved render targets
+	CommitRenderTargetsAndUAVs();
+
+	// Reset saved viewport
+	{
+		StateCache.SetViewports(1, &SavedViewport);
+	}
+}
+
+/**
+* Copies the contents of the given surface to its resolve target texture.
+* @param SourceSurface - surface with a resolve texture to copy to
+* @param bKeepOriginalSurface - true if the original surface will still be used after this function so must remain valid
+* @param ResolveParams - optional resolve params
+*/
+void FD3D12CommandContext::RHICopyToResolveTarget(FTextureRHIParamRef SourceTextureRHI, FTextureRHIParamRef DestTextureRHI, bool bKeepOriginalSurface, const FResolveParams& ResolveParams)
+{
+	if (!SourceTextureRHI || !DestTextureRHI)
+	{
+		// no need to do anything (silently ignored)
+		return;
+	}
+
+	FRHICommandList_RecursiveHazardous RHICmdList(this);
+
+	FD3D12Texture2D* SourceTexture2D = static_cast<FD3D12Texture2D*>(SourceTextureRHI->GetTexture2D());
+	FD3D12Texture2D* DestTexture2D = static_cast<FD3D12Texture2D*>(DestTextureRHI->GetTexture2D());
+
+	FD3D12TextureCube* SourceTextureCube = static_cast<FD3D12TextureCube*>(SourceTextureRHI->GetTextureCube());
+	FD3D12TextureCube* DestTextureCube = static_cast<FD3D12TextureCube*>(DestTextureRHI->GetTextureCube());
+
+	FD3D12Texture3D* SourceTexture3D = static_cast<FD3D12Texture3D*>(SourceTextureRHI->GetTexture3D());
+	FD3D12Texture3D* DestTexture3D = static_cast<FD3D12Texture3D*>(DestTextureRHI->GetTexture3D());
+
+#if CHECK_SRV_TRANSITIONS
+	//check first, as some of the resolve shaders could trigger the resource we're currently resolving.
+	ID3D12Resource* SourceResource = nullptr;
+	int32 ResolveSlice = -1;
+	if (SourceTexture2D)
+	{
+		SourceResource = SourceTexture2D->GetResource();
+	}
+	if (SourceTextureCube)
+	{
+		SourceResource = SourceTextureCube->GetResource();
+		ResolveSlice = ResolveParams.SourceArrayIndex * 6 + ResolveParams.CubeFace;
+	}
+	if (SourceTexture3D)
+	{
+		SourceResource = SourceTexture3D->GetResource();
+	}
+
+	if (SourceResource)
+	{
+		// if we don't have specific resolve data remove all references.
+		// this is helpful when you clear many mips/slices but only want to do one resolve.
+		if (ResolveParams.MipIndex == -1 && ResolveParams.SourceArrayIndex == -1)
+		{
+			UnresolvedTargets.Remove(SourceResource);
+		}
+		else
+		{
+			UnresolvedTargets.Remove(SourceResource, FUnresolvedRTInfo(SourceTextureRHI->GetName(), ResolveParams.MipIndex, 1, ResolveSlice, 1));
+		}
+	}
+#endif
+
+	if (SourceTexture2D && DestTexture2D)
+	{
+		check(!SourceTextureCube && !DestTextureCube);
+		if (SourceTexture2D != DestTexture2D)
+		{
+            OwningRHI.RegisterGPUWork();
+
+			if (GetParentDevice()->GetFeatureLevel() == D3D_FEATURE_LEVEL_11_0
+				&& DestTexture2D->GetDepthStencilView(FExclusiveDepthStencil::DepthWrite_StencilWrite)
+				&& SourceTextureRHI->IsMultisampled()
+				&& !DestTextureRHI->IsMultisampled())
+			{
+				D3D12_RESOURCE_DESC const& ResolveTargetDesc = DestTexture2D->GetResource()->GetDesc();
+
+				ResolveTextureUsingShader<FResolveDepthPS>(
+					RHICmdList,
+					SourceTexture2D,
+					DestTexture2D,
+					DestTexture2D->GetRenderTargetView(0, -1),
+					DestTexture2D->GetDepthStencilView(FExclusiveDepthStencil::DepthWrite_StencilWrite),
+					ResolveTargetDesc,
+					GetDefaultRect(ResolveParams.Rect, DestTexture2D->GetSizeX(), DestTexture2D->GetSizeY()),
+					GetDefaultRect(ResolveParams.Rect, DestTexture2D->GetSizeX(), DestTexture2D->GetSizeY()),
+					FDummyResolveParameter()
+					);
+			}
+			else if (GetParentDevice()->GetFeatureLevel() == D3D_FEATURE_LEVEL_10_0
+				&& DestTexture2D->GetDepthStencilView(FExclusiveDepthStencil::DepthWrite_StencilWrite))
+			{
+				D3D12_RESOURCE_DESC const& ResolveTargetDesc = DestTexture2D->GetResource()->GetDesc();
+
+				ResolveTextureUsingShader<FResolveDepthNonMSPS>(
+					RHICmdList,
+					SourceTexture2D,
+					DestTexture2D,
+					NULL,
+					DestTexture2D->GetDepthStencilView(FExclusiveDepthStencil::DepthWrite_StencilWrite),
+					ResolveTargetDesc,
+					GetDefaultRect(ResolveParams.Rect, DestTexture2D->GetSizeX(), DestTexture2D->GetSizeY()),
+					GetDefaultRect(ResolveParams.Rect, DestTexture2D->GetSizeX(), DestTexture2D->GetSizeY()),
+					FDummyResolveParameter()
+					);
+			}
+			else
+			{
+				DXGI_FORMAT SrcFmt = (DXGI_FORMAT)GPixelFormats[SourceTextureRHI->GetFormat()].PlatformFormat;
+				DXGI_FORMAT DstFmt = (DXGI_FORMAT)GPixelFormats[DestTexture2D->GetFormat()].PlatformFormat;
+
+				DXGI_FORMAT Fmt = ConvertTypelessToUnorm((DXGI_FORMAT)GPixelFormats[DestTexture2D->GetFormat()].PlatformFormat);
+
+				// Determine whether a MSAA resolve is needed, or just a copy.
+				if (SourceTextureRHI->IsMultisampled() && !DestTexture2D->IsMultisampled())
+				{
+					otherWorkCounter++;
+					CommandListHandle->ResolveSubresource(
+						DestTexture2D->GetResource()->GetResource(),
+						0,
+						SourceTexture2D->GetResource()->GetResource(),
+						0,
+						Fmt
+						);
+				}
+				else
+				{
+					if (ResolveParams.Rect.IsValid())
+					{
+						D3D12_BOX SrcBox;
+
+						SrcBox.left = ResolveParams.Rect.X1;
+						SrcBox.top = ResolveParams.Rect.Y1;
+						SrcBox.front = 0;
+						SrcBox.right = ResolveParams.Rect.X2;
+						SrcBox.bottom = ResolveParams.Rect.Y2;
+						SrcBox.back = 1;
+
+						{
+							FConditionalScopeResourceBarrier ConditionalScopeResourceBarrierDest(CommandListHandle, DestTexture2D->GetResource(), D3D12_RESOURCE_STATE_COPY_DEST, 0);
+							FConditionalScopeResourceBarrier ConditionalScopeResourceBarrierSource(CommandListHandle, SourceTexture2D->GetResource(), D3D12_RESOURCE_STATE_COPY_SOURCE, 0);
+
+							CD3DX12_TEXTURE_COPY_LOCATION DestCopyLocation(DestTexture2D->GetResource()->GetResource(), 0);
+							CD3DX12_TEXTURE_COPY_LOCATION SourceCopyLocation(SourceTexture2D->GetResource()->GetResource(), 0);
+
+							numCopies++;
+							CommandListHandle->CopyTextureRegion(
+								&DestCopyLocation,
+								ResolveParams.Rect.X1, ResolveParams.Rect.Y1, 0,
+								&SourceCopyLocation,
+								&SrcBox);
+						}
+
+						DEBUG_EXECUTE_COMMAND_LIST(this);
+					}
+					else
+					{
+						FConditionalScopeResourceBarrier ConditionalScopeResourceBarrierSource(CommandListHandle, SourceTexture2D->GetResource(), D3D12_RESOURCE_STATE_COPY_SOURCE, 0);
+						
+						// Resolve to a buffer.
+						D3D12_RESOURCE_DESC const& ResolveTargetDesc = DestTexture2D->GetResource()->GetDesc();
+						if (ResolveTargetDesc.Dimension == D3D12_RESOURCE_DIMENSION_BUFFER)
+						{
+							check(IsDefaultContext());
+							D3D12_RESOURCE_DESC const& srcDesc = SourceTexture2D->GetResource()->GetDesc();
+	
+							const uint32 BlockBytes = GPixelFormats[SourceTexture2D->GetFormat()].BlockBytes;
+							const uint32 XBytes = (uint32)srcDesc.Width * BlockBytes;
+							const uint32 XBytesAligned = Align(XBytes, D3D12_TEXTURE_DATA_PITCH_ALIGNMENT);
+							
+							D3D12_SUBRESOURCE_FOOTPRINT destSubresource;
+							destSubresource.Depth = 1;
+							destSubresource.Height = srcDesc.Height;
+							destSubresource.Width = srcDesc.Width;
+							destSubresource.Format = srcDesc.Format;
+							destSubresource.RowPitch = XBytesAligned;
+							
+							D3D12_PLACED_SUBRESOURCE_FOOTPRINT placedTexture2D = { 0 };
+							placedTexture2D.Offset = 0;
+							placedTexture2D.Footprint = destSubresource;
+	
+	                        CD3DX12_TEXTURE_COPY_LOCATION DestCopyLocation(DestTexture2D->GetResource()->GetResource(), placedTexture2D);
+	                        CD3DX12_TEXTURE_COPY_LOCATION SourceCopyLocation(SourceTexture2D->GetResource()->GetResource(), 0);
+
+							numCopies++;
+							CommandListHandle->CopyTextureRegion(
+	                            &DestCopyLocation,
+	                            0, 0, 0,
+	                            &SourceCopyLocation,
+	                            nullptr);
+	
+	                        // Save the command list handle. This lets us check when this command list is complete. Note: This must be saved before we execute the command list
+	                        DestTexture2D->SetReadBackListHandle(CommandListHandle);
+	
+							// Break up the command list here so that the wait on the previous frame's results don't block.
+	                        FlushCommands();
+						}
+						// Resolve to a texture.
+						else
+						{
+							// Transition to the copy dest state.
+							FConditionalScopeResourceBarrier ConditionalScopeResourceBarrierDest(CommandListHandle, DestTexture2D->GetResource(), D3D12_RESOURCE_STATE_COPY_DEST, 0);
+							
+							CD3DX12_TEXTURE_COPY_LOCATION DestCopyLocation(DestTexture2D->GetResource()->GetResource(), 0);
+							CD3DX12_TEXTURE_COPY_LOCATION SourceCopyLocation(SourceTexture2D->GetResource()->GetResource(), 0);
+							
+							numCopies++;
+							CommandListHandle->CopyTextureRegion(
+								&DestCopyLocation,
+								0, 0, 0,
+								&SourceCopyLocation,
+								nullptr);
+							
+							// Transition the resource back to it's previous state.
+							if (DestTexture2D->GetResource()->RequiresResourceStateTracking())
+							{
+								if (DestTexture2D->GetResource()->GetDesc().Flags & D3D12_RESOURCE_FLAG_ALLOW_DEPTH_STENCIL)
+									FD3D12DynamicRHI::TransitionResource(CommandListHandle, DestTexture2D->GetResource(), D3D12_RESOURCE_STATE_DEPTH_WRITE, 0);
+								else
+									FD3D12DynamicRHI::TransitionResource(CommandListHandle, DestTexture2D->GetResource(), D3D12_RESOURCE_STATE_RENDER_TARGET, 0);
+							}
+						}
+					}
+				}
+			}
+		}
+	}
+	else if (SourceTextureCube && DestTextureCube)
+	{
+		check(!SourceTexture2D && !DestTexture2D);
+
+		if (SourceTextureCube != DestTextureCube)
+		{
+            OwningRHI.RegisterGPUWork();
+
+			// Determine the cubemap face being resolved.
+			const uint32 D3DFace = GetD3D11CubeFace(ResolveParams.CubeFace);
+			const uint32 SourceSubresource = CalcSubresource(ResolveParams.MipIndex, ResolveParams.SourceArrayIndex * 6 + D3DFace, SourceTextureCube->GetNumMips());
+			const uint32 DestSubresource = CalcSubresource(ResolveParams.MipIndex, ResolveParams.DestArrayIndex * 6 + D3DFace, DestTextureCube->GetNumMips());
+
+			// Determine whether a MSAA resolve is needed, or just a copy.
+			if (SourceTextureRHI->IsMultisampled() && !DestTextureCube->IsMultisampled())
+			{
+				otherWorkCounter++;
+				CommandListHandle->ResolveSubresource(
+					DestTextureCube->GetResource()->GetResource(),
+					DestSubresource,
+					SourceTextureCube->GetResource()->GetResource(),
+					SourceSubresource,
+					(DXGI_FORMAT)GPixelFormats[DestTextureCube->GetFormat()].PlatformFormat
+					);
+			}
+			else
+			{
+                CD3DX12_TEXTURE_COPY_LOCATION DestCopyLocation(DestTextureCube->GetResource()->GetResource(), DestSubresource);
+				CD3DX12_TEXTURE_COPY_LOCATION SourceCopyLocation(SourceTextureCube->GetResource()->GetResource(), SourceSubresource);
+
+				{
+					FConditionalScopeResourceBarrier ConditionalScopeResourceBarrierDest(CommandListHandle, DestTextureCube->GetResource(), D3D12_RESOURCE_STATE_COPY_DEST, DestCopyLocation.SubresourceIndex);
+					FConditionalScopeResourceBarrier ConditionalScopeResourceBarrierSource(CommandListHandle, SourceTextureCube->GetResource(), D3D12_RESOURCE_STATE_COPY_SOURCE, SourceCopyLocation.SubresourceIndex);
+
+					numCopies++;
+					CommandListHandle->CopyTextureRegion(
+						&DestCopyLocation,
+						0, 0, 0,
+						&SourceCopyLocation,
+						nullptr);
+				}
+
+				DEBUG_EXECUTE_COMMAND_LIST(this);
+			}
+		}
+	}
+	else if (SourceTexture2D && DestTextureCube)
+	{
+		// If source is 2D and Dest is a cube then copy the 2D texture to the specified cube face.
+		// Determine the cubemap face being resolved.
+		const uint32 D3DFace = GetD3D11CubeFace(ResolveParams.CubeFace);
+		const uint32 Subresource = CalcSubresource(0, D3DFace, 1);
+
+        CD3DX12_TEXTURE_COPY_LOCATION DestCopyLocation(DestTextureCube->GetResource()->GetResource(), Subresource);
+        CD3DX12_TEXTURE_COPY_LOCATION SourceCopyLocation(SourceTexture2D->GetResource()->GetResource(), 0);
+
+		{
+			FConditionalScopeResourceBarrier ConditionalScopeResourceBarrierDest(CommandListHandle, DestTextureCube->GetResource(), D3D12_RESOURCE_STATE_COPY_DEST, DestCopyLocation.SubresourceIndex);
+			FConditionalScopeResourceBarrier ConditionalScopeResourceBarrierSource(CommandListHandle, SourceTexture2D->GetResource(), D3D12_RESOURCE_STATE_COPY_SOURCE, SourceCopyLocation.SubresourceIndex);
+
+			numCopies++;
+			CommandListHandle->CopyTextureRegion(
+				&DestCopyLocation,
+				0, 0, 0,
+				&SourceCopyLocation,
+				nullptr);
+		}
+
+		DEBUG_EXECUTE_COMMAND_LIST(this);
+	}
+	else if (SourceTexture3D && DestTexture3D)
+	{
+		// bit of a hack.  no one resolves slice by slice and 0 is the default value.  assume for the moment they are resolving the whole texture.
+		check(ResolveParams.SourceArrayIndex == 0);
+		check(SourceTexture3D == DestTexture3D);
+	}
+}
+
+/**
+* Helper for storing IEEE 32 bit float components
+*/
+struct FFloatIEEE
+{
+	union
+	{
+		struct
+		{
+			uint32	Mantissa : 23, Exponent : 8, Sign : 1;
+		} Components;
+
+		float	Float;
+	};
+};
+
+/**
+* Helper for storing 16 bit float components
+*/
+struct FD3DFloat16
+{
+	union
+	{
+		struct
+		{
+			uint16	Mantissa : 10, Exponent : 5, Sign : 1;
+		} Components;
+
+		uint16	Encoded;
+	};
+
+	/**
+	* @return full 32 bit float from the 16 bit value
+	*/
+	operator float()
+	{
+		FFloatIEEE	Result;
+
+		Result.Components.Sign = Components.Sign;
+		Result.Components.Exponent = Components.Exponent - 15 + 127; // Stored exponents are biased by half their range.
+		Result.Components.Mantissa = FMath::Min<uint32>(FMath::FloorToInt((float)Components.Mantissa / 1024.0f * 8388608.0f), (1 << 23) - 1);
+
+		return Result.Float;
+	}
+};
+
+/**
+* Helper for storing DXGI_FORMAT_R11G11B10_FLOAT components
+*/
+struct FD3DFloatR11G11B10
+{
+	// http://msdn.microsoft.com/En-US/library/bb173059(v=VS.85).aspx
+	uint32 R_Mantissa : 6;
+	uint32 R_Exponent : 5;
+	uint32 G_Mantissa : 6;
+	uint32 G_Exponent : 5;
+	uint32 B_Mantissa : 5;
+	uint32 B_Exponent : 5;
+
+	/**
+	* @return decompress into three 32 bit float
+	*/
+	operator FLinearColor()
+	{
+		FFloatIEEE	Result[3];
+
+		Result[0].Components.Sign = 0;
+		Result[0].Components.Exponent = R_Exponent - 15 + 127;
+		Result[0].Components.Mantissa = FMath::Min<uint32>(FMath::FloorToInt((float)R_Mantissa / 32.0f * 8388608.0f), (1 << 23) - 1);
+		Result[1].Components.Sign = 0;
+		Result[1].Components.Exponent = G_Exponent - 15 + 127;
+		Result[1].Components.Mantissa = FMath::Min<uint32>(FMath::FloorToInt((float)G_Mantissa / 64.0f * 8388608.0f), (1 << 23) - 1);
+		Result[2].Components.Sign = 0;
+		Result[2].Components.Exponent = B_Exponent - 15 + 127;
+		Result[2].Components.Mantissa = FMath::Min<uint32>(FMath::FloorToInt((float)B_Mantissa / 64.0f * 8388608.0f), (1 << 23) - 1);
+
+		return FLinearColor(Result[0].Float, Result[1].Float, Result[2].Float);
+	}
+};
+
+// Only supports the formats that are supported by ConvertRAWSurfaceDataToFColor()
+static uint32 ComputeBytesPerPixel(DXGI_FORMAT Format)
+{
+	uint32 BytesPerPixel = 0;
+
+	switch (Format)
+	{
+	case DXGI_FORMAT_R16_TYPELESS:
+		BytesPerPixel = 2;
+		break;
+	case DXGI_FORMAT_B8G8R8A8_TYPELESS:
+	case DXGI_FORMAT_B8G8R8A8_UNORM:
+	case DXGI_FORMAT_R8G8B8A8_TYPELESS:
+	case DXGI_FORMAT_R8G8B8A8_UNORM:
+	case DXGI_FORMAT_R24G8_TYPELESS:
+	case DXGI_FORMAT_R10G10B10A2_UNORM:
+	case DXGI_FORMAT_R11G11B10_FLOAT:
+	case DXGI_FORMAT_R16G16_UNORM:
+		BytesPerPixel = 4;
+		break;
+	case DXGI_FORMAT_R16G16B16A16_FLOAT:
+	case DXGI_FORMAT_R16G16B16A16_UNORM:
+		BytesPerPixel = 8;
+		break;
+#if DEPTH_32_BIT_CONVERSION
+		// Changing Depth Buffers to 32 bit on Dingo as D24S8 is actually implemented as a 32 bit buffer in the hardware
+	case DXGI_FORMAT_R32G8X24_TYPELESS:
+	case DXGI_FORMAT_D32_FLOAT_S8X24_UINT:
+	case DXGI_FORMAT_R32_FLOAT_X8X24_TYPELESS:
+#endif
+		BytesPerPixel = 5;
+		break;
+	case DXGI_FORMAT_R32G32B32A32_FLOAT:
+		BytesPerPixel = 16;
+		break;
+	case DXGI_FORMAT_R8_TYPELESS:
+	case DXGI_FORMAT_R8_UNORM:
+	case DXGI_FORMAT_R8_UINT:
+	case DXGI_FORMAT_R8_SNORM:
+	case DXGI_FORMAT_R8_SINT:
+		BytesPerPixel = 1;
+		break;
+	}
+
+	// format not supported yet
+	check(BytesPerPixel);
+
+	return BytesPerPixel;
+}
+
+TRefCountPtr<FD3D12Resource> FD3D12DynamicRHI::GetStagingTexture(FTextureRHIParamRef TextureRHI, FIntRect InRect, FIntRect& StagingRectOUT, FReadSurfaceDataFlags InFlags, D3D12_PLACED_SUBRESOURCE_FOOTPRINT &readbackHeapDesc)
+{
+    FD3D12CommandListHandle& hCommandList = GetRHIDevice()->GetDefaultCommandContext().CommandListHandle;
+	FD3D12TextureBase* Texture = GetD3D11TextureFromRHITexture(TextureRHI);
+	D3D12_RESOURCE_DESC const& SourceDesc = Texture->GetResource()->GetDesc();
+
+    // Ensure we're dealing with a Texture2D, which the rest of this function already assumes
+    check(TextureRHI->GetTexture2D());
+    FD3D12Texture2D *InTexture2D = static_cast<FD3D12Texture2D*>(Texture);
+
+    bool bRequiresTempStagingTexture = Texture->GetResource()->GetHeapType() != D3D12_HEAP_TYPE_READBACK;
+	if (bRequiresTempStagingTexture == false)
+	{
+		// Returning the same texture is considerably faster than creating and copying to
+		// a new staging texture as we do not have to wait for the GPU pipeline to catch up
+		// to the staging texture preparation work.
+
+		// Texture2Ds on the readback heap will have been flattened to 1D, so we need to retrieve pitch
+		// information from the original 2D version to correctly use sub-rects.
+		readbackHeapDesc = InTexture2D->GetReadBackHeapDesc();
+		StagingRectOUT = InRect;
+
+		return (Texture->GetResource());
+	}
+
+	// a temporary staging texture is needed.
+	int32 SizeX = InRect.Width();
+	int32 SizeY = InRect.Height();
+	// Read back the surface data in the defined rect
+	D3D12_BOX Rect;
+	Rect.left = InRect.Min.X;
+	Rect.top = InRect.Min.Y;
+	Rect.right = InRect.Max.X;
+	Rect.bottom = InRect.Max.Y;
+	Rect.back = 1;
+	Rect.front = 0;
+
+	// create a temp 2d texture to copy render target to
+	TRefCountPtr<FD3D12Resource> TempTexture2D;
+
+	const uint32 BlockBytes = GPixelFormats[TextureRHI->GetFormat()].BlockBytes;
+	const uint32 XBytesAligned = Align((uint32)SourceDesc.Width * BlockBytes, D3D12_TEXTURE_DATA_PITCH_ALIGNMENT);
+	const uint32 MipBytesAligned = XBytesAligned * SourceDesc.Height;
+    VERIFYD3D11RESULT(GetRHIDevice()->GetResourceHelper().CreateBuffer(D3D12_HEAP_TYPE_READBACK, MipBytesAligned, TempTexture2D.GetInitReference()));
+
+	// Staging rectangle is now the whole surface.
+	StagingRectOUT.Min = FIntPoint::ZeroValue;
+	StagingRectOUT.Max = FIntPoint(SizeX, SizeY);
+
+	// Copy the data to a staging resource.
+	uint32 Subresource = 0;
+    if (InTexture2D->IsCubemap())
+	{
+		uint32 D3DFace = GetD3D11CubeFace(InFlags.GetCubeFace());
+		Subresource = CalcSubresource(0, D3DFace, 1);
+	}
+
+	D3D12_BOX* RectPtr = nullptr; // API prefers NULL for entire texture.
+	if (Rect.left != 0 || Rect.top != 0 || Rect.right != SourceDesc.Width || Rect.bottom != SourceDesc.Height)
+	{
+		// ..Sub rectangle required, use the D3D12_BOX.
+		RectPtr = &Rect;
+	}
+
+	uint32 BytesPerPixel = ComputeBytesPerPixel(SourceDesc.Format);
+	D3D12_SUBRESOURCE_FOOTPRINT destSubresource;
+	destSubresource.Depth = 1;
+	destSubresource.Height = SourceDesc.Height;
+	destSubresource.Width = SourceDesc.Width;
+	destSubresource.Format = SourceDesc.Format;
+	destSubresource.RowPitch = XBytesAligned;
+	check(destSubresource.RowPitch % D3D12_TEXTURE_DATA_PITCH_ALIGNMENT == 0);	// Make sure we align correctly.
+
+	D3D12_PLACED_SUBRESOURCE_FOOTPRINT placedTexture2D = { 0 };
+	placedTexture2D.Offset = 0;
+	placedTexture2D.Footprint = destSubresource;
+
+    CD3DX12_TEXTURE_COPY_LOCATION DestCopyLocation(TempTexture2D->GetResource(), placedTexture2D);
+    CD3DX12_TEXTURE_COPY_LOCATION SourceCopyLocation(Texture->GetResource()->GetResource(), Subresource);
+
+	TransitionResource(hCommandList, Texture->GetResource(), D3D12_RESOURCE_STATE_COPY_SOURCE, SourceCopyLocation.SubresourceIndex);
+	// Upload heap doesn't need to transition
+
+	GetRHIDevice()->GetDefaultCommandContext().numCopies++;
+	hCommandList->CopyTextureRegion(
+        &DestCopyLocation,
+        0, 0, 0,
+        &SourceCopyLocation,
+        RectPtr);
+
+	// Remember the width, height, pitch, etc...
+	readbackHeapDesc = placedTexture2D;
+
+	// We need to execute the command list so we can read the data from readback heap
+    GetRHIDevice()->GetDefaultCommandContext().FlushCommands(true);
+
+	return TempTexture2D;
+}
+
+void FD3D12DynamicRHI::ReadSurfaceDataNoMSAARaw(FTextureRHIParamRef TextureRHI, FIntRect InRect, TArray<uint8>& OutData, FReadSurfaceDataFlags InFlags)
+{
+	FD3D12TextureBase* Texture = GetD3D11TextureFromRHITexture(TextureRHI);
+
+	const uint32 SizeX = InRect.Width();
+	const uint32 SizeY = InRect.Height();
+
+	// Check the format of the surface
+	FIntRect StagingRect;
+	D3D12_PLACED_SUBRESOURCE_FOOTPRINT readBackHeapDesc;
+	TRefCountPtr<FD3D12Resource> TempTexture2D = GetStagingTexture(TextureRHI, InRect, StagingRect, InFlags, readBackHeapDesc);
+
+	uint32 BytesPerPixel = GPixelFormats[TextureRHI->GetFormat()].BlockBytes;
+
+	// Allocate the output buffer.
+	OutData.Empty();
+	OutData.AddUninitialized(SizeX * SizeY * BytesPerPixel);
+
+	// Lock the staging resource.
+	void* pData;
+	VERIFYD3D11RESULT(TempTexture2D->GetResource()->Map(0, nullptr, &pData));
+
+	uint32 BytesPerLine = BytesPerPixel * InRect.Width();
+
+	const uint32 XBytesAligned = Align((uint32)readBackHeapDesc.Footprint.Width * BytesPerPixel, D3D12_TEXTURE_DATA_PITCH_ALIGNMENT);
+
+	uint8* DestPtr = OutData.GetData();
+	uint8* SrcPtr = (uint8*)pData + StagingRect.Min.X * BytesPerPixel + StagingRect.Min.Y * XBytesAligned;
+	for (uint32 Y = 0; Y < SizeY; Y++)
+	{
+		memcpy(DestPtr, SrcPtr, BytesPerLine);
+		DestPtr += BytesPerLine;
+		SrcPtr += XBytesAligned;
+	}
+
+	TempTexture2D->GetResource()->Unmap(0, nullptr);
+}
+
+/** Helper for accessing R10G10B10A2 colors. */
+struct FD3DR10G10B10A2
+{
+	uint32 R : 10;
+	uint32 G : 10;
+	uint32 B : 10;
+	uint32 A : 2;
+};
+
+/** Helper for accessing R16G16 colors. */
+struct FD3DRG16
+{
+	uint16 R;
+	uint16 G;
+};
+
+/** Helper for accessing R16G16B16A16 colors. */
+struct FD3DRGBA16
+{
+	uint16 R;
+	uint16 G;
+	uint16 B;
+	uint16 A;
+};
+
+// todo: this should be available for all RHI
+static void ConvertRAWSurfaceDataToFColor(DXGI_FORMAT Format, uint32 Width, uint32 Height, uint8 *In, uint32 SrcPitch, FColor* Out, FReadSurfaceDataFlags InFlags)
+{
+	bool bLinearToGamma = InFlags.GetLinearToGamma();
+
+	if (Format == DXGI_FORMAT_R16_TYPELESS)
+	{
+		// e.g. shadow maps
+		for (uint32 Y = 0; Y < Height; Y++)
+		{
+			uint16* SrcPtr = (uint16*)(In + Y * SrcPitch);
+			FColor* DestPtr = Out + Y * Width;
+
+			for (uint32 X = 0; X < Width; X++)
+			{
+				uint16 Value16 = *SrcPtr;
+				float Value = Value16 / (float)(0xffff);
+
+				*DestPtr = FLinearColor(Value, Value, Value).Quantize();
+				++SrcPtr;
+				++DestPtr;
+			}
+		}
+	}
+	else if (Format == DXGI_FORMAT_R8G8B8A8_TYPELESS || Format == DXGI_FORMAT_R8G8B8A8_UNORM)
+	{
+		// Read the data out of the buffer, converting it from ABGR to ARGB.
+		for (uint32 Y = 0; Y < Height; Y++)
+		{
+			FColor* SrcPtr = (FColor*)(In + Y * SrcPitch);
+			FColor* DestPtr = Out + Y * Width;
+			for (uint32 X = 0; X < Width; X++)
+			{
+				*DestPtr = FColor(SrcPtr->B, SrcPtr->G, SrcPtr->R, SrcPtr->A);
+				++SrcPtr;
+				++DestPtr;
+			}
+		}
+	}
+	else if (Format == DXGI_FORMAT_B8G8R8A8_TYPELESS || Format == DXGI_FORMAT_B8G8R8A8_UNORM || Format == DXGI_FORMAT_B8G8R8A8_UNORM_SRGB)
+	{
+		for (uint32 Y = 0; Y < Height; Y++)
+		{
+			FColor* SrcPtr = (FColor*)(In + Y * SrcPitch);
+			FColor* DestPtr = Out + Y * Width;
+
+			// Need to copy row wise since the Pitch might not match the Width.
+			FMemory::Memcpy(DestPtr, SrcPtr, sizeof(FColor) * Width);
+		}
+	}
+	else if (Format == DXGI_FORMAT_R10G10B10A2_UNORM)
+	{
+		// Read the data out of the buffer, converting it from R10G10B10A2 to FColor.
+		for (uint32 Y = 0; Y < Height; Y++)
+		{
+			FD3DR10G10B10A2* SrcPtr = (FD3DR10G10B10A2*)(In + Y * SrcPitch);
+			FColor* DestPtr = Out + Y * Width;
+			for (uint32 X = 0; X < Width; X++)
+			{
+				*DestPtr = FLinearColor(
+					(float)SrcPtr->R / 1023.0f,
+					(float)SrcPtr->G / 1023.0f,
+					(float)SrcPtr->B / 1023.0f,
+					(float)SrcPtr->A / 3.0f
+					).Quantize();
+				++SrcPtr;
+				++DestPtr;
+			}
+		}
+	}
+	else if (Format == DXGI_FORMAT_R16G16B16A16_FLOAT)
+	{
+		FPlane	MinValue(0.0f, 0.0f, 0.0f, 0.0f),
+			MaxValue(1.0f, 1.0f, 1.0f, 1.0f);
+
+		check(sizeof(FD3DFloat16) == sizeof(uint16));
+
+		for (uint32 Y = 0; Y < Height; Y++)
+		{
+			FD3DFloat16* SrcPtr = (FD3DFloat16*)(In + Y * SrcPitch);
+
+			for (uint32 X = 0; X < Width; X++)
+			{
+				MinValue.X = FMath::Min<float>(SrcPtr[0], MinValue.X);
+				MinValue.Y = FMath::Min<float>(SrcPtr[1], MinValue.Y);
+				MinValue.Z = FMath::Min<float>(SrcPtr[2], MinValue.Z);
+				MinValue.W = FMath::Min<float>(SrcPtr[3], MinValue.W);
+				MaxValue.X = FMath::Max<float>(SrcPtr[0], MaxValue.X);
+				MaxValue.Y = FMath::Max<float>(SrcPtr[1], MaxValue.Y);
+				MaxValue.Z = FMath::Max<float>(SrcPtr[2], MaxValue.Z);
+				MaxValue.W = FMath::Max<float>(SrcPtr[3], MaxValue.W);
+				SrcPtr += 4;
+			}
+		}
+
+		for (uint32 Y = 0; Y < Height; Y++)
+		{
+			FD3DFloat16* SrcPtr = (FD3DFloat16*)(In + Y * SrcPitch);
+			FColor* DestPtr = Out + Y * Width;
+
+			for (uint32 X = 0; X < Width; X++)
+			{
+				FColor NormalizedColor =
+					FLinearColor(
+					(SrcPtr[0] - MinValue.X) / (MaxValue.X - MinValue.X),
+					(SrcPtr[1] - MinValue.Y) / (MaxValue.Y - MinValue.Y),
+					(SrcPtr[2] - MinValue.Z) / (MaxValue.Z - MinValue.Z),
+					(SrcPtr[3] - MinValue.W) / (MaxValue.W - MinValue.W)
+					).ToFColor(bLinearToGamma);
+				FMemory::Memcpy(DestPtr++, &NormalizedColor, sizeof(FColor));
+				SrcPtr += 4;
+			}
+		}
+	}
+	else if (Format == DXGI_FORMAT_R11G11B10_FLOAT)
+	{
+		check(sizeof(FD3DFloatR11G11B10) == sizeof(uint32));
+
+		for (uint32 Y = 0; Y < Height; Y++)
+		{
+			FD3DFloatR11G11B10* SrcPtr = (FD3DFloatR11G11B10*)(In + Y * SrcPitch);
+			FColor* DestPtr = Out + Y * Width;
+
+			for (uint32 X = 0; X < Width; X++)
+			{
+				FLinearColor Value = *SrcPtr;
+
+				FColor NormalizedColor = Value.ToFColor(bLinearToGamma);
+				FMemory::Memcpy(DestPtr++, &NormalizedColor, sizeof(FColor));
+				++SrcPtr;
+			}
+		}
+	}
+	else if (Format == DXGI_FORMAT_R32G32B32A32_FLOAT)
+	{
+		FPlane MinValue(0.0f, 0.0f, 0.0f, 0.0f);
+		FPlane MaxValue(1.0f, 1.0f, 1.0f, 1.0f);
+
+		for (uint32 Y = 0; Y < Height; Y++)
+		{
+			float* SrcPtr = (float*)(In + Y * SrcPitch);
+
+			for (uint32 X = 0; X < Width; X++)
+			{
+				MinValue.X = FMath::Min<float>(SrcPtr[0], MinValue.X);
+				MinValue.Y = FMath::Min<float>(SrcPtr[1], MinValue.Y);
+				MinValue.Z = FMath::Min<float>(SrcPtr[2], MinValue.Z);
+				MinValue.W = FMath::Min<float>(SrcPtr[3], MinValue.W);
+				MaxValue.X = FMath::Max<float>(SrcPtr[0], MaxValue.X);
+				MaxValue.Y = FMath::Max<float>(SrcPtr[1], MaxValue.Y);
+				MaxValue.Z = FMath::Max<float>(SrcPtr[2], MaxValue.Z);
+				MaxValue.W = FMath::Max<float>(SrcPtr[3], MaxValue.W);
+				SrcPtr += 4;
+			}
+		}
+
+		for (uint32 Y = 0; Y < Height; Y++)
+		{
+			float* SrcPtr = (float*)In;
+			FColor* DestPtr = Out + Y * Width;
+
+			for (uint32 X = 0; X < Width; X++)
+			{
+				FColor NormalizedColor =
+					FLinearColor(
+					(SrcPtr[0] - MinValue.X) / (MaxValue.X - MinValue.X),
+					(SrcPtr[1] - MinValue.Y) / (MaxValue.Y - MinValue.Y),
+					(SrcPtr[2] - MinValue.Z) / (MaxValue.Z - MinValue.Z),
+					(SrcPtr[3] - MinValue.W) / (MaxValue.W - MinValue.W)
+					).ToFColor(bLinearToGamma);
+				FMemory::Memcpy(DestPtr++, &NormalizedColor, sizeof(FColor));
+				SrcPtr += 4;
+			}
+		}
+	}
+	else if (Format == DXGI_FORMAT_R24G8_TYPELESS)
+	{
+		// Depth stencil
+		for (uint32 Y = 0; Y < Height; Y++)
+		{
+			uint32* SrcPtr = (uint32 *)In;
+			FColor* DestPtr = Out + Y * Width;
+
+			for (uint32 X = 0; X < Width; X++)
+			{
+				FColor NormalizedColor;
+				if (InFlags.GetOutputStencil())
+				{
+					uint8 DeviceStencil = (*SrcPtr & 0xFF000000) >> 24;
+					NormalizedColor = FColor(DeviceStencil, DeviceStencil, DeviceStencil, 0xFF);
+				}
+				else
+				{
+					float DeviceZ = (*SrcPtr & 0xffffff) / (float)(1 << 24);
+					float LinearValue = FMath::Min(InFlags.ComputeNormalizedDepth(DeviceZ), 1.0f);
+					NormalizedColor = FLinearColor(LinearValue, LinearValue, LinearValue, 0).ToFColor(bLinearToGamma);
+				}
+
+				FMemory::Memcpy(DestPtr++, &NormalizedColor, sizeof(FColor));
+				++SrcPtr;
+			}
+		}
+	}
+#if DEPTH_32_BIT_CONVERSION
+	// Changing Depth Buffers to 32 bit on Dingo as D24S8 is actually implemented as a 32 bit buffer in the hardware
+	else if (Format == DXGI_FORMAT_R32G8X24_TYPELESS )
+	{
+		// Depth stencil
+		for(uint32 Y = 0; Y < Height; Y++)
+		{
+			float* SrcPtr = (float *)(In + Y * SrcPitch);
+			FColor* DestPtr = Out + Y * Width;
+
+			for(uint32 X = 0; X < Width; X++)
+			{
+				float DeviceZ = (*SrcPtr);
+
+				float LinearValue = FMath::Min(InFlags.ComputeNormalizedDepth(DeviceZ), 1.0f);
+
+				FColor NormalizedColor = FLinearColor(LinearValue, LinearValue, LinearValue, 0).ToFColor(bLinearToGamma);
+				FMemory::Memcpy(DestPtr++, &NormalizedColor, sizeof(FColor));
+				SrcPtr+=1; // todo: copies only depth, need to check how this format is read
+				UE_LOG(LogD3D12RHI, Warning, TEXT("CPU read of R32G8X24 is not tested and may not function."));
+			}
+		}
+	}
+#endif
+	else if (Format == DXGI_FORMAT_R16G16B16A16_UNORM)
+	{
+		// Read the data out of the buffer, converting it to FColor.
+		for (uint32 Y = 0; Y < Height; Y++)
+		{
+			FD3DRGBA16* SrcPtr = (FD3DRGBA16*)(In + Y * SrcPitch);
+			FColor* DestPtr = Out + Y * Width;
+			for (uint32 X = 0; X < Width; X++)
+			{
+				*DestPtr = FLinearColor(
+					(float)SrcPtr->R / 65535.0f,
+					(float)SrcPtr->G / 65535.0f,
+					(float)SrcPtr->B / 65535.0f,
+					(float)SrcPtr->A / 65535.0f
+					).Quantize();
+				++SrcPtr;
+				++DestPtr;
+			}
+		}
+	}
+	else if (Format == DXGI_FORMAT_R16G16_UNORM)
+	{
+		// Read the data out of the buffer, converting it to FColor.
+		for (uint32 Y = 0; Y < Height; Y++)
+		{
+			FD3DRG16* SrcPtr = (FD3DRG16*)(In + Y * SrcPitch);
+			FColor* DestPtr = Out + Y * Width;
+			for (uint32 X = 0; X < Width; X++)
+			{
+				*DestPtr = FLinearColor(
+					(float)SrcPtr->R / 65535.0f,
+					(float)SrcPtr->G / 65535.0f,
+					0).Quantize();
+				++SrcPtr;
+				++DestPtr;
+			}
+		}
+	}
+	else
+	{
+		// not supported yet
+		check(0);
+	}
+}
+
+void FD3D12DynamicRHI::RHIReadSurfaceData(FTextureRHIParamRef TextureRHI, FIntRect InRect, TArray<FColor>& OutData, FReadSurfaceDataFlags InFlags)
+{
+	TArray<uint8> OutDataRaw;
+
+	FD3D12TextureBase* Texture = GetD3D11TextureFromRHITexture(TextureRHI);
+
+    // Wait for the command list if needed
+    FD3D12Texture2D* DestTexture2D = static_cast<FD3D12Texture2D*>(TextureRHI->GetTexture2D());
+    FD3D12CLSyncPoint SyncPoint = DestTexture2D->GetReadBackSyncPoint();
+
+	if (!!SyncPoint)
+	{
+		CommandListState listState = GetRHIDevice()->GetCommandListManager().GetCommandListState(SyncPoint);
+		if (listState == CommandListState::kOpen)
+			GetRHIDevice()->GetDefaultCommandContext().FlushCommands(true);
+		else
+			GetRHIDevice()->GetCommandListManager().WaitForCompletion(SyncPoint);
+	}
+
+	// Check the format of the surface
+	D3D12_RESOURCE_DESC const& TextureDesc = Texture->GetResource()->GetDesc();
+
+	check(TextureDesc.SampleDesc.Count >= 1);
+
+	if (TextureDesc.SampleDesc.Count == 1)
+	{
+		ReadSurfaceDataNoMSAARaw(TextureRHI, InRect, OutDataRaw, InFlags);
+	}
+	else
+	{
+		FRHICommandList_RecursiveHazardous RHICmdList(RHIGetDefaultContext());
+		ReadSurfaceDataMSAARaw(RHICmdList, TextureRHI, InRect, OutDataRaw, InFlags);
+	}
+
+	const uint32 SizeX = InRect.Width() * TextureDesc.SampleDesc.Count;
+	const uint32 SizeY = InRect.Height();
+
+	// Allocate the output buffer.
+	OutData.Empty();
+	OutData.AddUninitialized(SizeX * SizeY);
+
+	FPixelFormatInfo FormatInfo = GPixelFormats[TextureRHI->GetFormat()];
+	uint32 BytesPerPixel = FormatInfo.BlockBytes;
+	uint32 SrcPitch = SizeX * BytesPerPixel;
+
+	ConvertRAWSurfaceDataToFColor((DXGI_FORMAT)FormatInfo.PlatformFormat, SizeX, SizeY, OutDataRaw.GetData(), SrcPitch, OutData.GetData(), InFlags);
+}
+
+void FD3D12DynamicRHI::ReadSurfaceDataMSAARaw(FRHICommandList_RecursiveHazardous& RHICmdList, FTextureRHIParamRef TextureRHI, FIntRect InRect, TArray<uint8>& OutData, FReadSurfaceDataFlags InFlags)
+{
+    FD3D12CommandContext& DefaultContext = GetRHIDevice()->GetDefaultCommandContext();
+	FD3D12CommandListHandle& hCommandList = DefaultContext.CommandListHandle;
+	FD3D12TextureBase* Texture = GetD3D11TextureFromRHITexture(TextureRHI);
+
+	const uint32 SizeX = InRect.Width();
+	const uint32 SizeY = InRect.Height();
+
+	// Check the format of the surface
+	D3D12_RESOURCE_DESC const& TextureDesc = Texture->GetResource()->GetDesc();
+
+	uint32 BytesPerPixel = ComputeBytesPerPixel(TextureDesc.Format);
+
+	const uint32 NumSamples = TextureDesc.SampleDesc.Count;
+
+	// Read back the surface data from the define rect
+	D3D12_BOX	Rect;
+	Rect.left = InRect.Min.X;
+	Rect.top = InRect.Min.Y;
+	Rect.right = InRect.Max.X;
+	Rect.bottom = InRect.Max.Y;
+	Rect.back = 1;
+	Rect.front = 0;
+
+	// Create a non-MSAA render target to resolve individual samples of the source surface to.
+	D3D12_RESOURCE_DESC NonMSAADesc;
+	ZeroMemory(&NonMSAADesc, sizeof(NonMSAADesc));
+	NonMSAADesc.Dimension = D3D12_RESOURCE_DIMENSION_TEXTURE2D;
+	NonMSAADesc.Width = SizeX;
+	NonMSAADesc.Height = SizeY;
+	NonMSAADesc.MipLevels = 1;
+	NonMSAADesc.DepthOrArraySize = 1;
+	NonMSAADesc.Format = TextureDesc.Format;
+	NonMSAADesc.SampleDesc.Count = 1;
+	NonMSAADesc.SampleDesc.Quality = 0;
+	NonMSAADesc.Flags = D3D12_RESOURCE_FLAG_ALLOW_RENDER_TARGET;
+	TRefCountPtr<FD3D12Resource> NonMSAATexture2D;
+    VERIFYD3D11RESULT(GetRHIDevice()->GetResourceHelper().CreateCommittedResource(NonMSAADesc, CD3DX12_HEAP_PROPERTIES(D3D12_HEAP_TYPE_DEFAULT), NonMSAATexture2D.GetInitReference()));
+
+	TRefCountPtr<FD3D12RenderTargetView> NonMSAARTV;
+	D3D12_RENDER_TARGET_VIEW_DESC RTVDesc;
+	FMemory::Memset(&RTVDesc, 0, sizeof(RTVDesc));
+
+	// typeless is not supported, similar code might be needed for other typeless formats
+	RTVDesc.Format = ConvertTypelessToUnorm(NonMSAADesc.Format);
+
+	RTVDesc.ViewDimension = D3D12_RTV_DIMENSION_TEXTURE2D;
+	RTVDesc.Texture2D.MipSlice = 0;
+    NonMSAARTV = new FD3D12RenderTargetView(GetRHIDevice(), &RTVDesc, NonMSAATexture2D);
+
+	// Create a CPU-accessible staging texture to copy the resolved sample data to.
+	TRefCountPtr<FD3D12Resource> StagingTexture2D;
+	const uint32 BlockBytes = GPixelFormats[TextureRHI->GetFormat()].BlockBytes;
+	const uint32 XBytesAligned = Align(SizeX * BlockBytes, D3D12_TEXTURE_DATA_PITCH_ALIGNMENT);
+	const uint32 MipBytesAligned = XBytesAligned * SizeY;
+    VERIFYD3D11RESULT(GetRHIDevice()->GetResourceHelper().CreateBuffer(D3D12_HEAP_TYPE_READBACK, MipBytesAligned, StagingTexture2D.GetInitReference()));
+
+    // Ensure we're dealing with a Texture2D, which the rest of this function already assumes
+    check(TextureRHI->GetTexture2D());
+    FD3D12Texture2D *InTexture2D = static_cast<FD3D12Texture2D*>(Texture);
+
+	// Determine the subresource index for cubemaps.
+	uint32 Subresource = 0;
+    if (InTexture2D->IsCubemap())
+	{
+		uint32 D3DFace = GetD3D11CubeFace(InFlags.GetCubeFace());
+		Subresource = CalcSubresource(0, D3DFace, 1);
+	}
+
+	// Setup the descriptions for the copy to the readback heap.
+	D3D12_SUBRESOURCE_FOOTPRINT destSubresource;
+	destSubresource.Depth = 1;
+	destSubresource.Height = SizeY;
+	destSubresource.Width = SizeX;
+	destSubresource.Format = TextureDesc.Format;
+	destSubresource.RowPitch = XBytesAligned;
+	check(destSubresource.RowPitch % D3D12_TEXTURE_DATA_PITCH_ALIGNMENT == 0);	// Make sure we align correctly.
+
+	D3D12_PLACED_SUBRESOURCE_FOOTPRINT placedTexture2D = { 0 };
+	placedTexture2D.Offset = 0;
+	placedTexture2D.Footprint = destSubresource;
+
+    CD3DX12_TEXTURE_COPY_LOCATION DestCopyLocation(StagingTexture2D->GetResource(), placedTexture2D);
+    CD3DX12_TEXTURE_COPY_LOCATION SourceCopyLocation(NonMSAATexture2D->GetResource(), Subresource);
+
+	// Allocate the output buffer.
+	OutData.Empty();
+	OutData.AddUninitialized(SizeX * SizeY * NumSamples * BytesPerPixel);
+
+	// Can be optimized by doing all subsamples into a large enough rendertarget in one pass (multiple draw calls)
+	for (uint32 SampleIndex = 0; SampleIndex < NumSamples; ++SampleIndex)
+	{
+		// Resolve the sample to the non-MSAA render target.
+		DefaultContext.ResolveTextureUsingShader<FResolveSingleSamplePS>(
+			RHICmdList,
+			(FD3D12Texture2D*)TextureRHI->GetTexture2D(),
+			NULL,
+			NonMSAARTV,
+			NULL,
+			NonMSAADesc,
+			FResolveRect(InRect.Min.X, InRect.Min.Y, InRect.Max.X, InRect.Max.Y),
+			FResolveRect(0, 0, SizeX, SizeY),
+			SampleIndex
+			);
+
+		TransitionResource(hCommandList, NonMSAATexture2D, D3D12_RESOURCE_STATE_COPY_SOURCE, SourceCopyLocation.SubresourceIndex);
+		// Upload heap doesn't need to transition
+
+		DefaultContext.numCopies++;
+		// Copy the resolved sample data to the staging texture.
+		hCommandList->CopyTextureRegion(
+            &DestCopyLocation,
+            0, 0, 0,
+            &SourceCopyLocation,
+            &Rect);
+
+		// We need to execute the command list so we can read the data in the map below
+        GetRHIDevice()->GetDefaultCommandContext().FlushCommands(true);
+
+		// Lock the staging texture.
+		void* pData;
+        VERIFYD3D11RESULT(StagingTexture2D->GetResource()->Map(0, nullptr, &pData));
+
+		// Read the data out of the buffer, could be optimized
+		for (int32 Y = InRect.Min.Y; Y < InRect.Max.Y; Y++)
+		{
+			uint8* SrcPtr = (uint8*)pData + (Y - InRect.Min.Y) * XBytesAligned + InRect.Min.X * BytesPerPixel;
+			uint8* DestPtr = &OutData[(Y - InRect.Min.Y) * SizeX * NumSamples * BytesPerPixel + SampleIndex * BytesPerPixel];
+
+			for (int32 X = InRect.Min.X; X < InRect.Max.X; X++)
+			{
+				for (uint32 i = 0; i < BytesPerPixel; ++i)
+				{
+					*DestPtr++ = *SrcPtr++;
+				}
+
+				DestPtr += (NumSamples - 1) * BytesPerPixel;
+			}
+		}
+
+		StagingTexture2D->GetResource()->Unmap(0, nullptr);
+	}
+}
+
+void FD3D12DynamicRHI::RHIMapStagingSurface(FTextureRHIParamRef TextureRHI, void*& OutData, int32& OutWidth, int32& OutHeight)
+{
+	FD3D12Resource* Texture = GetD3D11TextureFromRHITexture(TextureRHI)->GetResource();
+
+	DXGI_FORMAT Format = (DXGI_FORMAT)GPixelFormats[TextureRHI->GetFormat()].PlatformFormat;
+
+	uint32 BytesPerPixel = ComputeBytesPerPixel(Format);
+
+	// Wait for the command list if needed
+	FD3D12Texture2D* DestTexture2D = static_cast<FD3D12Texture2D*>(TextureRHI->GetTexture2D());
+	FD3D12CLSyncPoint SyncPoint = DestTexture2D->GetReadBackSyncPoint();
+	CommandListState listState = GetRHIDevice()->GetCommandListManager().GetCommandListState(SyncPoint);
+	if (listState == CommandListState::kOpen)
+        GetRHIDevice()->GetDefaultCommandContext().FlushCommands(true);
+	else
+        GetRHIDevice()->GetCommandListManager().WaitForCompletion(SyncPoint);
+
+	void* pData;
+	HRESULT Result = Texture->GetResource()->Map(0, nullptr, &pData);
+	if (Result == DXGI_ERROR_DEVICE_REMOVED)
+	{
+		// When reading back to the CPU, we have to watch out for DXGI_ERROR_DEVICE_REMOVED
+        GetRHIDevice()->SetDeviceRemoved(true);
+
+		OutData = NULL;
+		OutWidth = OutHeight = 0;
+
+        HRESULT hRes = GetRHIDevice()->GetDevice()->GetDeviceRemovedReason();
+
+		UE_LOG(LogD3D12RHI, Warning, TEXT("FD3D12DynamicRHI::RHIMapStagingSurface failed (GetDeviceRemovedReason(): %d)"), hRes);
+	}
+	else
+	{
+        VERIFYD3D11RESULT_EX(Result, GetRHIDevice()->GetDevice());
+
+		const D3D12_PLACED_SUBRESOURCE_FOOTPRINT& readBackHeapDesc = DestTexture2D->GetReadBackHeapDesc();
+		OutData = pData;
+		OutWidth = readBackHeapDesc.Footprint.Width;
+		OutHeight = readBackHeapDesc.Footprint.Height;
+
+		// MS: It seems like the second frame in some scenes comes into RHIMapStagingSurface BEFORE the copy to the staging texture, thus the readbackHeapDesc isn't set. This could be bug in UE4.
+		if (readBackHeapDesc.Footprint.Format != DXGI_FORMAT_UNKNOWN)
+		{
+			check(OutWidth != 0);
+			check(OutHeight != 0);
+		}
+
+		check(OutData);
+	}
+}
+
+void FD3D12DynamicRHI::RHIUnmapStagingSurface(FTextureRHIParamRef TextureRHI)
+{
+	ID3D12Resource* Texture = GetD3D11TextureFromRHITexture(TextureRHI)->GetResource()->GetResource();
+
+	Texture->Unmap(0, nullptr);
+}
+
+void FD3D12DynamicRHI::RHIReadSurfaceFloatData(FTextureRHIParamRef TextureRHI, FIntRect InRect, TArray<FFloat16Color>& OutData, ECubeFace CubeFace, int32 ArrayIndex, int32 MipIndex)
+{
+    FD3D12CommandContext& DefaultContext = GetRHIDevice()->GetDefaultCommandContext();
+	FD3D12CommandListHandle& hCommandList = DefaultContext.CommandListHandle;
+	FD3D12TextureBase* Texture = GetD3D11TextureFromRHITexture(TextureRHI);
+
+	uint32 SizeX = InRect.Width();
+	uint32 SizeY = InRect.Height();
+
+	// Check the format of the surface
+	D3D12_RESOURCE_DESC const& TextureDesc = Texture->GetResource()->GetDesc();
+
+	check(TextureDesc.Format == GPixelFormats[PF_FloatRGBA].PlatformFormat);
+
+	// Allocate the output buffer.
+	OutData.Empty(SizeX * SizeY);
+
+	// Read back the surface data from defined rect
+	D3D12_BOX	Rect;
+	Rect.left = InRect.Min.X;
+	Rect.top = InRect.Min.Y;
+	Rect.right = InRect.Max.X;
+	Rect.bottom = InRect.Max.Y;
+	Rect.back = 1;
+	Rect.front = 0;
+
+	// create a temp 2d texture to copy render target to
+	TRefCountPtr<FD3D12Resource> TempTexture2D;
+	const uint32 BlockBytes = GPixelFormats[TextureRHI->GetFormat()].BlockBytes;
+	const uint32 XBytesAligned = Align((uint32)TextureDesc.Width * BlockBytes, D3D12_TEXTURE_DATA_PITCH_ALIGNMENT);
+	const uint32 MipBytesAligned = XBytesAligned * TextureDesc.Height;
+    VERIFYD3D11RESULT(GetRHIDevice()->GetResourceHelper().CreateBuffer(D3D12_HEAP_TYPE_READBACK, MipBytesAligned, TempTexture2D.GetInitReference()));
+
+    // Ensure we're dealing with a Texture2D, which the rest of this function already assumes
+	bool bIsTextureCube = false;
+    check(TextureRHI->GetTexture2D() || TextureRHI->GetTexture2DArray() || TextureRHI->GetTextureCube());
+    FD3D12Texture2D *InTexture2D = static_cast<FD3D12Texture2D*>(Texture);
+	FD3D12Texture2DArray *InTexture2DArray = static_cast<FD3D12Texture2DArray*>(Texture);
+	FD3D12TextureCube *InTextureCube = static_cast<FD3D12TextureCube*>(Texture);
+	if (InTexture2D)
+	{
+		bIsTextureCube = InTexture2D->IsCubemap();
+	}
+	else if (InTexture2DArray)
+	{
+		bIsTextureCube = InTexture2DArray->IsCubemap();
+	}
+	else if (InTextureCube)
+	{
+		bIsTextureCube = InTextureCube->IsCubemap();
+		check(bIsTextureCube);
+	}
+	else
+	{
+		check(false);
+	}
+
+	// Copy the data to a staging resource.
+	uint32 Subresource = 0;
+	if (bIsTextureCube)
+	{
+		uint32 D3DFace = GetD3D11CubeFace(CubeFace);
+		Subresource = CalcSubresource(MipIndex, ArrayIndex * 6 + D3DFace, TextureDesc.MipLevels);
+	}
+
+	uint32 BytesPerPixel = ComputeBytesPerPixel(TextureDesc.Format);
+	D3D12_SUBRESOURCE_FOOTPRINT destSubresource;
+	destSubresource.Depth = 1;
+	destSubresource.Height = TextureDesc.Height;
+	destSubresource.Width = TextureDesc.Width;
+	destSubresource.Format = TextureDesc.Format;
+	destSubresource.RowPitch = XBytesAligned;
+	check(destSubresource.RowPitch % D3D12_TEXTURE_DATA_PITCH_ALIGNMENT == 0);	// Make sure we align correctly.
+
+	D3D12_PLACED_SUBRESOURCE_FOOTPRINT placedTexture2D = { 0 };
+	placedTexture2D.Offset = 0;
+	placedTexture2D.Footprint = destSubresource;
+
+    CD3DX12_TEXTURE_COPY_LOCATION DestCopyLocation(TempTexture2D->GetResource(), placedTexture2D);
+	CD3DX12_TEXTURE_COPY_LOCATION SourceCopyLocation(Texture->GetResource()->GetResource(), Subresource);
+
+	{
+		FConditionalScopeResourceBarrier ConditionalScopeResourceBarrier(hCommandList, Texture->GetResource(), D3D12_RESOURCE_STATE_COPY_SOURCE, SourceCopyLocation.SubresourceIndex);
+		// Don't need to transition upload heaps
+
+		DefaultContext.numCopies++;
+		hCommandList->CopyTextureRegion(
+			&DestCopyLocation,
+			0, 0, 0,
+			&SourceCopyLocation,
+			&Rect);
+	}
+
+	// We need to execute the command list so we can read the data from the map below
+    GetRHIDevice()->GetDefaultCommandContext().FlushCommands(true);
+
+	// Lock the staging resource.
+	void* pData;
+	VERIFYD3D11RESULT(TempTexture2D->GetResource()->Map(0, nullptr, &pData));
+
+	// Presize the array
+	int32 TotalCount = SizeX * SizeY;
+	if (TotalCount >= OutData.Num())
+	{
+		OutData.AddZeroed(TotalCount);
+	}
+
+	for (int32 Y = InRect.Min.Y; Y < InRect.Max.Y; Y++)
+	{
+		FFloat16Color* SrcPtr = (FFloat16Color*)((uint8*)pData + (Y - InRect.Min.Y) * XBytesAligned);
+		int32 Index = (Y - InRect.Min.Y) * SizeX;
+		check(Index + ((int32)SizeX - 1) < OutData.Num());
+		FFloat16Color* DestColor = &OutData[Index];
+		FFloat16* DestPtr = (FFloat16*)(DestColor);
+		FMemory::Memcpy(DestPtr, SrcPtr, SizeX * sizeof(FFloat16) * 4);
+	}
+
+	TempTexture2D->GetResource()->Unmap(0, nullptr);
+}
+
+void FD3D12DynamicRHI::RHIRead3DSurfaceFloatData(FTextureRHIParamRef TextureRHI, FIntRect InRect, FIntPoint ZMinMax, TArray<FFloat16Color>& OutData)
+{
+    FD3D12CommandContext& DefaultContext = GetRHIDevice()->GetDefaultCommandContext();
+	FD3D12CommandListHandle& hCommandList = DefaultContext.CommandListHandle;
+	FD3D12TextureBase* Texture = GetD3D11TextureFromRHITexture(TextureRHI);
+
+	uint32 SizeX = InRect.Width();
+	uint32 SizeY = InRect.Height();
+	uint32 SizeZ = ZMinMax.Y - ZMinMax.X;
+
+	// Check the format of the surface
+	D3D12_RESOURCE_DESC const& TextureDesc11 = Texture->GetResource()->GetDesc();
+	check(TextureDesc11.Format == GPixelFormats[PF_FloatRGBA].PlatformFormat);
+
+	// Allocate the output buffer.
+	OutData.Empty(SizeX * SizeY * SizeZ * sizeof(FFloat16Color));
+
+	// Read back the surface data from defined rect
+	D3D12_BOX	Rect;
+	Rect.left = InRect.Min.X;
+	Rect.top = InRect.Min.Y;
+	Rect.right = InRect.Max.X;
+	Rect.bottom = InRect.Max.Y;
+	Rect.back = ZMinMax.Y;
+	Rect.front = ZMinMax.X;
+
+	// create a temp 3d texture to copy render target to
+	TRefCountPtr<FD3D12Resource> TempTexture3D;
+	const uint32 BlockBytes = GPixelFormats[TextureRHI->GetFormat()].BlockBytes;
+	const uint32 XBytesAligned = Align(TextureDesc11.Width * BlockBytes, D3D12_TEXTURE_DATA_PITCH_ALIGNMENT);
+	const uint32 DepthBytesAligned = XBytesAligned * TextureDesc11.Height;
+	const uint32 MipBytesAligned = DepthBytesAligned * TextureDesc11.DepthOrArraySize;
+    VERIFYD3D11RESULT(GetRHIDevice()->GetResourceHelper().CreateBuffer(D3D12_HEAP_TYPE_READBACK, MipBytesAligned, TempTexture3D.GetInitReference()));
+
+	// Copy the data to a staging resource.
+	uint32 Subresource = 0;
+	uint32 BytesPerPixel = ComputeBytesPerPixel(TextureDesc11.Format);
+	D3D12_SUBRESOURCE_FOOTPRINT destSubresource;
+	destSubresource.Depth = TextureDesc11.DepthOrArraySize;
+	destSubresource.Height = TextureDesc11.Height;
+	destSubresource.Width = TextureDesc11.Width;
+	destSubresource.Format = TextureDesc11.Format;
+	destSubresource.RowPitch = XBytesAligned;
+	check(destSubresource.RowPitch % D3D12_TEXTURE_DATA_PITCH_ALIGNMENT == 0);	// Make sure we align correctly.
+
+	D3D12_PLACED_SUBRESOURCE_FOOTPRINT placedTexture3D = { 0 };
+	placedTexture3D.Offset = 0;
+	placedTexture3D.Footprint = destSubresource;
+
+    CD3DX12_TEXTURE_COPY_LOCATION DestCopyLocation(TempTexture3D->GetResource(), placedTexture3D);
+	CD3DX12_TEXTURE_COPY_LOCATION SourceCopyLocation(Texture->GetResource()->GetResource(), Subresource);
+
+	{
+		FConditionalScopeResourceBarrier ConditionalScopeResourceBarrier(hCommandList, Texture->GetResource(), D3D12_RESOURCE_STATE_COPY_SOURCE, SourceCopyLocation.SubresourceIndex);
+		// Don't need to transition upload heaps
+
+		DefaultContext.numCopies++;
+		hCommandList->CopyTextureRegion(
+			&DestCopyLocation,
+			0, 0, 0,
+			&SourceCopyLocation,
+			&Rect);
+	}
+
+	// We need to execute the command list so we can read the data from the map below
+    GetRHIDevice()->GetDefaultCommandContext().FlushCommands(true);
+
+	// Lock the staging resource.
+	void* pData;
+	VERIFYD3D11RESULT(TempTexture3D->GetResource()->Map(0, nullptr, &pData));
+
+	// Presize the array
+	int32 TotalCount = SizeX * SizeY * SizeZ;
+	if (TotalCount >= OutData.Num())
+	{
+		OutData.AddZeroed(TotalCount);
+	}
+
+	// Read the data out of the buffer, converting it from ABGR to ARGB.
+	for (int32 Z = ZMinMax.X; Z < ZMinMax.Y; ++Z)
+	{
+		for (int32 Y = InRect.Min.Y; Y < InRect.Max.Y; ++Y)
+		{
+			FFloat16Color* SrcPtr = (FFloat16Color*)((uint8*)pData + (Y - InRect.Min.Y) * XBytesAligned + (Z - ZMinMax.X) * DepthBytesAligned);
+			int32 Index = (Y - InRect.Min.Y) * SizeX + (Z - ZMinMax.X) * SizeX * SizeY;
+			check(Index < OutData.Num());
+			FFloat16Color* DestColor = &OutData[Index];
+			FFloat16* DestPtr = (FFloat16*)(DestColor);
+			FMemory::Memcpy(DestPtr, SrcPtr, SizeX * sizeof(FFloat16) * 4);
+		}
+	}
+
+	TempTexture3D->GetResource()->Unmap(0, nullptr);
+}
diff --git a/Engine/Source/Runtime/Windows/D3D12RHI/Private/D3D12Resources.cpp b/Engine/Source/Runtime/Windows/D3D12RHI/Private/D3D12Resources.cpp
new file mode 100644
--- /dev/null
+++ b/Engine/Source/Runtime/Windows/D3D12RHI/Private/D3D12Resources.cpp
@@ -0,0 +1,110 @@
+// Copyright 1998-2014 Epic Games, Inc. All Rights Reserved.
+
+/*=============================================================================
+D3D12Resources.cpp: D3D RHI utility implementation.
+=============================================================================*/
+
+#include "D3D12RHIPrivate.h"
+#include "EngineModule.h"
+
+
+FD3D12DeferredDeletionQueue::FD3D12DeferredDeletionQueue(FD3D12Device* InParent) :
+    FD3D12DeviceChild(InParent){}
+
+bool FD3D12DeferredDeletionQueue::ReleaseResources()
+{
+	struct FDequeueFenceObject
+	{
+		FDequeueFenceObject(FD3D12CommandListManager* InCommandListManager)
+			: CommandListManager(InCommandListManager)
+		{
+		}
+
+		bool operator() (FencedObjectType FenceObject)
+		{
+			return CommandListManager->GetFence(FT_Frame).IsFenceFinished(FenceObject.Value);
+		}
+
+		FD3D12CommandListManager* CommandListManager;
+	};
+
+	FD3D12CommandListManager& CommandListManager = GetParentDevice()->GetCommandListManager();
+
+	FencedObjectType FenceObject;
+	FDequeueFenceObject DequeueFenceObject(&CommandListManager);
+
+	const uint64 LastCompletedFrameFence = CommandListManager.GetFence(EFenceType::FT_Frame).GetLastCompletedFence();
+	while (DeferredReleaseQueue.Dequeue(FenceObject, DequeueFenceObject))
+	{
+		FenceObject.Key->Release();
+	}
+
+	return DeferredReleaseQueue.IsEmpty();
+}
+
+#if UE_BUILD_DEBUG
+int64 FD3D12Resource::TotalResourceCount = 0;
+int64 FD3D12Resource::NoStateTrackingResourceCount = 0;
+#endif
+
+FD3D12Resource::~FD3D12Resource()
+{
+#if SUPPORTS_MEMORY_RESIDENCY
+	if (HeapType == D3D12_HEAP_TYPE_DEFAULT)
+		FD3D12DynamicRHI::GetD3DRHI()->GetResourceResidencyManager().ResourceFreed (this);
+#endif
+
+	if (pResourceState)
+	{
+		delete pResourceState;
+		pResourceState = nullptr;
+	}
+}
+
+#if SUPPORTS_MEMORY_RESIDENCY
+void FD3D12Resource::UpdateResourceSize()
+{
+	if (HeapType == D3D12_HEAP_TYPE_DEFAULT)
+		ResourceSize = (uint32)GetRequiredIntermediateSize(GetResource(), 0, Desc.MipLevels);
+}
+
+void FD3D12Resource::UpdateResidency()
+{
+	if (HeapType == D3D12_HEAP_TYPE_DEFAULT)
+		FD3D12DynamicRHI::GetD3DRHI()->GetResourceResidencyManager().UpdateResidency (this);
+}
+#endif
+
+void inline FD3D12CommandListHandle::FD3D12CommandListData::FCommandListResourceState::ConditionalInitalize(FD3D12Resource* pResource, CResourceState& ResourceState)
+{
+	// If there is no entry, all subresources should be in the resource's TBD state.
+	// This means we need to have pending resource barrier(s).
+	if (!ResourceState.CheckResourceStateInitalized())
+	{
+		ResourceState.Initialize(pResource->GetSubresourceCount());
+		check(ResourceState.CheckResourceState(D3D12_RESOURCE_STATE_TBD));
+	}
+
+	check(ResourceState.CheckResourceStateInitalized());
+}
+
+CResourceState& FD3D12CommandListHandle::FD3D12CommandListData::FCommandListResourceState::GetResourceState(FD3D12Resource* pResource)
+{
+	// Only certain resources should use this
+	check(pResource->RequiresResourceStateTracking());
+
+	CResourceState& ResourceState = ResourceStates.FindOrAdd(pResource);
+	ConditionalInitalize(pResource, ResourceState);
+	return ResourceState;
+}
+
+void FD3D12CommandListHandle::FD3D12CommandListData::FCommandListResourceState::Empty()
+{
+	ResourceStates.Empty();
+}
+
+void FD3D12CommandListHandle::Execute (bool WaitForCompletion)
+{
+	check (CommandListData);
+	CommandListData->CommandListManager->ExecuteCommandList(*this, WaitForCompletion);
+}
\ No newline at end of file
diff --git a/Engine/Source/Runtime/Windows/D3D12RHI/Private/D3D12Shaders.cpp b/Engine/Source/Runtime/Windows/D3D12RHI/Private/D3D12Shaders.cpp
new file mode 100644
--- /dev/null
+++ b/Engine/Source/Runtime/Windows/D3D12RHI/Private/D3D12Shaders.cpp
@@ -0,0 +1,381 @@
+// Copyright 1998-2014 Epic Games, Inc. All Rights Reserved.
+
+/*=============================================================================
+	D3D12Shaders.cpp: D3D shader RHI implementation.
+	=============================================================================*/
+
+#include "D3D12RHIPrivate.h"
+
+FVertexShaderRHIRef FD3D12DynamicRHI::RHICreateVertexShader(const TArray<uint8>& Code)
+{
+	check(Code.Num());
+	FD3D12VertexShader* Shader = new FD3D12VertexShader;
+
+	FMemoryReader Ar(Code, true);
+	Ar << Shader->ShaderResourceTable;
+	int32 Offset = Ar.Tell();
+	const uint8* CodePtr = Code.GetData() + Offset;
+	const SIZE_T CodeSize = Code.Num() - Offset - 5;
+
+	// bGlobalUniformBufferUsed and resource counts are packed in the last few bytes, see CompileD3D11Shader
+	Shader->bShaderNeedsGlobalConstantBuffer = Code[Code.Num() - 5] != 0;
+	Shader->SamplerCount                     = Code[Code.Num() - 4];
+	Shader->SRVCount                         = Code[Code.Num() - 3];
+	Shader->UAVCount                         = Code[Code.Num() - 1];
+
+	//TODO Special case for Constant Buffer count as it currently needs all buffers set if it needs any set.
+	Shader->CBCount = (Code[Code.Num() - 2]) ? MAX_CBS : 0;
+
+	Shader->Code = Code;
+	Shader->Offset = Offset;
+
+	D3D12_SHADER_BYTECODE ShaderBytecode;
+	ShaderBytecode.pShaderBytecode = Shader->Code.GetData() + Offset;
+	ShaderBytecode.BytecodeLength = CodeSize;
+	Shader->ShaderBytecode.SetShaderBytecode(ShaderBytecode);
+
+	return Shader;
+}
+
+FPixelShaderRHIRef FD3D12DynamicRHI::RHICreatePixelShader(const TArray<uint8>& Code)
+{
+	check(Code.Num());
+
+	FD3D12PixelShader* Shader = new FD3D12PixelShader;
+
+	FMemoryReader Ar(Code, true);
+	Ar << Shader->ShaderResourceTable;
+	int32 Offset = Ar.Tell();
+	const uint8* CodePtr = Code.GetData() + Offset;
+	const SIZE_T CodeSize = Code.Num() - Offset - 5;
+
+	// bGlobalUniformBufferUsed and resource counts are packed in the last few bytes, see CompileD3D11Shader
+	Shader->bShaderNeedsGlobalConstantBuffer = Code[Code.Num() - 5] != 0;
+	Shader->SamplerCount                     = Code[Code.Num() - 4];
+	Shader->SRVCount                         = Code[Code.Num() - 3];
+	Shader->UAVCount                         = Code[Code.Num() - 1];
+
+	//TODO Special case for Constant Buffer count as it currently needs all buffers set if it needs any set.
+	Shader->CBCount = (Code[Code.Num() - 2]) ? MAX_CBS : 0;
+
+	Shader->Code = Code;
+
+	D3D12_SHADER_BYTECODE ShaderBytecode;
+	ShaderBytecode.pShaderBytecode = Shader->Code.GetData() + Offset;
+	ShaderBytecode.BytecodeLength = CodeSize;
+	Shader->ShaderBytecode.SetShaderBytecode(ShaderBytecode);
+
+	return Shader;
+}
+
+FHullShaderRHIRef FD3D12DynamicRHI::RHICreateHullShader(const TArray<uint8>& Code)
+{
+	check(Code.Num());
+
+	FD3D12HullShader* Shader = new FD3D12HullShader;
+
+	FMemoryReader Ar(Code, true);
+	Ar << Shader->ShaderResourceTable;
+	int32 Offset = Ar.Tell();
+	const uint8* CodePtr = Code.GetData() + Offset;
+	const SIZE_T CodeSize = Code.Num() - Offset - 5;
+
+	// bGlobalUniformBufferUsed and resource counts are packed in the last few bytes, see CompileD3D11Shader
+	Shader->bShaderNeedsGlobalConstantBuffer = Code[Code.Num() - 5] != 0;
+	Shader->SamplerCount                     = Code[Code.Num() - 4];
+	Shader->SRVCount                         = Code[Code.Num() - 3];
+	Shader->UAVCount                         = Code[Code.Num() - 1];
+
+	//TODO Special case for Constant Buffer count as it currently needs all buffers set if it needs any set.
+	Shader->CBCount = (Code[Code.Num() - 2]) ? MAX_CBS : 0;
+
+	Shader->Code = Code;
+
+	D3D12_SHADER_BYTECODE ShaderBytecode;
+	ShaderBytecode.pShaderBytecode = Shader->Code.GetData() + Offset;
+	ShaderBytecode.BytecodeLength = CodeSize;
+	Shader->ShaderBytecode.SetShaderBytecode(ShaderBytecode);
+
+	return Shader;
+}
+
+FDomainShaderRHIRef FD3D12DynamicRHI::RHICreateDomainShader(const TArray<uint8>& Code)
+{
+	check(Code.Num());
+
+	FD3D12DomainShader* Shader = new FD3D12DomainShader;
+
+	FMemoryReader Ar(Code, true);
+	Ar << Shader->ShaderResourceTable;
+	int32 Offset = Ar.Tell();
+	const uint8* CodePtr = Code.GetData() + Offset;
+	const SIZE_T CodeSize = Code.Num() - Offset - 5;
+
+	// bGlobalUniformBufferUsed and resource counts are packed in the last few bytes, see CompileD3D11Shader
+	Shader->bShaderNeedsGlobalConstantBuffer = Code[Code.Num() - 5] != 0;
+	Shader->SamplerCount                     = Code[Code.Num() - 4];
+	Shader->SRVCount                         = Code[Code.Num() - 3];
+	Shader->UAVCount                         = Code[Code.Num() - 1];
+
+	//TODO Special case for Constant Buffer count as it currently needs all buffers set if it needs any set.
+	Shader->CBCount = (Code[Code.Num() - 2] > 0) ? MAX_CBS : 0;
+
+	Shader->Code = Code;
+
+	D3D12_SHADER_BYTECODE ShaderBytecode;
+	ShaderBytecode.pShaderBytecode = Shader->Code.GetData() + Offset;
+	ShaderBytecode.BytecodeLength = CodeSize;
+	Shader->ShaderBytecode.SetShaderBytecode(ShaderBytecode);
+
+	return Shader;
+}
+
+FGeometryShaderRHIRef FD3D12DynamicRHI::RHICreateGeometryShader(const TArray<uint8>& Code)
+{
+	check(Code.Num());
+
+	FD3D12GeometryShader* Shader = new FD3D12GeometryShader;
+
+	FMemoryReader Ar(Code, true);
+	Ar << Shader->ShaderResourceTable;
+	int32 Offset = Ar.Tell();
+	const uint8* CodePtr = Code.GetData() + Offset;
+	const SIZE_T CodeSize = Code.Num() - Offset - 5;
+
+	// bGlobalUniformBufferUsed and resource counts are packed in the last few bytes, see CompileD3D11Shader
+	Shader->bShaderNeedsGlobalConstantBuffer = Code[Code.Num() - 5] != 0;
+	Shader->SamplerCount                     = Code[Code.Num() - 4];
+	Shader->SRVCount                         = Code[Code.Num() - 3];
+	Shader->UAVCount                         = Code[Code.Num() - 1];
+
+	//TODO Special case for Constant Buffer count as it currently needs all buffers set if it needs any set.
+	Shader->CBCount = (Code[Code.Num() - 2] > 0) ? MAX_CBS : 0;
+
+	Shader->Code = Code;
+
+	D3D12_SHADER_BYTECODE ShaderBytecode;
+	ShaderBytecode.pShaderBytecode = Shader->Code.GetData() + Offset;
+	ShaderBytecode.BytecodeLength = CodeSize;
+	Shader->ShaderBytecode.SetShaderBytecode(ShaderBytecode);
+
+	return Shader;
+}
+
+FGeometryShaderRHIRef FD3D12DynamicRHI::RHICreateGeometryShaderWithStreamOutput(const TArray<uint8>& Code, const FStreamOutElementList& ElementList, uint32 NumStrides, const uint32* Strides, int32 RasterizedStream)
+{
+	check(Code.Num());
+
+	FD3D12GeometryShader* Shader = new FD3D12GeometryShader;
+
+	FMemoryReader Ar(Code, true);
+	Ar << Shader->ShaderResourceTable;
+	int32 Offset = Ar.Tell();
+	const uint8* CodePtr = Code.GetData() + Offset;
+	const SIZE_T CodeSize = Code.Num() - Offset - 5;
+
+	Shader->StreamOutput.RasterizedStream = RasterizedStream;
+	if (RasterizedStream == -1)
+	{
+		Shader->StreamOutput.RasterizedStream = D3D12_SO_NO_RASTERIZED_STREAM;
+	}
+
+	Shader->StreamOutput.NumEntries = ElementList.Num();
+	Shader->pStreamOutEntries = new D3D12_SO_DECLARATION_ENTRY[ElementList.Num()];
+	Shader->StreamOutput.pSODeclaration = Shader->pStreamOutEntries;
+	if (Shader->pStreamOutEntries == nullptr)
+	{
+		return nullptr;	// Out of memory
+	}
+	for (int32 EntryIndex = 0; EntryIndex < ElementList.Num(); EntryIndex++)
+	{
+		Shader->pStreamOutEntries[EntryIndex].Stream = ElementList[EntryIndex].Stream;
+		Shader->pStreamOutEntries[EntryIndex].SemanticName = ElementList[EntryIndex].SemanticName;
+		Shader->pStreamOutEntries[EntryIndex].SemanticIndex = ElementList[EntryIndex].SemanticIndex;
+		Shader->pStreamOutEntries[EntryIndex].StartComponent = ElementList[EntryIndex].StartComponent;
+		Shader->pStreamOutEntries[EntryIndex].ComponentCount = ElementList[EntryIndex].ComponentCount;
+		Shader->pStreamOutEntries[EntryIndex].OutputSlot = ElementList[EntryIndex].OutputSlot;
+	}
+
+
+	// bGlobalUniformBufferUsed and resource counts are packed in the last few bytes, see CompileD3D11Shader
+	Shader->bShaderNeedsGlobalConstantBuffer = Code[Code.Num() - 5] != 0;
+	Shader->SamplerCount                     = Code[Code.Num() - 4];
+	Shader->SRVCount                         = Code[Code.Num() - 3];
+	Shader->UAVCount                         = Code[Code.Num() - 1];
+
+	//TODO Special case for Constant Buffer count as it currently needs all buffers set if it needs any set.
+	Shader->CBCount = (Code[Code.Num() - 2] > 0) ? MAX_CBS : 0;
+
+	// Indicate this shader uses stream output
+	Shader->bShaderNeedsStreamOutput = true;
+
+	// Copy the strides
+	Shader->StreamOutput.NumStrides = NumStrides;
+	Shader->pStreamOutStrides = new uint32[NumStrides];
+	Shader->StreamOutput.pBufferStrides = Shader->pStreamOutStrides;
+	if (Shader->pStreamOutStrides == nullptr)
+	{
+		return nullptr;	// Out of memory
+	}
+	FMemory::Memcpy(Shader->pStreamOutStrides, Strides, sizeof(uint32) * NumStrides);
+
+	Shader->Code = Code;
+
+	D3D12_SHADER_BYTECODE ShaderBytecode;
+	ShaderBytecode.pShaderBytecode = Shader->Code.GetData() + Offset;
+	ShaderBytecode.BytecodeLength = CodeSize;
+	Shader->ShaderBytecode.SetShaderBytecode(ShaderBytecode);
+
+	return Shader;
+}
+
+FComputeShaderRHIRef FD3D12DynamicRHI::RHICreateComputeShader(const TArray<uint8>& Code)
+{
+	check(Code.Num());
+
+	FD3D12ComputeShader* Shader = new FD3D12ComputeShader;
+
+	FMemoryReader Ar(Code, true);
+	Ar << Shader->ShaderResourceTable;
+	int32 Offset = Ar.Tell();
+	const uint8* CodePtr = Code.GetData() + Offset;
+	const SIZE_T CodeSize = Code.Num() - Offset - 5;
+
+	// bGlobalUniformBufferUsed and resource counts are packed in the last few bytes, see CompileD3D11Shader
+	Shader->bShaderNeedsGlobalConstantBuffer = Code[Code.Num() - 5] != 0;
+	Shader->SamplerCount                     = Code[Code.Num() - 4];
+	Shader->SRVCount                         = Code[Code.Num() - 3];
+	Shader->UAVCount                         = Code[Code.Num() - 1];
+
+	//TODO Special case for Constant Buffer count as it currently needs all buffers set if it needs any set.
+	Shader->CBCount = (Code[Code.Num() - 2] > 0) ? MAX_CBS : 0;
+
+	Shader->Code = Code;
+
+	D3D12_SHADER_BYTECODE ShaderBytecode;
+	ShaderBytecode.pShaderBytecode = Shader->Code.GetData() + Offset;
+	ShaderBytecode.BytecodeLength = CodeSize;
+	Shader->ShaderBytecode.SetShaderBytecode(ShaderBytecode);
+
+	return Shader;
+}
+
+void FD3D12CommandContext::RHISetMultipleViewports(uint32 Count, const FViewportBounds* Data)
+{
+	// structures are chosen to be directly mappable
+	D3D12_VIEWPORT* D3DData = (D3D12_VIEWPORT*)Data;
+
+	StateCache.SetViewports(Count, D3DData);
+}
+
+static volatile uint64 BoundShaderStateID = 0;
+
+FD3D12BoundShaderState::FD3D12BoundShaderState(
+	FVertexDeclarationRHIParamRef InVertexDeclarationRHI,
+	FVertexShaderRHIParamRef InVertexShaderRHI,
+	FPixelShaderRHIParamRef InPixelShaderRHI,
+	FHullShaderRHIParamRef InHullShaderRHI,
+	FDomainShaderRHIParamRef InDomainShaderRHI,
+	FGeometryShaderRHIParamRef InGeometryShaderRHI
+	) :
+	CacheLink(InVertexDeclarationRHI, InVertexShaderRHI, InPixelShaderRHI, InHullShaderRHI, InDomainShaderRHI, InGeometryShaderRHI, this),
+    UniqueID(InterlockedIncrement64(reinterpret_cast<volatile int64*>(&BoundShaderStateID)))
+{
+	INC_DWORD_STAT(STAT_D3D12NumBoundShaderState);
+
+    // Warning: Input layout desc contains padding which must be zero-initialized to prevent PSO cache misses
+    FMemory::Memzero(&InputLayout, sizeof(InputLayout));
+
+	FD3D12VertexDeclaration*  InVertexDeclaration = FD3D12DynamicRHI::ResourceCast(InVertexDeclarationRHI);
+	FD3D12VertexShader*  InVertexShader = FD3D12DynamicRHI::ResourceCast(InVertexShaderRHI);
+	FD3D12PixelShader*  InPixelShader = FD3D12DynamicRHI::ResourceCast(InPixelShaderRHI);
+	FD3D12HullShader*  InHullShader = FD3D12DynamicRHI::ResourceCast(InHullShaderRHI);
+	FD3D12DomainShader*  InDomainShader = FD3D12DynamicRHI::ResourceCast(InDomainShaderRHI);
+	FD3D12GeometryShader*  InGeometryShader = FD3D12DynamicRHI::ResourceCast(InGeometryShaderRHI);
+
+	// Create an input layout for this combination of vertex declaration and vertex shader.
+	InputLayout.NumElements = (InVertexDeclaration ? InVertexDeclaration->VertexElements.Num() : 0);
+	InputLayout.pInputElementDescs = (InVertexDeclaration ? InVertexDeclaration->VertexElements.GetData() : nullptr);
+
+	bShaderNeedsGlobalConstantBuffer[SF_Vertex] = InVertexShader->bShaderNeedsGlobalConstantBuffer;
+	bShaderNeedsGlobalConstantBuffer[SF_Hull] = InHullShader ? InHullShader->bShaderNeedsGlobalConstantBuffer : false;
+	bShaderNeedsGlobalConstantBuffer[SF_Domain] = InDomainShader ? InDomainShader->bShaderNeedsGlobalConstantBuffer : false;
+	bShaderNeedsGlobalConstantBuffer[SF_Pixel] = InPixelShader ? InPixelShader->bShaderNeedsGlobalConstantBuffer : false;
+	bShaderNeedsGlobalConstantBuffer[SF_Geometry] = InGeometryShader ? InGeometryShader->bShaderNeedsGlobalConstantBuffer : false;
+
+	static_assert(ARRAY_COUNT(bShaderNeedsGlobalConstantBuffer) == SF_NumFrequencies, "EShaderFrequency size should match with array count of bShaderNeedsGlobalConstantBuffer.");
+
+#if D3D12_SUPPORTS_PARALLEL_RHI_EXECUTE
+	CacheLink.AddToCache();
+#endif
+}
+
+FD3D12BoundShaderState::~FD3D12BoundShaderState()
+{
+	DEC_DWORD_STAT(STAT_D3D12NumBoundShaderState);
+#if D3D12_SUPPORTS_PARALLEL_RHI_EXECUTE
+	CacheLink.RemoveFromCache();
+#endif
+}
+
+/**
+* Creates a bound shader state instance which encapsulates a decl, vertex shader, and pixel shader
+* @param VertexDeclaration - existing vertex decl
+* @param StreamStrides - optional stream strides
+* @param VertexShader - existing vertex shader
+* @param HullShader - existing hull shader
+* @param DomainShader - existing domain shader
+* @param PixelShader - existing pixel shader
+* @param GeometryShader - existing geometry shader
+*/
+FBoundShaderStateRHIRef FD3D12DynamicRHI::RHICreateBoundShaderState(
+	FVertexDeclarationRHIParamRef VertexDeclarationRHI,
+	FVertexShaderRHIParamRef VertexShaderRHI,
+	FHullShaderRHIParamRef HullShaderRHI,
+	FDomainShaderRHIParamRef DomainShaderRHI,
+	FPixelShaderRHIParamRef PixelShaderRHI,
+	FGeometryShaderRHIParamRef GeometryShaderRHI
+	)
+{
+	SCOPE_CYCLE_COUNTER(STAT_D3D12CreateBoundShaderStateTime);
+
+    checkf(GIsRHIInitialized && GetRHIDevice()->GetCommandListManager().IsReady(), (TEXT("Bound shader state RHI resource was created without initializing Direct3D first")));
+
+#if D3D12_SUPPORTS_PARALLEL_RHI_EXECUTE
+	// Check for an existing bound shader state which matches the parameters
+	FBoundShaderStateRHIRef CachedBoundShaderState = GetCachedBoundShaderState_Threadsafe(
+		VertexDeclarationRHI,
+		VertexShaderRHI,
+		PixelShaderRHI,
+		HullShaderRHI,
+		DomainShaderRHI,
+		GeometryShaderRHI
+		);
+	if(CachedBoundShaderState.GetReference())
+	{
+		// If we've already created a bound shader state with these parameters, reuse it.
+		return CachedBoundShaderState;
+	}
+#else
+	check(IsInRenderingThread() || IsInRHIThread());
+	// Check for an existing bound shader state which matches the parameters
+	FCachedBoundShaderStateLink* CachedBoundShaderStateLink = GetCachedBoundShaderState(
+		VertexDeclarationRHI,
+		VertexShaderRHI,
+		PixelShaderRHI,
+		HullShaderRHI,
+		DomainShaderRHI,
+		GeometryShaderRHI
+		);
+	if(CachedBoundShaderStateLink)
+	{
+		// If we've already created a bound shader state with these parameters, reuse it.
+		return CachedBoundShaderStateLink->BoundShaderState;
+	}
+#endif
+	else
+	{
+		SCOPE_CYCLE_COUNTER(STAT_D3D12NewBoundShaderStateTime);
+		return new FD3D12BoundShaderState(VertexDeclarationRHI, VertexShaderRHI, PixelShaderRHI, HullShaderRHI, DomainShaderRHI, GeometryShaderRHI);
+	}
+}
