UE4 - D3D12RHI (not working)

diff --git a/Engine/Source/Runtime/Windows/D3D12RHI/Private/D3D12StateCache.cpp b/Engine/Source/Runtime/Windows/D3D12RHI/Private/D3D12StateCache.cpp
new file mode 100644
--- /dev/null
+++ b/Engine/Source/Runtime/Windows/D3D12RHI/Private/D3D12StateCache.cpp
@@ -0,0 +1,3252 @@
+// Copyright 1998-2014 Epic Games, Inc. All Rights Reserved.
+
+// Implementation of Device Context State Caching to improve draw
+//	thread performance by removing redundant device context calls.
+
+//-----------------------------------------------------------------------------
+//	Include Files
+//-----------------------------------------------------------------------------
+#include "D3D12RHIPrivate.h"
+#include "D3D12StateCache.h"
+#include <emmintrin.h>
+
+#define IL_MAX_SEMANTIC_NAME 255
+
+extern bool D3D12RHI_ShouldCreateWithD3DDebug();
+
+inline bool operator!=(D3D12_CPU_DESCRIPTOR_HANDLE lhs, D3D12_CPU_DESCRIPTOR_HANDLE rhs)
+{
+	return lhs.ptr != rhs.ptr;
+}
+
+#if D3D12_STATE_CACHE_RUNTIME_TOGGLE
+
+// Default the state caching system to on.
+bool GD3D12SkipStateCaching = false;
+
+// A self registering exec helper to check for the TOGGLESTATECACHE command.
+class FD3D12ToggleStateCacheExecHelper : public FSelfRegisteringExec
+{
+	virtual bool Exec(class UWorld* InWorld, const TCHAR* Cmd, FOutputDevice& Ar)
+	{
+		if (FParse::Command(&Cmd, TEXT("TOGGLESTATECACHE")))
+		{
+			GD3D12SkipStateCaching = !GD3D12SkipStateCaching;
+			Ar.Log(FString::Printf(TEXT("D3D12 State Caching: %s"), GD3D12SkipStateCaching ? TEXT("OFF") : TEXT("ON")));
+			return true;
+		}
+		return false;
+	}
+};
+static FD3D12ToggleStateCacheExecHelper GD3D12ToggleStateCacheExecHelper;
+
+#endif	// D3D12_STATE_CACHE_RUNTIME_TOGGLE
+
+// FD3D12DynamicHeapAllocator
+FD3D12DynamicHeapAllocator::FD3D12DynamicHeapAllocator(FD3D12Device* InParent, D3D12_HEAP_TYPE InHeapType)
+	: HeapType(InHeapType)
+	, FastAllocBuffer(nullptr)
+	, FastAllocData(nullptr)
+	, NextFastAllocOffset(0UL - 1)
+	, CurrentCommandContext(nullptr)
+	, FD3D12ResourceAllocator(InParent)
+{
+#if UE_BUILD_DEBUG
+	FMemory::Memzero(&AllocatorStats, sizeof(AllocatorStats));
+#endif
+}
+
+FD3D12DynamicHeapAllocator::~FD3D12DynamicHeapAllocator()
+{
+	// Ensures there are no lingering resource locations with a pointer to this allocator
+	FRHIResource::FlushPendingDeletes();
+
+#if DO_CHECK
+	for (uint32 i = 0; i < NumBuckets; ++i)
+	{
+		check(AvailableBlocks[i].IsEmpty());
+	}
+#endif
+
+#if UE_BUILD_DEBUG
+	FOutputDeviceRedirector* pOutputDevice = FOutputDeviceRedirector::Get();
+	DumpAllocatorStats(*pOutputDevice);
+#endif
+}
+
+HRESULT FD3D12DynamicHeapAllocator::CreateResource(FD3D12DynamicHeapAllocator* Allocator, D3D12_HEAP_TYPE heapType, uint32 size, FD3D12Resource **ppResource, void **ppBaseAddress)
+{
+	HRESULT hr = S_OK;
+	TRefCountPtr<FD3D12Resource> Buffer;
+
+	VERIFYD3D11RESULT(hr = Allocator->GetParentDevice()->GetResourceHelper().CreateBuffer(heapType, size, Buffer.GetInitReference()));
+	if (FAILED(hr))
+	{
+		return hr;
+	}
+
+	void* pData;
+	switch (heapType)
+	{
+	case D3D12_HEAP_TYPE_UPLOAD:
+	case D3D12_HEAP_TYPE_READBACK:
+		VERIFYD3D11RESULT(Buffer->GetResource()->Map(0, nullptr, &pData));
+		break;
+	default:
+		// Only upload and readback heaps are supported by FD3D12DynamicHeapAllocator
+		return E_INVALIDARG;
+		break;
+	}
+
+	*ppResource = Buffer;
+	Buffer->AddRef();
+
+	*ppBaseAddress = pData;
+
+	return S_OK;
+}
+
+void FD3D12DynamicHeapAllocator::ReleaseAllResources()
+{
+	TRefCountPtr<FD3D12ResourceBlockInfo> pBlock;
+	
+	// Empty the expired blocks list
+	while(!ExpiredBlocks.IsEmpty())
+	{
+		ExpiredBlocks.Dequeue(pBlock);
+	}
+
+	// Empty all available blocks queues
+	for (uint32 i = 0; i < NumBuckets; ++i)
+	{
+		for (; AvailableBlocks[i].Dequeue(pBlock);)
+		{
+#if UE_BUILD_DEBUG
+			AllocatorStats.AvailableAllocs[i]--;
+#endif
+		}
+	}
+}
+
+void FD3D12DynamicHeapAllocator::Alloc(uint32 size, uint32 alignment, FD3D12ResourceBlockInfo **ppBlock)
+{
+	// Size cannot be smaller than the requested alignment
+	size = FMath::Max(size, alignment);
+
+	uint32 bucket = BucketFromSize(size, BucketShift);
+	check(bucket < NumBuckets);
+	TRefCountPtr<FD3D12ResourceBlockInfo> pBlock;
+
+	// See if a block is already available in the bucket
+	if (AvailableBlocks[bucket].Dequeue(pBlock))
+	{
+#if UE_BUILD_DEBUG
+		AllocatorStats.AvailableAllocs[bucket]--;
+#endif
+		check(pBlock.GetReference());
+		check(pBlock->Address);
+		check(pBlock->ResourceHeap);
+	}
+	else
+	{
+		// No blocks of the requested size are available so make one
+		TRefCountPtr<FD3D12Resource> Resource;
+		void *BaseAddress;
+
+		// Allocate a block
+		uint32 BlockSize = BlockSizeFromBufferSize(size, BucketShift);
+
+		check(BlockSize >= size);
+		if (FAILED(CreateResource(this, HeapType, BlockSize < MIN_HEAP_SIZE ? MIN_HEAP_SIZE : BlockSize, Resource.GetInitReference(), &BaseAddress)))
+		{
+			*ppBlock = nullptr;
+			return;
+		}
+
+#if UE_BUILD_DEBUG
+		AllocatorStats.NumCreatedResources[bucket]++;
+
+		const uint32 allocatedSize = BlockSize < MIN_HEAP_SIZE ? MIN_HEAP_SIZE : BlockSize;
+		AllocatorStats.BucketSize[bucket] += allocatedSize;
+
+		if (AllocatorStats.BlockSize[bucket] == 0)
+		{
+			AllocatorStats.BlockSize[bucket] = BlockSize;
+		}
+		check(AllocatorStats.BlockSize[bucket] == BlockSize);
+#endif
+
+		check(BaseAddress == (uint8*)(((uint64)BaseAddress + alignment - 1) & ~((uint64)alignment - 1)));
+
+		pBlock = new FD3D12ResourceBlockInfo(Resource, BaseAddress, 0, bucket, this);
+
+		// If bucket blocks are less than 65536 bytes then use FDynamicResourcePool
+		if (BlockSize < MIN_HEAP_SIZE)
+		{
+			// Create additional available blocks that can be sub-allocated from the same resource
+			for (uint32 Offset = BlockSize; Offset <= MIN_HEAP_SIZE - BlockSize; Offset += BlockSize)
+			{
+				void *Start = (void *)((uint8 *)BaseAddress + Offset);
+				FD3D12ResourceBlockInfo *NewBlock = new FD3D12ResourceBlockInfo(Resource, Start, Offset, bucket, this);
+
+				// Add the bucket to the available list
+				AvailableBlocks[bucket].Enqueue(NewBlock);
+
+#if UE_BUILD_DEBUG
+				AllocatorStats.AvailableAllocs[bucket]++;
+#endif
+			}
+		}
+	}
+
+#if UE_BUILD_DEBUG
+	AllocatorStats.CurAllocs[bucket]++;
+	AllocatorStats.TotalAllocs[bucket]++;
+#endif
+
+	// Check the returned block matches the requested alignment
+	check(((size_t)(pBlock->Address) % alignment) == 0);
+
+	pBlock->AddRef();
+	*ppBlock = pBlock;
+}
+
+void *FD3D12DynamicHeapAllocator::Alloc(uint32 size, uint32 alignment, FD3D12ResourceLocation **ResourceLocation)
+{
+	*ResourceLocation = new FD3D12ResourceLocation(GetParentDevice(), size);
+	(*ResourceLocation)->AddRef();
+
+	return Alloc(size, alignment, *ResourceLocation);
+}
+
+void *FD3D12DynamicHeapAllocator::Alloc(uint32 size, uint32 alignment, FD3D12ResourceLocation *ResourceLocation)
+{
+	FD3D12ResourceBlockInfo *pOldBlock = ResourceLocation->GetBlockInfo();
+
+	if (pOldBlock)
+	{
+		// Expire the old block
+		ExpireBlock(pOldBlock);
+	}
+
+	TRefCountPtr<FD3D12ResourceBlockInfo> pBlock;
+	Alloc(size, alignment, pBlock.GetInitReference());
+	check(pBlock);
+
+	FD3D12CommandListManager& CommandListManager = GetParentDevice()->GetCommandListManager();
+	pBlock->FrameFence = CommandListManager.GetFence(FT_Frame).GetCurrentFence();
+	check(pBlock->FrameFence);
+
+	FD3D12StateCache* pStateCache = CurrentCommandContext ? &CurrentCommandContext->StateCache : nullptr;
+	ResourceLocation->SetBlockInfo(pBlock, pStateCache);
+	ResourceLocation->EffectiveBufferSize = size;
+
+	check(pBlock->Address);
+	check(pBlock->ResourceHeap);
+
+	return pBlock->Address;
+}
+
+void *FD3D12DynamicHeapAllocator::FastAlloc(uint32 size, uint32 alignment, class FD3D12ResourceLocation *ResourceLocation)
+{
+	void *pData = nullptr;
+	if (size > MaxFastAllocBufferSize)
+	{
+		return Alloc(size, alignment, ResourceLocation);
+	}
+
+	uint32 CurrentOffset = Align(NextFastAllocOffset, alignment);
+
+	// See if there is room in the current pool
+	if (FastAllocBuffer.GetReference() == nullptr || MaxFastAllocBufferSize < CurrentOffset + size)
+	{
+		// No room so allocate a new pool
+		FastAllocBuffer = new FD3D12ResourceLocation();
+		FastAllocData = Alloc(MaxFastAllocBufferSize, 4, FastAllocBuffer);
+		CurrentOffset = 0;
+	}
+
+	check(MaxFastAllocBufferSize - size >= CurrentOffset);
+
+	// Create a FD3D12ResourceLocation representing a sub-section of the pool resource
+	ResourceLocation->SetFromD3DResource(FastAllocBuffer->GetResource(), CurrentOffset, size);
+	ResourceLocation->LinkToResourceLocation(FastAllocBuffer);
+	NextFastAllocOffset = CurrentOffset + size;
+	pData = (void *)((uint8 *)FastAllocData + CurrentOffset);
+
+	return pData;
+}
+
+
+
+void FD3D12DynamicHeapAllocator::ExpireBlock(FD3D12ResourceBlockInfo *BlockInfo)
+{
+	FD3D12CommandListManager& CommandListManager = GetParentDevice()->GetCommandListManager();
+	BlockInfo->FrameFence = CommandListManager.GetFence(FT_Frame).GetCurrentFence();
+
+	ExpiredBlocks.Enqueue(BlockInfo);
+}
+
+// This frame count value helps makes sure that we don't delete resources too soon. If resources are deleted too soon,
+// we can get in a loop the heap allocator will be constantly deleting and creating resources every frame which
+// results in CPU stutters. DynamicRetentionFrameCount was tested and set to a value that appears to be adequate for
+// creating a stable state on the Infiltrator demo.
+const uint64 DynamicRetentionFrameCount = 10;
+const uint64 DefaultRetentionFrameCount = 5;
+
+struct FDequeueAlloc
+{
+	FDequeueAlloc (FD3D12CommandListManager* InCommandListManager, const uint64 InRetentionFrameCount)
+		: CommandListManager (InCommandListManager)
+		, RetentionFrameCount(InRetentionFrameCount)
+	{
+	}
+
+	bool operator() (FD3D12ResourceBlockInfo* pBlockInfo)
+	{
+		return CommandListManager->GetFence(FT_Frame).IsFenceFinished(pBlockInfo->FrameFence + RetentionFrameCount);
+	}
+
+	FD3D12CommandListManager* CommandListManager;
+	const uint64              RetentionFrameCount;
+};
+
+void FD3D12DynamicHeapAllocator::CleanupFreeBlocks()
+{
+	FD3D12CommandListManager& CommandListManager = GetParentDevice()->GetCommandListManager();
+
+#if SUB_ALLOCATED_DEFAULT_ALLOCATIONS
+	const static uint32 MinCleanupBucket = FMath::Max<uint32> (0, BucketFromSize(MIN_HEAP_SIZE, BucketShift));
+#else
+	const static uint32 MinCleanupBucket = 0;
+#endif
+
+	// Start at bucket 8 since smaller buckets are sub-allocated resources
+	// and would be fragmented by deleting blocks
+	for (uint32 bucket = MinCleanupBucket; bucket < NumBuckets; bucket++)
+	{
+		FDequeueAlloc DequeueAlloc(&CommandListManager, DynamicRetentionFrameCount);
+		TRefCountPtr<FD3D12ResourceBlockInfo> pBlockInfo;
+
+		while (AvailableBlocks[bucket].Dequeue (pBlockInfo, DequeueAlloc))
+		{
+#if UE_BUILD_DEBUG
+			AllocatorStats.AvailableAllocs[bucket]--;
+#endif
+		}
+	}
+	
+	TRefCountPtr<FD3D12ResourceBlockInfo> pBlockInfo;
+	FDequeueAlloc DequeueAlloc (&CommandListManager, 0);
+
+	while(ExpiredBlocks.Dequeue (pBlockInfo, DequeueAlloc))
+	 {
+		check(pBlockInfo);
+		uint32 bucket = pBlockInfo->Bucket;
+		// Add the bucket to the available list
+		AvailableBlocks[bucket].Enqueue(pBlockInfo);
+#if UE_BUILD_DEBUG
+		AllocatorStats.AvailableAllocs[bucket]++;
+#endif
+	}
+}
+
+void FD3D12DynamicHeapAllocator::DumpAllocatorStats(class FOutputDevice& Ar)
+{
+#if 0//UE_BUILD_DEBUG
+	FBufferedOutputDevice BufferedOutput;
+	{
+		// This is the memory tracked inside individual allocation pools
+		FD3D12DynamicRHI *D3DRHI = FD3D12DynamicRHI::GetD3DRHI();
+		FName categoryName(L"DynamicHeapAllocator");
+
+		BufferedOutput.CategorizedLogf(categoryName, ELogVerbosity::Log, TEXT(""));
+		BufferedOutput.CategorizedLogf(categoryName, ELogVerbosity::Log, TEXT("Block Size Last Fence Cur Allocs Available Allocs Total Allocs Num Resources Mem Used"));
+		BufferedOutput.CategorizedLogf(categoryName, ELogVerbosity::Log, TEXT("---------- ---------- ---------- ---------------- ------------ ------------- --------"));
+
+		for (uint32 i = 0; i < NumBuckets; i++)
+		{
+			BufferedOutput.CategorizedLogf(categoryName, ELogVerbosity::Log, TEXT("% 10i % 10i % 16i % 12i % 13i % 8iK"),
+				AllocatorStats.BlockSize[i],
+				AllocatorStats.CurAllocs[i],
+				AllocatorStats.AvailableAllocs[i],
+				AllocatorStats.TotalAllocs[i],
+				AllocatorStats.NumCreatedResources[i],
+				AllocatorStats.BucketSize[i] / 1024);
+		}
+	}
+
+	BufferedOutput.RedirectTo(Ar);
+#endif
+}
+
+// FD3D12DefaultBufferPool
+void FD3D12DefaultBufferPool::CleanupFreeBlocks()
+{
+	FD3D12CommandListManager& CommandListManager = GetParentDevice()->GetCommandListManager();
+
+#if SUB_ALLOCATED_DEFAULT_ALLOCATIONS
+	const static uint32 MinCleanupBucket = FMath::Max<uint32> (0, BucketFromSize(MIN_HEAP_SIZE, BucketShift));
+#else
+	const static uint32 MinCleanupBucket = 0;
+#endif
+
+	// Start at bucket 8 since smaller buckets are sub-allocated resources
+	// and would be fragmented by deleting blocks
+	for (uint32 bucket = MinCleanupBucket; bucket < NumBuckets; bucket++)
+	{
+		FDequeueAlloc DequeueAlloc (&CommandListManager, DefaultRetentionFrameCount);
+
+		TRefCountPtr<FD3D12ResourceBlockInfo> pBlockInfo;
+
+		while (AvailableBuffers[bucket].Dequeue(pBlockInfo, DequeueAlloc))
+		{
+		}
+	}
+
+	TRefCountPtr<FD3D12ResourceBlockInfo> pBlockInfo;
+	FDequeueAlloc DequeueAlloc (&CommandListManager, 0);
+
+	while(ExpiredBlocks.Dequeue (pBlockInfo, DequeueAlloc))
+	{
+		check(pBlockInfo);
+		// Add the bucket to the available list if there is room
+		AvailableBuffers[pBlockInfo->Bucket].Enqueue(pBlockInfo);
+	}
+}
+
+FD3D12DefaultBufferPool::FD3D12DefaultBufferPool(FD3D12Device* InParent)
+	: bFirstAllocFromPool(true),
+	FD3D12ResourceAllocator(InParent)
+{
+#if UE_BUILD_DEBUG
+	FMemory::Memzero(&NumCreatedResources, sizeof(NumCreatedResources));
+	FMemory::Memzero(&BucketSize, sizeof(BucketSize));
+#endif
+}
+
+FD3D12DefaultBufferPool::FD3D12DefaultBufferPool()
+	: bFirstAllocFromPool(true),
+	FD3D12ResourceAllocator(nullptr)
+{
+#if UE_BUILD_DEBUG
+	FMemory::Memzero(&NumCreatedResources, sizeof(NumCreatedResources));
+	FMemory::Memzero(&BucketSize, sizeof(BucketSize));
+#endif
+}
+
+FD3D12DefaultBufferPool::~FD3D12DefaultBufferPool()
+{
+	ReleaseAllResources();
+}
+
+
+// Grab a buffer from the available buffers or create a new buffer if none are available
+HRESULT FD3D12DefaultBufferPool::AllocDefaultResource(const D3D12_RESOURCE_DESC& Desc, D3D12_SUBRESOURCE_DATA *pInitialData, FD3D12ResourceLocation *ResourceLocation, uint32 Alignment)
+{
+	FD3D12ResourceBlockInfo *pOldBlock = ResourceLocation->GetBlockInfo();
+
+	if (pOldBlock)
+	{
+		// Expire the old block
+		ExpireBlock(pOldBlock);
+	}
+
+	// MS: Used for debug purposes
+#define D3D12_DEFAULT_BUFFER_POOL_ALWAYS_CREATE_NEW_RESOURCE 0
+	TRefCountPtr<FD3D12ResourceBlockInfo> pBlock = nullptr;
+
+#if D3D12_DEFAULT_BUFFER_POOL_ALWAYS_CREATE_NEW_RESOURCE
+	uint32 requestedSize = Desc.Width;
+	Alignment = 0;
+#else
+	uint32 requestedSize = Desc.Width + Alignment;
+#endif
+
+	const uint32 bucket = BucketFromSize(requestedSize, BucketShift);
+	check(bucket < NumBuckets);
+	HRESULT hr = S_OK;
+
+	// Copy the resource desc so we can track it for all other allocations in this pool
+	if (bFirstAllocFromPool)
+	{
+		resourceDesc = Desc;
+		bFirstAllocFromPool = false;
+	}
+
+	// Ensure we're allocating from the correct pool
+	check(Desc.Flags == resourceDesc.Flags);
+
+#if D3D12_DEFAULT_BUFFER_POOL_ALWAYS_CREATE_NEW_RESOURCE
+	check(AvailableBuffers[bucket].IsEmpty());
+#endif
+
+	// Cleanup the pool
+	CleanupFreeBlocks();
+
+	// Check if there are any available buffers in the requested bucket
+	if (!AvailableBuffers[bucket].Dequeue(pBlock))
+	{
+		// No buffers are available in the requested bucket size, so make one
+		TRefCountPtr<FD3D12Resource> pBuffer = nullptr;
+		auto &ResourceHelper = GetParentDevice()->GetResourceHelper();
+		uint32 BlockSize = BlockSizeFromBufferSize(requestedSize, BucketShift);
+
+		// Use the nearest bucket size
+		resourceDesc.Width = BlockSize < MIN_HEAP_SIZE ? MIN_HEAP_SIZE : BlockSize;
+
+		check(BlockSize >= requestedSize);
+#if D3D12_DEFAULT_BUFFER_POOL_ALWAYS_CREATE_NEW_RESOURCE
+		resourceDesc.Width = requestedSize;
+#endif
+
+		VERIFYD3D11RESULT(hr = ResourceHelper.CreateDefaultResource(resourceDesc, pBuffer.GetInitReference()));
+
+#if UE_BUILD_DEBUG
+		const uint32 allocatedSize = resourceDesc.Width;
+		NumCreatedResources[bucket]++;
+		if (BucketSize[bucket] == 0)
+		{
+			BucketSize[bucket] = allocatedSize;
+		}
+#if !D3D12_DEFAULT_BUFFER_POOL_ALWAYS_CREATE_NEW_RESOURCE
+		check(BucketSize[bucket] == allocatedSize);
+#endif
+#endif
+
+		pBlock = new FD3D12ResourceBlockInfo(pBuffer, 0, 0, bucket, this);
+
+#if !D3D12_DEFAULT_BUFFER_POOL_ALWAYS_CREATE_NEW_RESOURCE
+#if SUB_ALLOCATED_DEFAULT_ALLOCATIONS
+		// Suballocate default resources if they are below the min heap size.
+		if (BlockSize < MIN_HEAP_SIZE)
+		{
+			// Create additional available blocks that can be sub-allocated from the same resource
+			for (uint32 Offset = BlockSize; Offset <= MIN_HEAP_SIZE - BlockSize; Offset += BlockSize)
+			{
+				FD3D12ResourceBlockInfo *NewBlock = new FD3D12ResourceBlockInfo(pBuffer, 0, Offset, bucket, this);
+
+				// Add the bucket to the available list
+				AvailableBuffers[bucket].Enqueue(NewBlock);
+			}
+	}
+#endif
+#endif
+	}
+
+	check(BlockSizeFromBufferSize(requestedSize, BucketShift) >= requestedSize);
+
+	// Grab a buffer from the front of the available bucket
+	check(pBlock != nullptr);
+
+	// Pass the resource location back to the caller
+#if D3D12_DEFAULT_BUFFER_POOL_ALWAYS_CREATE_NEW_RESOURCE
+	ResourceLocation->SetBlockInfo(pBlock, &GetParentDevice()->GetDefaultCommandContext().StateCache);	// Don't assign an allocator
+#else
+	ResourceLocation->SetBlockInfo(pBlock, &GetParentDevice()->GetDefaultCommandContext().StateCache);
+#endif
+	requestedSize -= Alignment;
+
+	ResourceLocation->EffectiveBufferSize = requestedSize;
+
+	if (Alignment > 0)
+	{
+		uint32 AlignmentMismatch = ResourceLocation->GetOffset() % Alignment;
+
+		if (AlignmentMismatch != 0)
+		{
+			ResourceLocation->SetPadding (Alignment - AlignmentMismatch);
+
+			check (ResourceLocation->GetOffset() % Alignment == 0);
+			check (ResourceLocation->GetOffset() % 4         == 0);
+		}
+	}
+
+	if (pInitialData)
+	{
+		FD3D12CommandListHandle& hCommandList = GetParentDevice()->GetDefaultCommandContext().CommandListHandle;
+		// Queue up a CopySubresource() to handle initial data...
+
+		// Get an upload heap and initialize data
+		FD3D12ResourceLocation SrcResourceLoc;
+		void *pData = GetParentDevice()->GetDefaultUploadHeapAllocator().FastAlloc(requestedSize, 4UL, &SrcResourceLoc);
+		check(pData);
+		FMemory::Memcpy(pData, pInitialData->pData, requestedSize);
+
+		// Copy from the temporary upload heap to the default resource
+		{
+			// Writable structured bufferes are sometimes initialized with inital data which means they sometimes need tracking.
+			FConditionalScopeResourceBarrier ConditionalScopeResourceBarrier(hCommandList, pBlock->ResourceHeap, D3D12_RESOURCE_STATE_COPY_DEST, 0);
+
+			GetParentDevice()->GetDefaultCommandContext().numCopies++;
+			hCommandList->CopyBufferRegion(
+				pBlock->ResourceHeap->GetResource(),
+				ResourceLocation->GetOffset(),
+				SrcResourceLoc.GetResource()->GetResource(),
+				SrcResourceLoc.GetOffset(), requestedSize);
+		}
+	}
+
+	return hr;
+}
+
+// Move the specified resources to the available pool
+void FD3D12DefaultBufferPool::ExpireBlock(FD3D12ResourceBlockInfo *Block)
+{
+	FD3D12CommandListManager& CommandListManager = GetParentDevice()->GetCommandListManager();
+	Block->FrameFence = CommandListManager.GetFence(FT_Frame).GetCurrentFence();
+	ExpiredBlocks.Enqueue(Block);
+}
+
+// Release all resources in the pool
+void FD3D12DefaultBufferPool::ReleaseAllResources()
+{
+	TRefCountPtr<FD3D12ResourceBlockInfo> pBlock;
+
+	// Empty the expired blocks list
+	while (ExpiredBlocks.Dequeue(pBlock))
+	{
+	}
+
+	// Empty all available blocks queues
+	for (uint32 i = 0; i < NumBuckets; ++i)
+	{
+		while (AvailableBuffers[i].Dequeue(pBlock))
+		{
+		}
+	}
+}
+
+// Grab a buffer from the available buffers or create a new buffer if none are available
+HRESULT FD3D12DefaultBufferAllocator::AllocDefaultResource(const D3D12_RESOURCE_DESC& Desc, D3D12_SUBRESOURCE_DATA *pInitialData, FD3D12ResourceLocation *ResourceLocation, uint32 Alignment)
+{
+	check (Desc.Flags < MAX_DEFAULT_POOLS);
+
+	return DefaultBufferPools[Desc.Flags]->AllocDefaultResource(Desc, pInitialData, ResourceLocation, Alignment);
+}
+
+HRESULT FD3D12DefaultBufferAllocator::AllocDefaultResource(const D3D12_RESOURCE_DESC& Desc, D3D12_SUBRESOURCE_DATA *pInitialData, FD3D12ResourceLocation **ResourceLocation, uint32 Alignment)
+{
+	*ResourceLocation = new FD3D12ResourceLocation(GetParentDevice(), Desc.Width);
+	(*ResourceLocation)->AddRef();
+
+	return AllocDefaultResource(Desc, pInitialData, *ResourceLocation, Alignment);
+}
+
+void FD3D12DefaultBufferAllocator::FreeDefaultBufferPools()
+{
+	for (uint32 i = 0; i < MAX_DEFAULT_POOLS; ++i)
+	{
+		delete DefaultBufferPools[i];
+		DefaultBufferPools[i] = nullptr;
+	}
+}
+
+void FD3D12DefaultBufferAllocator::CleanupFreeBlocks()
+{
+	for (uint32 i = 0; i < MAX_DEFAULT_POOLS; ++i)
+		DefaultBufferPools[i]->CleanupFreeBlocks();
+}
+
+HRESULT FD3D12ResourceHelper::CreateCommittedResource(const D3D12_RESOURCE_DESC& Desc, const D3D12_HEAP_PROPERTIES& HeapProps, FD3D12Resource** ppResource)
+{
+	return CreateCommittedResource(Desc, HeapProps, D3D12_RESOURCE_STATE_COMMON, ppResource);
+}
+
+HRESULT FD3D12ResourceHelper::CreateCommittedResource(const D3D12_RESOURCE_DESC& Desc, const D3D12_HEAP_PROPERTIES& HeapProps, const D3D12_RESOURCE_STATES& InitialUsage, FD3D12Resource** ppResource)
+{
+	return FD3D12DynamicRHI::CreateCommittedResource (Desc, HeapProps, InitialUsage, ppResource);
+}
+
+HRESULT FD3D12ResourceHelper::CreateDefaultResource(const D3D12_RESOURCE_DESC& Desc, FD3D12Resource** ppResource)
+{
+	CD3DX12_HEAP_PROPERTIES heapProperties(D3D12_HEAP_TYPE_DEFAULT);
+	return FD3D12DynamicRHI::CreateCommittedResource (Desc, heapProperties, D3D12_RESOURCE_STATE_COMMON, ppResource);
+}
+
+HRESULT FD3D12ResourceHelper::CreateBuffer(D3D12_HEAP_TYPE heapType, uint64 initHeapSize, FD3D12Resource** ppResource)
+{
+	TRefCountPtr<ID3D12Resource> pResource;
+	ID3D12Device *pD3DDevice = GetParentDevice()->GetDevice();
+
+	D3D12_RESOURCE_DESC BufDesc = CD3DX12_RESOURCE_DESC::Buffer(initHeapSize);
+	D3D12_RESOURCE_STATES InitialState = (heapType == D3D12_HEAP_TYPE_UPLOAD) ? D3D12_RESOURCE_STATE_GENERIC_READ : D3D12_RESOURCE_STATE_COPY_DEST;
+
+	if (!ppResource)
+	{
+		return E_POINTER;
+	}
+
+	CD3DX12_HEAP_PROPERTIES Props(heapType);
+
+	return FD3D12DynamicRHI::CreateCommittedResource(BufDesc, Props, InitialState, ppResource);
+}
+
+FD3D12ResourceHelper::FD3D12ResourceHelper(FD3D12Device* InParent) :
+	FD3D12DeviceChild(InParent)
+{}
+
+void FD3D12StateCacheBase::Init(FD3D12Device* InParent, FD3D12CommandContext* InCmdContext, const FD3D12StateCacheBase* AncestralState, bool bInAlwaysSetIndexBuffers)
+{
+	Parent = InParent;
+	CmdContext = InCmdContext;
+
+	DescriptorCache.Init(InParent, InCmdContext);
+
+	if (AncestralState)
+		InheritState(*AncestralState);
+	else
+		ClearState();
+
+	bAlwaysSetIndexBuffers = bInAlwaysSetIndexBuffers;
+}
+
+void FD3D12StateCacheBase::Clear()
+{
+	ClearState();
+
+	// Release references to cached objects
+	DescriptorCache.Clear();
+}
+
+void FD3D12StateCacheBase::ClearSRVs()
+{
+	if (bSRVSCleared)
+		return;
+
+	ZeroMemory(PipelineState.Common.CurrentShaderResourceViewsIntersectWithDepthRT, sizeof(PipelineState.Common.CurrentShaderResourceViewsIntersectWithDepthRT));
+	PipelineState.Common.ShaderResourceViewsIntersectWithDepthCount = 0;
+	for (uint32 ShaderFrequency = 0; ShaderFrequency < SF_NumFrequencies; ShaderFrequency++)
+	{
+		int32& MaxResourceIndex = PipelineState.Common.MaxBoundShaderResourcesIndex[ShaderFrequency];
+		for (int32 resIndex = 0; resIndex <= MaxResourceIndex; ++resIndex)
+		{
+			PipelineState.Common.CurrentShaderResourceViews[ShaderFrequency][resIndex] = nullptr;
+		}
+
+		MaxResourceIndex = INDEX_NONE;
+	}
+
+	bSRVSCleared = true;
+}
+
+void FD3D12StateCacheBase::ClearState()
+{
+	// Shader Resource View State Cache
+	bSRVSCleared = false;
+
+	ClearSRVs();
+
+	FMemory::Memzero(PipelineState.Common.CurrentShaderSamplerCounts, sizeof (PipelineState.Common.CurrentShaderSamplerCounts));
+	FMemory::Memzero(PipelineState.Common.CurrentShaderSRVCounts, sizeof (PipelineState.Common.CurrentShaderSRVCounts));
+	FMemory::Memzero(PipelineState.Common.CurrentShaderCBCounts, sizeof (PipelineState.Common.CurrentShaderCBCounts));
+	FMemory::Memzero(PipelineState.Common.CurrentShaderUAVCounts, sizeof (PipelineState.Common.CurrentShaderUAVCounts));
+
+	// Unordered Access View State Cache
+	for (int i = 0; i < D3D12_PS_CS_UAV_REGISTER_COUNT; ++i)
+		PipelineState.Common.UnorderedAccessViewArray[i] = nullptr;
+	PipelineState.Common.CurrentNumberOfSimultaneousUAVs = 0;
+
+	PipelineState.Graphics.HighLevelDesc.NumRenderTargets = 0;
+	PipelineState.Common.CurrentUAVStage = SF_Vertex;
+	PipelineState.Common.CurrentUAVStartSlot = -1;
+	PipelineState.Graphics.CurrentNumberOfStreamOutTargets = 0;
+	PipelineState.Graphics.CurrentNumberOfScissorRects = 0;
+
+	// Rasterizer State Cache
+	PipelineState.Graphics.HighLevelDesc.RasterizerState = nullptr;
+
+	// Depth Stencil State Cache
+	PipelineState.Graphics.CurrentReferenceStencil = 0;
+	PipelineState.Graphics.HighLevelDesc.DepthStencilState = nullptr;
+	PipelineState.Graphics.CurrentDepthStencilTarget = nullptr;
+
+	// Shader Cache
+	PipelineState.Graphics.HighLevelDesc.BoundShaderState = nullptr;
+	PipelineState.Compute.CurrentComputeShader = nullptr;
+
+	// Blend State Cache
+	PipelineState.Graphics.CurrentBlendFactor[0] = 1.0f;
+	PipelineState.Graphics.CurrentBlendFactor[1] = 1.0f;
+	PipelineState.Graphics.CurrentBlendFactor[2] = 1.0f;
+	PipelineState.Graphics.CurrentBlendFactor[3] = 1.0f;
+
+	FMemory::Memset(&PipelineState.Graphics.CurrentViewport[0], 0, sizeof(D3D12_VIEWPORT) * D3D12_VIEWPORT_AND_SCISSORRECT_OBJECT_COUNT_PER_PIPELINE);
+	PipelineState.Graphics.CurrentNumberOfViewports = 0;
+
+	PipelineState.Graphics.bNeedRebuildPSO = true;
+	PipelineState.Compute.bNeedRebuildPSO = true;
+	PipelineState.Graphics.CurrentPipelineStateObject = nullptr;
+	PipelineState.Compute.CurrentPipelineStateObject = nullptr;
+	PipelineState.Common.CurrentPipelineStateObject = nullptr;
+	FMemory::Memzero(PipelineState.Graphics.CurrentStreamOutTargets, sizeof (PipelineState.Graphics.CurrentStreamOutTargets));
+	FMemory::Memzero(PipelineState.Graphics.CurrentSOOffsets, sizeof (PipelineState.Graphics.CurrentSOOffsets));
+
+	FMemory::Memzero(PipelineState.Graphics.CurrentScissorRects, sizeof(PipelineState.Graphics.CurrentScissorRects));
+
+	D3D12_RECT ScissorRect;
+	ScissorRect.left = 0;
+	ScissorRect.right = GetMax2DTextureDimension();
+	ScissorRect.top = 0;
+	ScissorRect.bottom = GetMax2DTextureDimension();
+	SetScissorRects(1, &ScissorRect);
+
+	PipelineState.Graphics.HighLevelDesc.SampleMask = 0xffffffff;
+	PipelineState.Graphics.HighLevelDesc.BlendState = nullptr;
+
+	for(int i = 0; i < D3D12_IA_VERTEX_INPUT_RESOURCE_SLOT_COUNT; ++i)
+		PipelineState.Graphics.CurrentVertexBuffers[i].VertexBufferLocation = nullptr;
+	FMemory::Memzero(PipelineState.Common.CurrentSamplerStates, sizeof(PipelineState.Common.CurrentSamplerStates));
+
+	FMemory::Memzero(PipelineState.Graphics.RenderTargetArray, sizeof(PipelineState.Graphics.RenderTargetArray));
+
+	PipelineState.Graphics.CurrentIndexBufferLocation = nullptr;
+	PipelineState.Graphics.CurrentIndexFormat = DXGI_FORMAT_UNKNOWN;
+	PipelineState.Graphics.CurrentIndexOffset = 0;
+
+	PipelineState.Graphics.CurrentPrimitiveTopology = D3D_PRIMITIVE_TOPOLOGY_UNDEFINED;
+
+	bAutoFlushComputeShaderCache = true;
+}
+
+void FD3D12StateCacheBase::RestoreState()
+{
+	// Mark bits dirty so the next call to ApplyState will set all this state again
+	PipelineState.Common.bNeedSetPSO = true;
+	bNeedSetVB = true;
+	bNeedSetIB = true;
+	bNeedSetSOs = true;
+	bNeedSetUAVs = true;
+	bNeedSetRTs = true;
+	bNeedSetViewports = true;
+	bNeedSetScissorRects = true;
+	bNeedSetPrimitiveTopology = true;
+	bNeedSetBlendFactor = true;
+	bNeedSetStencilRef = true;
+	bNeedSetSamplers = true;
+	bNeedSetSRVs = true;
+	bNeedSetConstantBuffers = true;
+	for (int i = 0; i < SF_NumFrequencies; i++)
+	{
+		bNeedSetSamplersPerShaderStage[i] = true;
+		bNeedSetSRVsPerShaderStage[i] = true;
+		bNeedSetConstantBuffersPerShaderStage[i] = true;
+	}
+}
+
+D3D12_STATE_CACHE_INLINE void FD3D12StateCacheBase::SetViewport(D3D12_VIEWPORT Viewport)
+{
+	if ((PipelineState.Graphics.CurrentNumberOfViewports != 1 || FMemory::Memcmp(&PipelineState.Graphics.CurrentViewport[0], &Viewport, sizeof(D3D12_VIEWPORT))) || GD3D12SkipStateCaching)
+	{
+		FMemory::Memcpy(&PipelineState.Graphics.CurrentViewport[0], &Viewport, sizeof(D3D12_VIEWPORT));
+		PipelineState.Graphics.CurrentNumberOfViewports = 1;
+		bNeedSetViewports = true;
+		UpdateViewportScissorRects();
+	}
+}
+
+D3D12_STATE_CACHE_INLINE void FD3D12StateCacheBase::SetViewports(uint32 Count, D3D12_VIEWPORT* Viewports)
+{
+	if ((PipelineState.Graphics.CurrentNumberOfViewports != Count || FMemory::Memcmp(&PipelineState.Graphics.CurrentViewport[0], Viewports, sizeof(D3D12_VIEWPORT)* Count)) || GD3D12SkipStateCaching)
+	{
+		FMemory::Memcpy(&PipelineState.Graphics.CurrentViewport[0], Viewports, sizeof(D3D12_VIEWPORT)* Count);
+		PipelineState.Graphics.CurrentNumberOfViewports = Count;
+		bNeedSetViewports = true;
+		UpdateViewportScissorRects();
+	}
+}
+
+D3D12_STATE_CACHE_INLINE void FD3D12StateCacheBase::SetScissorRect(D3D12_RECT ScissorRect)
+{
+	if ((PipelineState.Graphics.CurrentNumberOfScissorRects != 1 || FMemory::Memcmp(&PipelineState.Graphics.CurrentScissorRects[0], &ScissorRect, sizeof(D3D12_RECT))) || GD3D12SkipStateCaching)
+	{
+		FMemory::Memcpy(&PipelineState.Graphics.CurrentScissorRects[0], &ScissorRect, sizeof(D3D12_RECT));
+		PipelineState.Graphics.CurrentNumberOfScissorRects = 1;
+		UpdateViewportScissorRects();
+	}
+}
+
+D3D12_STATE_CACHE_INLINE void FD3D12StateCacheBase::SetScissorRects(uint32 Count, D3D12_RECT* ScissorRects)
+{
+	if ((PipelineState.Graphics.CurrentNumberOfScissorRects != Count || FMemory::Memcmp(&PipelineState.Graphics.CurrentScissorRects[0], ScissorRects, sizeof(D3D12_RECT)* Count)) || GD3D12SkipStateCaching)
+	{
+		FMemory::Memcpy(&PipelineState.Graphics.CurrentScissorRects[0], ScissorRects, sizeof(D3D12_RECT)* Count);
+		PipelineState.Graphics.CurrentNumberOfScissorRects = Count;
+		UpdateViewportScissorRects();
+	}
+}
+
+void FD3D12HighLevelGraphicsPipelineStateDesc::GetLowLevelDesc(ID3D12RootSignature* rootSignature, FD3D12LowLevelGraphicsPipelineStateDesc& psoDesc)
+{
+	ZeroMemory(&psoDesc, sizeof(psoDesc));
+
+	psoDesc.Desc.pRootSignature = rootSignature;
+	psoDesc.Desc.SampleMask = SampleMask;
+	psoDesc.Desc.PrimitiveTopologyType = PrimitiveTopologyType;
+	psoDesc.Desc.NumRenderTargets = NumRenderTargets;
+	FMemory::Memcpy(psoDesc.Desc.RTVFormats, RTVFormats, sizeof(RTVFormats[0]) * NumRenderTargets);
+	psoDesc.Desc.DSVFormat = DSVFormat;
+	psoDesc.Desc.SampleDesc = SampleDesc;
+
+	psoDesc.Desc.InputLayout = BoundShaderState->InputLayout;
+
+	if (BoundShaderState->GetGeometryShader())
+		psoDesc.Desc.StreamOutput = BoundShaderState->GetGeometryShader()->StreamOutput;
+
+#define COPY_SHADER(Initial, Name) \
+	if (FD3D12##Name##Shader* Shader = BoundShaderState->Get##Name##Shader()) \
+	{ \
+		psoDesc.Desc.Initial##S = Shader->ShaderBytecode.GetShaderBytecode(); \
+		psoDesc.Initial##SHash = Shader->ShaderBytecode.GetHash(); \
+	}
+	COPY_SHADER(V, Vertex);
+	COPY_SHADER(P, Pixel);
+	COPY_SHADER(D, Domain);
+	COPY_SHADER(H, Hull);
+	COPY_SHADER(G, Geometry);
+#undef COPY_SHADER
+
+	psoDesc.Desc.BlendState = BlendState ? *BlendState : CD3DX12_BLEND_DESC(D3D12_DEFAULT);
+	psoDesc.Desc.RasterizerState = RasterizerState ? *RasterizerState : CD3DX12_RASTERIZER_DESC(D3D12_DEFAULT);
+	psoDesc.Desc.DepthStencilState = DepthStencilState ? *DepthStencilState : CD3DX12_DEPTH_STENCIL_DESC(D3D12_DEFAULT);
+}
+
+void FD3D12StateCacheBase::ApplyState(bool IsCompute)
+{
+	SCOPE_CYCLE_COUNTER(STAT_D3D12ApplyStateTime);
+	const bool bForceState = false;
+
+	FD3D12CommandListHandle& CommandList = CmdContext->CommandListHandle;
+	FD3D12PipelineStateCache& PSOCache = GetParentDevice()->GetPSOCache();
+
+	// PSO
+	if (IsCompute)
+	{
+		if (PipelineState.Compute.bNeedRebuildPSO || bForceState)
+		{
+			FD3D12ComputePipelineStateDesc psoDescCompute;
+			ZeroMemory(&psoDescCompute, sizeof(psoDescCompute));
+
+			psoDescCompute.Desc.pRootSignature = GetParentDevice()->GetComputeRootSignature();
+			psoDescCompute.Desc.CS = PipelineState.Compute.CurrentComputeShader->ShaderBytecode.GetShaderBytecode();
+			psoDescCompute.CSHash = PipelineState.Compute.CurrentComputeShader->ShaderBytecode.GetHash();
+
+			ID3D12PipelineState* psoCompute = PSOCache.FindCompute(psoDescCompute, true);
+
+			// Save the current PSO
+			PipelineState.Common.CurrentPipelineStateObject = psoCompute;
+			PipelineState.Compute.CurrentPipelineStateObject = psoCompute;
+
+			// Indicate we need to set the PSO on the command list
+			PipelineState.Common.bNeedSetPSO = true;
+			PipelineState.Compute.bNeedRebuildPSO = false;
+		}
+	}
+	else
+	{
+		if (PipelineState.Graphics.bNeedRebuildPSO || bForceState)
+		{
+			// The desc is mostly initialized, just need to copy the RTV/DSV formats and sample properties in
+			FD3D12HighLevelGraphicsPipelineStateDesc& psoDesc = PipelineState.Graphics.HighLevelDesc;
+
+			ZeroMemory(psoDesc.RTVFormats, sizeof(psoDesc.RTVFormats));
+			psoDesc.SampleDesc.Count = psoDesc.SampleDesc.Quality = 0;
+			for (uint32 i = 0; i < psoDesc.NumRenderTargets; i++)
+			{
+				if (PipelineState.Graphics.RenderTargetArray[i] != NULL)
+				{
+					const D3D12_RENDER_TARGET_VIEW_DESC &desc = PipelineState.Graphics.RenderTargetArray[i]->GetDesc();
+					D3D12_RESOURCE_DESC const& resDesc = PipelineState.Graphics.RenderTargetArray[i]->GetResource()->GetDesc();
+
+					if (desc.Format == DXGI_FORMAT_UNKNOWN)
+					{
+						psoDesc.RTVFormats[i] = resDesc.Format;
+					}
+					else
+					{
+						psoDesc.RTVFormats[i] = desc.Format;
+					}
+					check(psoDesc.RTVFormats[i] != DXGI_FORMAT_UNKNOWN);
+
+					if (psoDesc.SampleDesc.Count == 0)
+					{
+						psoDesc.SampleDesc.Count = resDesc.SampleDesc.Count;
+						psoDesc.SampleDesc.Quality = resDesc.SampleDesc.Quality;
+					}
+				}
+			}
+
+			psoDesc.DSVFormat = DXGI_FORMAT_UNKNOWN;
+			if (PipelineState.Graphics.CurrentDepthStencilTarget)
+			{
+				const D3D12_DEPTH_STENCIL_VIEW_DESC &dsvDesc = PipelineState.Graphics.CurrentDepthStencilTarget->GetDesc();
+				D3D12_RESOURCE_DESC const& resDesc = PipelineState.Graphics.CurrentDepthStencilTarget->GetResource()->GetDesc();
+				
+				psoDesc.DSVFormat = dsvDesc.Format;
+				if (psoDesc.NumRenderTargets == 0 || psoDesc.SampleDesc.Count == 0)
+				{
+					psoDesc.SampleDesc.Count = resDesc.SampleDesc.Count;
+					psoDesc.SampleDesc.Quality = resDesc.SampleDesc.Quality;
+				}
+			}
+
+			ID3D12PipelineState* psoGraphics = 
+				PSOCache.FindGraphics(PipelineState.Graphics.HighLevelDesc, GetParentDevice()->GetGraphicsRootSignature(), true);
+
+			// Save the current PSO
+			PipelineState.Common.CurrentPipelineStateObject = psoGraphics;
+			PipelineState.Graphics.CurrentPipelineStateObject = psoGraphics;
+
+			// Indicate we need to set the PSO on the command list
+			PipelineState.Common.bNeedSetPSO = true;
+			PipelineState.Graphics.bNeedRebuildPSO = false;
+		}
+	}
+
+	// See if we need to set our PSO:
+	// In D3D11, you could Set dispatch arguments, then set Draw arguments, then call Draw/Dispatch/Draw/Dispatch without setting arguments again.
+	// In D3D12, we need to understand when the app switches between Draw/Dispatch and make sure the correct PSO is set.
+	if (PipelineState.Common.bNeedSetPSO || bForceState)
+	{
+		// Set the PSO, could be a Compute or Graphics PSO
+		CommandList->SetPipelineState(PipelineState.Common.CurrentPipelineStateObject);
+		PipelineState.Common.bNeedSetPSO = false;
+	}
+
+	if (!IsCompute)
+	{
+		// Setup non-heap bindings
+		if (bNeedSetVB || bForceState)
+		{
+			SCOPE_CYCLE_COUNTER(STAT_D3D12ApplyStateSetVertexBufferTime);
+			DescriptorCache.SetVertexBuffers(PipelineState.Graphics.CurrentVertexBuffers, D3D12_IA_VERTEX_INPUT_RESOURCE_SLOT_COUNT);
+			bNeedSetVB = false;
+		}
+		if (bNeedSetIB || bForceState)
+		{
+			if (PipelineState.Graphics.CurrentIndexBufferLocation != nullptr)
+			{
+				DescriptorCache.SetIndexBuffer(PipelineState.Graphics.CurrentIndexBufferLocation, PipelineState.Graphics.CurrentIndexFormat, PipelineState.Graphics.CurrentIndexOffset);
+			}
+			bNeedSetIB = false;
+		}
+		if (bNeedSetSOs || bForceState)
+		{
+			DescriptorCache.SetStreamOutTargets(PipelineState.Graphics.CurrentStreamOutTargets, PipelineState.Graphics.CurrentNumberOfStreamOutTargets, PipelineState.Graphics.CurrentSOOffsets);
+			bNeedSetSOs = false;
+		}
+
+		if (bNeedSetViewports || bForceState)
+		{
+			CommandList->RSSetViewports(PipelineState.Graphics.CurrentNumberOfViewports, PipelineState.Graphics.CurrentViewport);
+			bNeedSetViewports = false;
+		}
+		if (bNeedSetScissorRects || bForceState)
+		{
+			CommandList->RSSetScissorRects(PipelineState.Graphics.CurrentNumberOfScissorRects, PipelineState.Graphics.CurrentViewportScissorRects);
+			bNeedSetScissorRects = false;
+		}
+		if (bNeedSetPrimitiveTopology || bForceState)
+		{
+			CommandList->IASetPrimitiveTopology(PipelineState.Graphics.CurrentPrimitiveTopology);
+			bNeedSetPrimitiveTopology = false;
+		}
+		if (bNeedSetBlendFactor || bForceState)
+		{
+			CommandList->OMSetBlendFactor(PipelineState.Graphics.CurrentBlendFactor);
+			bNeedSetBlendFactor = false;
+		}
+		if (bNeedSetStencilRef || bForceState)
+		{
+			CommandList->OMSetStencilRef(PipelineState.Graphics.CurrentReferenceStencil);
+			bNeedSetStencilRef = false;
+		}
+		if (bNeedSetRTs || bForceState)
+		{
+			const bool bDepthIsBoundAsSRV = PipelineState.Common.ShaderResourceViewsIntersectWithDepthCount > 0;
+			DescriptorCache.SetRenderTargets(PipelineState.Graphics.RenderTargetArray, PipelineState.Graphics.HighLevelDesc.NumRenderTargets, PipelineState.Graphics.CurrentDepthStencilTarget, bDepthIsBoundAsSRV);
+			bNeedSetRTs = false;
+		}
+	}
+
+	// Reserve space in descriptor heaps
+	// Since this can cause heap rollover (which causes old bindings to become invalid), the reserve must be done atomically
+	const uint32 StartStage = IsCompute ? SF_Compute : 0;
+	const uint32 EndStage = IsCompute ? SF_NumFrequencies : SF_Compute;
+	uint32 NumUAVs = 0;
+	uint32 NumSRVs[SF_NumFrequencies + 1] = {};
+	uint32 NumCBs[SF_NumFrequencies + 1] = {};
+	uint32 NumSamplers[SF_NumFrequencies + 1] = {};
+
+	uint32 ViewHeapSlot = 0;
+	uint32 SamplerHeapSlot = 0;
+
+	for (uint32 iTries = 0; iTries < 2; ++iTries)
+	{
+		if (bNeedSetUAVs || bForceState)
+		{
+			NumUAVs = PipelineState.Common.CurrentNumberOfSimultaneousUAVs;
+		}
+		for (uint32 Stage = StartStage; Stage < EndStage; ++Stage)
+		{
+			if ((bNeedSetSRVs && bNeedSetSRVsPerShaderStage[Stage]) || bForceState)
+			{
+				NumSRVs[Stage] = PipelineState.Common.MaxBoundShaderResourcesIndex[Stage] + 1;
+			}
+			if ((bNeedSetConstantBuffers && (bNeedSetConstantBuffersPerShaderStage[Stage])) || bForceState)
+			{
+				NumCBs[Stage] = PipelineState.Common.CurrentShaderCBCounts[Stage];
+			}
+			NumSRVs[SF_NumFrequencies] += NumSRVs[Stage];
+			NumCBs[SF_NumFrequencies] += NumCBs[Stage];
+		}
+
+		uint32 NumViews = NumUAVs + NumSRVs[SF_NumFrequencies] + NumCBs[SF_NumFrequencies];
+		if (!DescriptorCache.ViewHeap.CanReserveSlots(NumViews))
+		{
+			DescriptorCache.ViewHeap.RollOver();
+			NumSRVs[SF_NumFrequencies] = 0;
+			NumCBs[SF_NumFrequencies] = 0;
+
+			bNeedSetUAVs = true;
+			bNeedSetConstantBuffers = true;
+			bNeedSetSRVs = true;
+			memset(bNeedSetConstantBuffersPerShaderStage, 0xffffffff, sizeof(bNeedSetConstantBuffersPerShaderStage));
+			memset(bNeedSetSRVsPerShaderStage, 0xffffffff, sizeof(bNeedSetSRVsPerShaderStage));
+			continue;
+		}
+		ViewHeapSlot = DescriptorCache.ViewHeap.ReserveSlots(NumViews);
+		break;
+	}
+
+	for (uint32 iTries = 0; iTries < 2; ++iTries)
+	{
+		if (bNeedSetSamplers)
+		{
+			for (uint32 Stage = StartStage; Stage < EndStage; ++Stage)
+			{
+				if (bNeedSetSamplersPerShaderStage[Stage] || bForceState)
+				{
+					NumSamplers[Stage] = PipelineState.Common.CurrentShaderSamplerCounts[Stage];
+				}
+				NumSamplers[SF_NumFrequencies] += NumSamplers[Stage];
+			}
+		}
+		if (!DescriptorCache.SamplerHeap.CanReserveSlots(NumSamplers[SF_NumFrequencies]))
+		{
+			DescriptorCache.SamplerHeap.RollOver();
+			NumSamplers[SF_NumFrequencies] = 0;
+
+			bNeedSetSamplers = true;
+			memset(bNeedSetSamplersPerShaderStage, 0xffffffff, sizeof(bNeedSetSamplersPerShaderStage));
+			continue;
+		}
+		SamplerHeapSlot = DescriptorCache.SamplerHeap.ReserveSlots(NumSamplers[SF_NumFrequencies]);
+		break;
+	}
+
+	if ((bNeedSetUAVs || bForceState) && PipelineState.Common.CurrentNumberOfSimultaneousUAVs > 0)
+	{
+		SCOPE_CYCLE_COUNTER(STAT_D3D12ApplyStateSetUAVTime);
+		DescriptorCache.SetUAVs(PipelineState.Common.CurrentUAVStage, PipelineState.Common.CurrentUAVStartSlot, PipelineState.Common.UnorderedAccessViewArray, NumUAVs, ViewHeapSlot);
+		bNeedSetUAVs = false;
+	}
+	if (bNeedSetSamplers || bForceState)
+	{
+		for (uint32 i = StartStage; i < EndStage; i++)
+		{
+			if (bNeedSetSamplersPerShaderStage[i] || bForceState)
+			{
+				DescriptorCache.SetSamplers((EShaderFrequency)i, PipelineState.Common.CurrentSamplerStates[i], NumSamplers[i], SamplerHeapSlot);
+				bNeedSetSamplersPerShaderStage[i] = false;
+			}
+		}
+		bNeedSetSamplers = false;
+
+		DescriptorCache.SamplerHeap.SetNextSlot(SamplerHeapSlot);
+	}
+
+	// Shader resource views
+	if (bNeedSetSRVs || bForceState)
+	{
+		SCOPE_CYCLE_COUNTER(STAT_D3D12ApplyStateSetSRVTime);
+		for (uint32 i = StartStage; i < EndStage; i++)
+		{
+			if ((NumSRVs[i] > 0) || bForceState)
+			{
+				DescriptorCache.SetSRVs((EShaderFrequency)i, PipelineState.Common.CurrentShaderResourceViews[i], NumSRVs[i], PipelineState.Common.CurrentShaderResourceViewsIntersectWithDepthRT[i], ViewHeapSlot);
+				bNeedSetSRVsPerShaderStage[i] = false;
+			}
+		}
+		bNeedSetSRVs = false;
+	}
+
+	// Constant buffers
+	if (bNeedSetConstantBuffers || bForceState)
+	{
+		SCOPE_CYCLE_COUNTER(STAT_D3D12ApplyStateSetConstantBufferTime);
+		for (uint32 i = StartStage; i < EndStage; i++)
+		{
+			if (bNeedSetConstantBuffersPerShaderStage[i] || bForceState)
+			{
+				DescriptorCache.SetConstantBuffers((EShaderFrequency)i, NumCBs[i], ViewHeapSlot);
+				bNeedSetConstantBuffersPerShaderStage[i] = false;
+			}
+		}
+		bNeedSetConstantBuffers = false;
+	}
+
+#if DEBUG_RESOURCE_STATES
+	bool bSucceeded = VerifyResourceStates(IsCompute);
+	check(bSucceeded);
+#endif
+}
+
+void FDiskCacheInterface::Init(FString &filename)
+{
+	mFileStart = nullptr;
+	hFile = 0;
+	hMemoryMap = 0;
+	hMapAddress = 0;
+	mCurrentFileMapSize = 0;
+	mCurrentOffset = 0;
+	mCacheExists = false;
+	mInErrorState = false;
+
+	mFileName = filename;
+	mCacheExists = true;
+	WIN32_FIND_DATA fileData;
+	FindFirstFile(mFileName.GetCharArray().GetData(), &fileData);
+	if (GetLastError() == ERROR_FILE_NOT_FOUND)
+	{
+		mCacheExists = false;
+	}
+	bool fileFound = mCacheExists;
+	mCurrentFileMapSize = 1;
+	GrowMapping(64 * 1024, true);
+
+	if (fileFound && mFileStart)
+	{
+		mHeader = *(FDiskCacheHeader*)mFileStart;
+		if (mHeader.mHeaderVersion != mCurrentHeaderVersion || (!mHeader.mUsesAPILibraries && FD3D12PipelineStateCache::bUseAPILibaries))
+		{
+			UE_LOG(LogD3D12RHI, Warning, TEXT("Disk cache is stale. Disk Cache version: %d App version: %d"), mHeader.mHeaderVersion, mCurrentHeaderVersion);
+			ClearDiskCache();
+		}
+	}
+	else
+	{
+		mHeader.mHeaderVersion = mCurrentHeaderVersion;
+		mHeader.mNumPsos = 0;
+		mHeader.mUsesAPILibraries = FD3D12PipelineStateCache::bUseAPILibaries;
+	}
+}
+
+void FDiskCacheInterface::GrowMapping(SIZE_T size, bool firstrun)
+{
+	if (mInErrorState)
+	{
+		return;
+	}
+
+	if (mCurrentFileMapSize - mCurrentOffset < size)
+	{
+		while ((mCurrentFileMapSize - mCurrentOffset) < size){ mCurrentFileMapSize += mFileGrowSize; }
+	}
+	else
+	{
+		return;
+	}
+
+	if (hMapAddress){
+		FlushViewOfFile(hMapAddress, mCurrentOffset);
+		UnmapViewOfFile(hMapAddress);
+	}
+	if (hMemoryMap){ CloseHandle(hMemoryMap); }
+	if (hFile){ CloseHandle(hFile); }
+
+	uint32 flag = (mCacheExists) ? OPEN_EXISTING : CREATE_NEW;
+	// open the shader cache file
+	hFile = CreateFile(mFileName.GetCharArray().GetData(), GENERIC_READ | GENERIC_WRITE, FILE_SHARE_READ | FILE_SHARE_WRITE, NULL, flag, FILE_ATTRIBUTE_NORMAL, NULL);
+	if (hFile == INVALID_HANDLE_VALUE)
+	{
+		//error state!
+		mInErrorState = true;
+		return;
+	}
+
+	mCacheExists = true;
+
+	uint32 fileSize = GetFileSize(hFile, NULL);
+	if (fileSize == 0)
+	{
+		byte data[64];
+		FMemory::Memset(data, NULL, _countof(data));
+		//It's invalide to map a zero sized file so write some junk data in that case
+		WriteFile(hFile, data, _countof(data), NULL, NULL);
+	}
+	else if (firstrun)
+	{
+		mCurrentFileMapSize = fileSize;
+	}
+
+	hMemoryMap = CreateFileMapping(hFile, NULL, PAGE_READWRITE, 0, (uint32)mCurrentFileMapSize, NULL);
+	if (hMemoryMap == (HANDLE)nullptr)
+	{
+		//error state!
+		mInErrorState = true;
+		ClearDiskCache();
+		return;
+	}
+
+	hMapAddress = MapViewOfFile(hMemoryMap, FILE_MAP_ALL_ACCESS, 0, 0, mCurrentFileMapSize);
+	if (hMapAddress == (HANDLE)nullptr)
+	{
+		//error state!
+		mInErrorState = true;
+		ClearDiskCache();
+		return;
+	}
+
+	mFileStart = (byte*)hMapAddress;
+}
+
+bool FDiskCacheInterface::AppendData(void* pData, size_t size)
+{
+	GrowMapping(size, false);
+	if (IsInErrorState())
+	{
+		return false;
+	}
+
+	memcpy((mFileStart + mCurrentOffset), pData, size);
+	mCurrentOffset += size;
+
+	return true;
+}
+
+bool FDiskCacheInterface::SetPointerAndAdvanceFilePosition(void** pDest, size_t size, bool backWithSystemMemory)
+{
+	GrowMapping(size, false);
+	if (IsInErrorState())
+	{
+		return false;
+	}
+
+	// Back persistent objects in system memory to avoid
+	// troubles when re-mapping the file.
+	if (backWithSystemMemory)
+	{
+		// Optimization: most (all?) of the shader input layout semantic names are "ATTRIBUTE"...
+		// instead of making 1000's of attribute strings, just set it to a static one.
+		static const char attribute[] = "ATTRIBUTE";
+		if (size == sizeof(attribute) && FMemory::Memcmp((mFileStart + mCurrentOffset), attribute, sizeof(attribute)) == 0)
+		{
+			*pDest = (void*)attribute;
+		}
+		else
+		{
+			void* newMemory = FMemory::Malloc(size);
+			if (newMemory)
+			{
+				FMemory::Memcpy(newMemory, (mFileStart + mCurrentOffset), size);
+				mBackedMemory.Add(newMemory);
+				*pDest = newMemory;
+			}
+			else
+			{
+				check(false);
+				return false;
+			}
+		}
+	}
+	else
+	{
+		*pDest = mFileStart + mCurrentOffset;
+	}
+
+	mCurrentOffset += size;
+
+	return true;
+}
+
+void FDiskCacheInterface::Reset()
+{
+	mCurrentOffset = sizeof(FDiskCacheHeader);
+}
+
+void FDiskCacheInterface::Close(uint32 numberOfPSOs)
+{
+	mHeader.mNumPsos = numberOfPSOs;
+	mHeader.mUsesAPILibraries = FD3D12PipelineStateCache::bUseAPILibaries;
+
+	if (!IsInErrorState())
+	{
+		if (hMapAddress){
+			*(FDiskCacheHeader*)mFileStart = mHeader;
+			FlushViewOfFile(hMapAddress, mCurrentOffset);
+			UnmapViewOfFile(hMapAddress);
+		}
+		if (hMemoryMap){ CloseHandle(hMemoryMap); }
+		if (hFile){ CloseHandle(hFile); }
+	}
+}
+
+void FDiskCacheInterface::ClearDiskCache()
+{
+	// Prevent reads/writes 
+	mInErrorState = true;
+	mHeader.mHeaderVersion = mCurrentHeaderVersion;
+	mHeader.mNumPsos = 0;
+	mHeader.mUsesAPILibraries = FD3D12PipelineStateCache::bUseAPILibaries;
+
+	if (hMapAddress){ UnmapViewOfFile(hMapAddress); }
+	if (hMemoryMap){ CloseHandle(hMemoryMap); }
+	if (hFile){ CloseHandle(hFile); }
+#ifdef UNICODE
+	BOOL result = DeleteFileW(mFileName.GetCharArray().GetData());
+#else
+	bool result = DeleteFileA(mFileName.GetCharArray().GetData());
+#endif
+	UE_LOG(LogD3D12RHI, Warning, TEXT("Deleted PSO Cache with result %d"), result);
+}
+
+void FDiskCacheInterface::Flush(uint32 numberOfPSOs)
+{
+	mHeader.mNumPsos = numberOfPSOs;
+	mHeader.mUsesAPILibraries = FD3D12PipelineStateCache::bUseAPILibaries;
+
+	if (hMapAddress && !IsInErrorState()){
+		*(FDiskCacheHeader*)mFileStart = mHeader;
+		FlushViewOfFile(hMapAddress, mCurrentOffset);
+	}
+}
+
+void FD3D12PipelineStateCache::RebuildFromDiskCache(ID3D12RootSignature* pGraphicsRootSig, ID3D12RootSignature* pComputeRootSignature)
+{
+	FScopeLock Lock (&CS);
+	if (DiskCaches[PSO_CACHE_GRAPHICS].IsInErrorState() || DiskCaches[PSO_CACHE_COMPUTE].IsInErrorState())
+	{
+		return;
+	}
+	// The only time shader code is ever read back is on debug builds
+	// when it checks for hash collisions in the PSO map. Therefore
+	// there is no point backing the memory on release.
+#if UE_BUILD_DEBUG
+	static const bool bBackShadersWithSystemMemory = true;
+#else
+	static const bool bBackShadersWithSystemMemory = false;
+#endif
+
+	DiskCaches[PSO_CACHE_GRAPHICS].Reset();
+	DiskCaches[PSO_CACHE_COMPUTE].Reset();
+
+	uint32 NumGraphicsPSOs = DiskCaches[PSO_CACHE_GRAPHICS].GetNumPSOs();
+
+	for (uint32 i = 0; i < NumGraphicsPSOs; i++)
+	{
+		FD3D12LowLevelGraphicsPipelineStateDesc* GraphicsPSODesc = nullptr;
+		DiskCaches[PSO_CACHE_GRAPHICS].SetPointerAndAdvanceFilePosition((void**)&GraphicsPSODesc, sizeof(*GraphicsPSODesc));
+		D3D12_GRAPHICS_PIPELINE_STATE_DESC* PSODesc = &GraphicsPSODesc->Desc;
+
+		PSODesc->pRootSignature = pGraphicsRootSig;
+
+		if (PSODesc->InputLayout.NumElements)
+		{
+			DiskCaches[PSO_CACHE_GRAPHICS].SetPointerAndAdvanceFilePosition((void**)&PSODesc->InputLayout.pInputElementDescs, PSODesc->InputLayout.NumElements * sizeof(D3D12_INPUT_ELEMENT_DESC), true);
+			for (uint32 i = 0; i < PSODesc->InputLayout.NumElements; i++)
+			{
+			   // Get the Sematic name string
+				uint32* stringLength = nullptr;
+				DiskCaches[PSO_CACHE_GRAPHICS].SetPointerAndAdvanceFilePosition((void**)&stringLength, sizeof(uint32));
+				DiskCaches[PSO_CACHE_GRAPHICS].SetPointerAndAdvanceFilePosition((void**)&PSODesc->InputLayout.pInputElementDescs[i].SemanticName, *stringLength, true);
+			}
+		}
+		if (PSODesc->StreamOutput.NumEntries)
+		{
+			DiskCaches[PSO_CACHE_GRAPHICS].SetPointerAndAdvanceFilePosition((void**)&PSODesc->StreamOutput.pSODeclaration, PSODesc->StreamOutput.NumEntries * sizeof(D3D12_SO_DECLARATION_ENTRY), true);
+			for (uint32 i = 0; i < PSODesc->StreamOutput.NumEntries; i++)
+			{
+				//Get the Sematic name string
+				uint32* stringLength = nullptr;
+				DiskCaches[PSO_CACHE_GRAPHICS].SetPointerAndAdvanceFilePosition((void**)&stringLength, sizeof(uint32));
+				DiskCaches[PSO_CACHE_GRAPHICS].SetPointerAndAdvanceFilePosition((void**)&PSODesc->StreamOutput.pSODeclaration[i].SemanticName, *stringLength, true);
+			}
+		}
+		if (PSODesc->StreamOutput.NumStrides)
+		{
+			DiskCaches[PSO_CACHE_GRAPHICS].SetPointerAndAdvanceFilePosition((void**)&PSODesc->StreamOutput.pBufferStrides, PSODesc->StreamOutput.NumStrides * sizeof(uint32), true);
+		}
+		if (PSODesc->VS.BytecodeLength)
+		{
+			DiskCaches[PSO_CACHE_GRAPHICS].SetPointerAndAdvanceFilePosition((void**)&PSODesc->VS.pShaderBytecode, PSODesc->VS.BytecodeLength, bBackShadersWithSystemMemory);
+		}
+		if (PSODesc->PS.BytecodeLength)
+		{
+			DiskCaches[PSO_CACHE_GRAPHICS].SetPointerAndAdvanceFilePosition((void**)&PSODesc->PS.pShaderBytecode, PSODesc->PS.BytecodeLength, bBackShadersWithSystemMemory);
+		}
+		if (PSODesc->DS.BytecodeLength)
+		{
+			DiskCaches[PSO_CACHE_GRAPHICS].SetPointerAndAdvanceFilePosition((void**)&PSODesc->DS.pShaderBytecode, PSODesc->DS.BytecodeLength, bBackShadersWithSystemMemory);
+		}
+		if (PSODesc->HS.BytecodeLength)
+		{
+			DiskCaches[PSO_CACHE_GRAPHICS].SetPointerAndAdvanceFilePosition((void**)&PSODesc->HS.pShaderBytecode, PSODesc->HS.BytecodeLength, bBackShadersWithSystemMemory);
+		}
+		if (PSODesc->GS.BytecodeLength)
+		{
+			DiskCaches[PSO_CACHE_GRAPHICS].SetPointerAndAdvanceFilePosition((void**)&PSODesc->GS.pShaderBytecode, PSODesc->GS.BytecodeLength, bBackShadersWithSystemMemory);
+		}
+
+		SIZE_T* cachedBlobSize = nullptr;
+		DiskCaches[PSO_CACHE_GRAPHICS].SetPointerAndAdvanceFilePosition((void**)&cachedBlobSize, sizeof(SIZE_T));
+
+		check(cachedBlobSize);
+		if (bUseAPILibaries && cachedBlobSize) { check(*cachedBlobSize); }
+
+		if (bUseAPILibaries && cachedBlobSize &&*cachedBlobSize)
+		{
+			GraphicsPSODesc->Desc.CachedPSO.CachedBlobSizeInBytes = *cachedBlobSize;
+			DiskCaches[PSO_CACHE_GRAPHICS].SetPointerAndAdvanceFilePosition((void**)&GraphicsPSODesc->Desc.CachedPSO.pCachedBlob, *cachedBlobSize);
+		}
+		else
+		{
+			GraphicsPSODesc->Desc.CachedPSO.CachedBlobSizeInBytes = 0;
+			GraphicsPSODesc->Desc.CachedPSO.pCachedBlob = nullptr;
+		}
+
+		if (!DiskCaches[PSO_CACHE_GRAPHICS].IsInErrorState())
+		{
+			GraphicsPSODesc->CombinedHash = FD3D12PipelineStateCache::HashPSODesc(*GraphicsPSODesc);
+
+			if (this->Add(*GraphicsPSODesc) == nullptr)
+			{
+				// The cache has failed, build PSOs at runtime
+				return;
+			}
+		}
+		else
+		{
+			UE_LOG(LogD3D12RHI, Warning, TEXT("PSO Cache read error!"));
+			break;
+		}
+	}
+
+	uint32 NumComputePSOs = DiskCaches[PSO_CACHE_COMPUTE].GetNumPSOs();
+
+	for (uint32 i = 0; i < NumComputePSOs; i++)
+	{
+		FD3D12ComputePipelineStateDesc* ComputePSODesc = nullptr;
+		DiskCaches[PSO_CACHE_COMPUTE].SetPointerAndAdvanceFilePosition((void**)&ComputePSODesc, sizeof(FD3D12ComputePipelineStateDesc));
+		D3D12_COMPUTE_PIPELINE_STATE_DESC* PSODesc = &ComputePSODesc->Desc;
+
+		PSODesc->pRootSignature = pComputeRootSignature;
+		if (PSODesc->CS.BytecodeLength)
+		{
+			DiskCaches[PSO_CACHE_COMPUTE].SetPointerAndAdvanceFilePosition((void**)&PSODesc->CS.pShaderBytecode, PSODesc->CS.BytecodeLength, true);
+		}
+
+		SIZE_T* cachedBlobSize = nullptr;
+		DiskCaches[PSO_CACHE_COMPUTE].SetPointerAndAdvanceFilePosition((void**)&cachedBlobSize, sizeof(SIZE_T));
+
+		check(cachedBlobSize);
+		if (bUseAPILibaries && cachedBlobSize) { check(*cachedBlobSize); }
+
+		if (bUseAPILibaries && cachedBlobSize &&*cachedBlobSize)
+		{
+			ComputePSODesc->Desc.CachedPSO.CachedBlobSizeInBytes = *cachedBlobSize;
+			DiskCaches[PSO_CACHE_COMPUTE].SetPointerAndAdvanceFilePosition((void**)&ComputePSODesc->Desc.CachedPSO.pCachedBlob, *cachedBlobSize);
+		}
+		else
+		{
+			ComputePSODesc->Desc.CachedPSO.CachedBlobSizeInBytes = 0;
+			ComputePSODesc->Desc.CachedPSO.pCachedBlob = nullptr;
+		}
+
+		if (!DiskCaches[PSO_CACHE_COMPUTE].IsInErrorState())
+		{
+			ComputePSODesc->CombinedHash = FD3D12PipelineStateCache::HashPSODesc(*ComputePSODesc);
+			ID3D12PipelineState* returnVal = this->Add(*ComputePSODesc);
+			if (returnVal == nullptr)
+			{
+				// The cache has failed, build PSOs at runtime
+				return;
+			}
+		}
+		else
+		{
+			UE_LOG(LogD3D12RHI, Warning, TEXT("PSO Cache read error!"));
+			break;
+		}
+	}
+}
+
+ID3D12PipelineState* FD3D12PipelineStateCache::FindGraphics(FD3D12HighLevelGraphicsPipelineStateDesc &graphicsPSODesc, ID3D12RootSignature* rootSignature, bool insertIntoDiskCache)
+{
+	FScopeLock Lock (&CS);
+
+#if UE_BUILD_DEBUG
+	++GraphicsCacheRequestCount;
+#endif
+
+	graphicsPSODesc.CombinedHash = FD3D12PipelineStateCache::HashPSODesc(graphicsPSODesc);
+
+	uint64 BSSUniqueID = graphicsPSODesc.BoundShaderState ? graphicsPSODesc.BoundShaderState->UniqueID : 0;
+	TPair<ID3D12PipelineState*, uint64>* HighLevelCacheEntry = HighLevelGraphicsPipelineStateCache.Find(graphicsPSODesc);
+	if (HighLevelCacheEntry && HighLevelCacheEntry->Value == BSSUniqueID)
+	{
+#if UE_BUILD_DEBUG
+		++HighLevelCacheFulfillCount; // No low-level cache hit
+#endif
+		return HighLevelCacheEntry->Key;
+	}
+
+	FD3D12LowLevelGraphicsPipelineStateDesc lowLevelDesc;
+	graphicsPSODesc.GetLowLevelDesc(rootSignature, lowLevelDesc);
+	ID3D12PipelineState* PSO = FindGraphicsLowLevel(lowLevelDesc, insertIntoDiskCache);
+
+	if (HighLevelCacheEntry)
+	{
+#if UE_BUILD_DEBUG
+		++HighLevelCacheStaleCount; // High-level cache hit, but was stale due to BSS memory re-use
+#endif
+		HighLevelCacheEntry->Key = PSO;
+	}
+	else
+	{
+#if UE_BUILD_DEBUG
+		++HighLevelCacheMissCount; // No high-level cache hit
+#endif
+		HighLevelGraphicsPipelineStateCache.Add(graphicsPSODesc,
+			TPairInitializer<ID3D12PipelineState*, uint64>(PSO, BSSUniqueID));
+	}
+	return PSO;
+}
+
+ID3D12PipelineState* FD3D12PipelineStateCache::FindGraphicsLowLevel(FD3D12LowLevelGraphicsPipelineStateDesc &graphicsPSODesc, bool insertIntoDiskCache)
+{
+	// Lock already taken by high level find
+	graphicsPSODesc.CombinedHash = FD3D12PipelineStateCache::HashPSODesc(graphicsPSODesc);
+
+	if (TRefCountPtr<ID3D12PipelineState>* PSO = LowLevelGraphicsPipelineStateCache.Find(graphicsPSODesc))
+		return PSO->GetReference();
+
+	return Add(graphicsPSODesc, insertIntoDiskCache);
+}
+
+ID3D12PipelineState* FD3D12PipelineStateCache::FindCompute(FD3D12ComputePipelineStateDesc &computePSODesc, bool insertIntoDiskCache)
+{
+	FScopeLock Lock (&CS);
+	computePSODesc.CombinedHash = FD3D12PipelineStateCache::HashPSODesc(computePSODesc);
+
+	if (TRefCountPtr<ID3D12PipelineState>* PSO = ComputePipelineStateCache.Find(computePSODesc))
+		return PSO->GetReference();
+
+	return Add(computePSODesc, insertIntoDiskCache);
+}
+
+ID3D12PipelineState* FD3D12PipelineStateCache::Add(FD3D12LowLevelGraphicsPipelineStateDesc &graphicsPSODesc, bool insertIntoDiskCache)
+{
+	FScopeLock Lock(&CS);
+
+	TRefCountPtr<ID3D12PipelineState> PSO;
+	D3D12_GRAPHICS_PIPELINE_STATE_DESC &psoDesc = graphicsPSODesc.Desc;
+
+	HRESULT Hr = GetParentDevice()->GetDevice()->CreateGraphicsPipelineState(&psoDesc, IID_PPV_ARGS(PSO.GetInitReference()));
+
+	// It is possible that the on disk descriptor gets corrupted. In that case failure is ok we'll create
+	// the PSO later at runtime. It is also possible that the user changed their graphics device and if
+	// we are using the API libraries compilation will fail because the cache is stale.
+	if (!insertIntoDiskCache && Hr != S_OK)
+	{
+		UE_LOG(LogD3D12RHI, Warning, TEXT("PSO re-creation failed. Most likely on disk descriptor corruption or change of video device/driver."));
+		for (uint32 i = 0; i < NUM_PSO_CACHE_TYPES; i++)
+		{
+			DiskCaches[i].ClearDiskCache();
+		}
+
+		return nullptr;
+	}
+	VERIFYD3D11RESULT(Hr);
+
+#if UE_BUILD_DEBUG
+	check(LowLevelGraphicsPipelineStateCache.Find(graphicsPSODesc) == nullptr);
+#endif
+
+	LowLevelGraphicsPipelineStateCache.Add(graphicsPSODesc, PSO);
+
+	//TODO: Optimize by only storing unique pointers
+	if (insertIntoDiskCache && !DiskCaches[PSO_CACHE_GRAPHICS].IsInErrorState())
+	{
+		DiskCaches[PSO_CACHE_GRAPHICS].AppendData(&graphicsPSODesc, sizeof(graphicsPSODesc));
+
+		if (psoDesc.InputLayout.NumElements)
+		{
+			//Save the layout structs
+			DiskCaches[PSO_CACHE_GRAPHICS].AppendData((void*)psoDesc.InputLayout.pInputElementDescs, psoDesc.InputLayout.NumElements * sizeof(D3D12_INPUT_ELEMENT_DESC));
+			for (uint32 i = 0; i < psoDesc.InputLayout.NumElements; i++)
+			{
+				//Save the Sematic name string
+				uint32 stringLength = (uint32)strnlen_s(psoDesc.InputLayout.pInputElementDescs[i].SemanticName, IL_MAX_SEMANTIC_NAME);
+				stringLength++; // include the NULL char
+				DiskCaches[PSO_CACHE_GRAPHICS].AppendData((void*)&stringLength, sizeof(stringLength));
+				DiskCaches[PSO_CACHE_GRAPHICS].AppendData((void*)psoDesc.InputLayout.pInputElementDescs[i].SemanticName, stringLength);
+			}
+		}
+		if (psoDesc.StreamOutput.NumEntries)
+		{
+			DiskCaches[PSO_CACHE_GRAPHICS].AppendData((void*)&psoDesc.StreamOutput.pSODeclaration, psoDesc.StreamOutput.NumEntries * sizeof(D3D12_SO_DECLARATION_ENTRY));
+			for (uint32 i = 0; i < psoDesc.StreamOutput.NumEntries; i++)
+			{
+				//Save the Sematic name string
+				uint32 stringLength = (uint32)strnlen_s(psoDesc.StreamOutput.pSODeclaration[i].SemanticName, IL_MAX_SEMANTIC_NAME);
+				stringLength++; // include the NULL char
+				DiskCaches[PSO_CACHE_GRAPHICS].AppendData((void*)&stringLength, sizeof(stringLength));
+				DiskCaches[PSO_CACHE_GRAPHICS].AppendData((void*)psoDesc.StreamOutput.pSODeclaration[i].SemanticName, stringLength);
+			}
+		}
+		if (psoDesc.StreamOutput.NumStrides)
+		{
+			DiskCaches[PSO_CACHE_GRAPHICS].AppendData((void*)&psoDesc.StreamOutput.pBufferStrides, psoDesc.StreamOutput.NumStrides * sizeof(uint32));
+		}
+		if (psoDesc.VS.pShaderBytecode)
+		{
+			DiskCaches[PSO_CACHE_GRAPHICS].AppendData((void*)psoDesc.VS.pShaderBytecode, psoDesc.VS.BytecodeLength);
+		}
+		if (psoDesc.PS.pShaderBytecode)
+		{
+			DiskCaches[PSO_CACHE_GRAPHICS].AppendData((void*)psoDesc.PS.pShaderBytecode, psoDesc.PS.BytecodeLength);
+		}
+		if (psoDesc.DS.pShaderBytecode)
+		{
+			DiskCaches[PSO_CACHE_GRAPHICS].AppendData((void*)psoDesc.DS.pShaderBytecode, psoDesc.DS.BytecodeLength);
+		}
+		if (psoDesc.HS.pShaderBytecode)
+		{
+			DiskCaches[PSO_CACHE_GRAPHICS].AppendData((void*)psoDesc.HS.pShaderBytecode, psoDesc.HS.BytecodeLength);
+		}
+		if (psoDesc.GS.pShaderBytecode)
+		{
+			DiskCaches[PSO_CACHE_GRAPHICS].AppendData((void*)psoDesc.GS.pShaderBytecode, psoDesc.GS.BytecodeLength);
+		}
+
+		if (bUseAPILibaries)
+		{
+			TRefCountPtr<ID3DBlob> cachedBlob;
+			HRESULT result = PSO->GetCachedBlob(cachedBlob.GetInitReference());
+			VERIFYD3D11RESULT(result);
+			if (SUCCEEDED(result))
+			{
+				SIZE_T bufferSize = cachedBlob->GetBufferSize();
+				DiskCaches[PSO_CACHE_GRAPHICS].AppendData(&bufferSize, sizeof(bufferSize));
+				DiskCaches[PSO_CACHE_GRAPHICS].AppendData(cachedBlob->GetBufferPointer(), bufferSize);
+			}
+			else
+			{
+				check(false);
+				SIZE_T bufferSize = 0;
+				DiskCaches[PSO_CACHE_GRAPHICS].AppendData(&bufferSize, sizeof(bufferSize));
+			}
+		}
+		else
+		{
+			SIZE_T bufferSize = 0;
+			DiskCaches[PSO_CACHE_GRAPHICS].AppendData(&bufferSize, sizeof(bufferSize));
+		}
+
+		DiskCaches[PSO_CACHE_GRAPHICS].Flush(LowLevelGraphicsPipelineStateCache.Num());
+	}
+
+	return PSO.GetReference();
+}
+
+ID3D12PipelineState* FD3D12PipelineStateCache::Add(FD3D12ComputePipelineStateDesc &computePSODesc, bool insertIntoDiskCache)
+{
+	FScopeLock Lock (&CS);
+
+	TRefCountPtr<ID3D12PipelineState> PSO;
+	D3D12_COMPUTE_PIPELINE_STATE_DESC &psoDesc = computePSODesc.Desc;
+
+	HRESULT Hr = GetParentDevice()->GetDevice()->CreateComputePipelineState(&psoDesc, IID_PPV_ARGS(PSO.GetInitReference()));
+
+	// It is possible that the on disk descriptor gets corrupted. In that case failure is ok we'll create
+	// the PSO later at runtime.
+	if (!insertIntoDiskCache && Hr != S_OK)
+	{
+		UE_LOG(LogD3D12RHI, Warning, TEXT("PSO re-creation failed. Most likely on disk descriptor corruption."));
+		for (uint32 i = 0; i < NUM_PSO_CACHE_TYPES; i++)
+		{
+			DiskCaches[i].ClearDiskCache();
+		}
+		return nullptr;
+	}
+	VERIFYD3D11RESULT(Hr);
+
+#if UE_BUILD_DEBUG
+	check(ComputePipelineStateCache.Find(computePSODesc) == nullptr);
+#endif
+
+	ComputePipelineStateCache.Add(computePSODesc, PSO);
+
+	if (insertIntoDiskCache && !DiskCaches[PSO_CACHE_COMPUTE].IsInErrorState())
+	{
+		DiskCaches[PSO_CACHE_COMPUTE].AppendData(&computePSODesc, sizeof(computePSODesc));
+		if (psoDesc.CS.pShaderBytecode)
+		{
+			DiskCaches[PSO_CACHE_COMPUTE].AppendData((void*)psoDesc.CS.pShaderBytecode, psoDesc.CS.BytecodeLength);
+		}
+
+		if (bUseAPILibaries)
+		{
+			TRefCountPtr<ID3DBlob> cachedBlob;
+			HRESULT result = PSO->GetCachedBlob(cachedBlob.GetInitReference());
+			VERIFYD3D11RESULT(result);
+			if (SUCCEEDED(result))
+			{
+				SIZE_T bufferSize = cachedBlob->GetBufferSize();
+				DiskCaches[PSO_CACHE_COMPUTE].AppendData(&bufferSize, sizeof(bufferSize));
+				DiskCaches[PSO_CACHE_COMPUTE].AppendData(cachedBlob->GetBufferPointer(), bufferSize);
+			}
+			else
+			{
+				check(false);
+				SIZE_T bufferSize = 0;
+				DiskCaches[PSO_CACHE_COMPUTE].AppendData(&bufferSize, sizeof(bufferSize));
+			}
+		}
+		else
+		{
+			SIZE_T bufferSize = 0;
+			DiskCaches[PSO_CACHE_COMPUTE].AppendData(&bufferSize, sizeof(bufferSize));
+		}
+
+		DiskCaches[PSO_CACHE_COMPUTE].Flush(ComputePipelineStateCache.Num());
+	}
+
+	return PSO.GetReference();
+}
+
+void FD3D12PipelineStateCache::Close()
+{
+	FScopeLock Lock (&CS);
+
+	DiskCaches[PSO_CACHE_GRAPHICS].Reset();
+	DiskCaches[PSO_CACHE_COMPUTE].Reset();
+
+	DiskCaches[PSO_CACHE_GRAPHICS].Close(LowLevelGraphicsPipelineStateCache.Num());
+	DiskCaches[PSO_CACHE_COMPUTE].Close(ComputePipelineStateCache.Num());
+
+	HighLevelGraphicsPipelineStateCache.Empty();
+	LowLevelGraphicsPipelineStateCache.Empty();
+	ComputePipelineStateCache.Empty();
+}
+
+FORCEINLINE uint32 SSE4_CRC32(void* data, SIZE_T numBytes)
+{
+	check(bCPUSupportsSSE4);
+	uint32 hash = 0;
+#if _WIN64
+	static const SIZE_T alignment = 8;//64 Bit
+#elif _WIN32
+	static const SIZE_T alignment = 4;//32 Bit
+#else
+	check(0);
+	return 0;
+#endif
+
+	const SIZE_T roundingIterations = (numBytes & (alignment - 1));
+	uint8* unalignedData = (uint8*)data;
+	for (SIZE_T i = 0; i < roundingIterations; i++)
+	{
+		hash = _mm_crc32_u8(hash, unalignedData[i]);
+	}
+	unalignedData += roundingIterations;
+	numBytes -= roundingIterations;
+
+	SIZE_T* alignedData = (SIZE_T*)unalignedData;
+	check((numBytes % alignment) == 0);
+	const SIZE_T numIterations = (numBytes / alignment);
+	for (SIZE_T i = 0; i < numIterations; i++)
+	{
+#if _WIN64
+		hash = _mm_crc32_u64(hash, alignedData[i]);
+#else
+		hash = _mm_crc32_u32(hash, alignedData[i]);
+#endif
+	}
+
+	return hash;
+}
+
+SIZE_T FD3D12PipelineStateCache::HashData(void* data, SIZE_T numBytes)
+{
+	if (bCPUSupportsSSE4)
+	{
+		return SIZE_T(SSE4_CRC32(data, numBytes));
+	}
+	else
+	{
+		return SIZE_T(FCrc::MemCrc32(data, numBytes));
+	}
+}
+
+//TODO: optimize by pre-hashing these things at set time
+SIZE_T FD3D12PipelineStateCache::HashPSODesc(const FD3D12LowLevelGraphicsPipelineStateDesc &graphicsPSODesc)
+{
+	__declspec(align(32)) FD3D12LowLevelGraphicsPipelineStateDesc hashTemp;
+	FMemory::Memcpy(&hashTemp, &graphicsPSODesc, sizeof(hashTemp)); // Memcpy to avoid introducing garbage due to alignment
+	hashTemp.Desc.VS.pShaderBytecode = nullptr; // null out pointers so stale ones don't ruin the hash
+	hashTemp.Desc.PS.pShaderBytecode = nullptr;
+	hashTemp.Desc.HS.pShaderBytecode = nullptr;
+	hashTemp.Desc.DS.pShaderBytecode = nullptr;
+	hashTemp.Desc.GS.pShaderBytecode = nullptr;
+	hashTemp.Desc.InputLayout.pInputElementDescs = nullptr;
+	hashTemp.Desc.StreamOutput.pBufferStrides = nullptr;
+	hashTemp.Desc.StreamOutput.pSODeclaration = nullptr;
+	hashTemp.Desc.CachedPSO.pCachedBlob = nullptr;
+	hashTemp.Desc.CachedPSO.CachedBlobSizeInBytes = 0;
+	hashTemp.CombinedHash = 0;
+	hashTemp.Desc.pRootSignature = nullptr;
+
+	return HashData(&hashTemp, sizeof(hashTemp));
+}
+
+SIZE_T FD3D12PipelineStateCache::HashPSODesc(const FD3D12HighLevelGraphicsPipelineStateDesc &graphicsPSODesc)
+{
+	__declspec(align(32)) FD3D12HighLevelGraphicsPipelineStateDesc hashTemp;
+	FMemory::Memcpy(&hashTemp, &graphicsPSODesc, sizeof(hashTemp)); // Memcpy to avoid introducing garbage due to alignment
+	hashTemp.CombinedHash = 0;
+	return HashData(&hashTemp, sizeof(hashTemp));
+}
+
+SIZE_T FD3D12PipelineStateCache::HashPSODesc(const FD3D12ComputePipelineStateDesc &psoDesc)
+{
+	__declspec(align(32)) FD3D12ComputePipelineStateDesc hashTemp;
+	FMemory::Memcpy(&hashTemp, &psoDesc, sizeof(hashTemp)); // Memcpy to avoid introducing garbage due to alignment
+	hashTemp.Desc.CS.pShaderBytecode = nullptr;  // null out pointers so stale ones don't ruin the hash
+	hashTemp.Desc.CachedPSO.pCachedBlob = nullptr;
+	hashTemp.Desc.CachedPSO.CachedBlobSizeInBytes = 0;
+	hashTemp.CombinedHash = 0;
+	hashTemp.Desc.pRootSignature = nullptr;
+	return HashData(&hashTemp, sizeof(hashTemp));
+}
+
+
+#define SSE4_2     0x100000 
+#define SSE4_CPUID_ARRAY_INDEX 2
+void FD3D12PipelineStateCache::Init(FString &GraphicsCacheFilename, FString &ComputeCacheFilename)
+{
+	FScopeLock Lock (&CS);
+
+	DiskCaches[PSO_CACHE_GRAPHICS].Init(GraphicsCacheFilename);
+	DiskCaches[PSO_CACHE_COMPUTE].Init(ComputeCacheFilename);
+
+	// Check for SSE4 support see: https://msdn.microsoft.com/en-us/library/vstudio/hskdteyh(v=vs.100).aspx
+	{
+		int32 cpui[4];
+		__cpuidex(cpui, 1, 0);
+		bCPUSupportsSSE4 = !!(cpui[SSE4_CPUID_ARRAY_INDEX] & SSE4_2);
+	}
+
+}
+
+FD3D12PipelineStateCache::FD3D12PipelineStateCache(FD3D12Device* InParent) :
+	FD3D12DeviceChild(InParent)
+{
+
+}
+
+FD3D12PipelineStateCache::~FD3D12PipelineStateCache()
+{
+
+}
+
+bool FD3D12StateCacheBase::VerifyResourceStates(const bool IsCompute)
+{
+	// Can only verify resource states if the debug layer is used
+	const bool bWithD3DDebug = D3D12RHI_ShouldCreateWithD3DDebug();
+	if (!bWithD3DDebug)
+	{
+		UE_LOG(LogD3D12RHI, Log, TEXT("*** VerifyResourceStates requires the debug layer ***"), this);
+		return false;
+	}
+
+	// Get the debug command queue
+	TRefCountPtr<ID3D12DebugCommandList> pDebugCommandList;
+	VERIFYD3D11RESULT(CmdContext->CommandListHandle->QueryInterface(pDebugCommandList.GetInitReference()));
+
+	// Note: There is nothing special to check when IsCompute = true
+	if (!IsCompute)
+	{
+		//
+		// Verify graphics pipeline state
+		//
+
+		// DSV
+		{
+			FD3D12DepthStencilView *pCurrentView = PipelineState.Graphics.CurrentDepthStencilTarget;
+
+			if (pCurrentView)
+			{
+				// Check if the depth/stencil resource has an SRV bound
+				const bool bSRVBound = PipelineState.Common.ShaderResourceViewsIntersectWithDepthCount > 0;
+				uint32 SanityCheckCount = 0;
+				const uint32 StartStage = 0;
+				const uint32 EndStage = SF_Compute;
+				for (uint32 Stage = StartStage; Stage < EndStage; Stage++)
+				{
+					for (uint32 i = 0; i < MAX_SRVS; i++)
+					{
+						if (PipelineState.Common.CurrentShaderResourceViewsIntersectWithDepthRT[Stage][i])
+						{
+							SanityCheckCount++;
+						}
+					}
+				}
+				check(SanityCheckCount == PipelineState.Common.ShaderResourceViewsIntersectWithDepthCount);
+				
+				const D3D12_DEPTH_STENCIL_VIEW_DESC& desc = pCurrentView->GetDesc();
+				const bool bDepthIsReadOnly = !!(desc.Flags & D3D12_DSV_FLAG_READ_ONLY_DEPTH);
+				const bool bStencilIsReadOnly = !!(desc.Flags & D3D12_DSV_FLAG_READ_ONLY_STENCIL);
+
+				// Decompose the view into the subresources (depth and stencil are on different planes)
+				FD3D12Resource *pResource = pCurrentView->GetResource();
+				const CViewSubresourceSubset subresourceSubset = pCurrentView->GetViewSubresourceSubset();
+				for (CViewSubresourceSubset::CViewSubresourceIterator it = subresourceSubset.begin(); it != subresourceSubset.end(); ++it)
+				{
+					for (uint32 SubresourceIndex = it.StartSubresource(); SubresourceIndex < it.EndSubresource(); SubresourceIndex++)
+					{
+						uint16 MipSlice;
+						uint16 ArraySlice;
+						uint8 PlaneSlice;
+						D3D12DecomposeSubresource(SubresourceIndex,
+							pResource->GetMipLevels(),
+							pResource->GetArraySize(),
+							MipSlice, ArraySlice, PlaneSlice);
+
+						D3D12_RESOURCE_STATES expectedState;
+						if (PlaneSlice == 0)
+						{
+							// Depth plane
+							expectedState = bDepthIsReadOnly ? D3D12_RESOURCE_STATE_DEPTH_READ : D3D12_RESOURCE_STATE_DEPTH_WRITE;
+							if (bSRVBound)
+							{
+								// Depth SRVs just contain the depth plane
+								check(bDepthIsReadOnly);
+								expectedState |=
+									D3D12_RESOURCE_STATE_NON_PIXEL_SHADER_RESOURCE |
+									D3D12_RESOURCE_STATE_PIXEL_SHADER_RESOURCE;
+							}
+						}
+						else
+						{
+							// Stencil plane
+							expectedState = bStencilIsReadOnly ? D3D12_RESOURCE_STATE_DEPTH_READ : D3D12_RESOURCE_STATE_DEPTH_WRITE;
+						}
+						
+						bool bGoodState = !!pDebugCommandList->AssertResourceState(pResource->GetResource(), SubresourceIndex, expectedState);
+						if (!bGoodState)
+						{
+							return false;
+						}
+					}
+				}
+
+			}
+		}
+
+		// RTV
+		{
+			const uint32 numRTVs = _countof(PipelineState.Graphics.RenderTargetArray);
+			for (uint32 i = 0; i < numRTVs; i++)
+			{
+				FD3D12RenderTargetView *pCurrentView = PipelineState.Graphics.RenderTargetArray[i];
+				if (!VerifyViewState(pDebugCommandList, pCurrentView, D3D12_RESOURCE_STATE_RENDER_TARGET))
+				{
+					return false;
+				}
+			}
+		}
+	}
+
+	//
+	// Verify common pipeline state
+	//
+
+	// UAVs
+	{
+		// MSFT: Disabled UAV checks for now. This will fail in Infiltrator during the FXSystem's SortGPUBuffers call
+		// because a single sub-resources is used for UAV and SRV. Need to identify what the correct state should be.
+		/*const uint32 numUAVs = _countof(PipelineState.Common.UnorderedAccessViewArray);
+		for (uint32 i = 0; i < numUAVs; i++)
+		{
+			FD3D12UnorderedAccessView *pCurrentView = PipelineState.Common.UnorderedAccessViewArray[i];
+			if (!VerifyViewState(pDebugCommandList, pCurrentView, D3D12_RESOURCE_STATE_UNORDERED_ACCESS))
+			{
+				return false;
+			}
+		}*/
+	}
+
+	// SRVs
+	{
+		const uint32 numStages = _countof(PipelineState.Common.CurrentShaderResourceViews);
+		for (uint32 stage = 0; stage < numStages; stage++)
+		{
+			const uint32 numSRVs = _countof(PipelineState.Common.CurrentShaderResourceViews[stage]);
+			for (uint32 i = 0; i < numSRVs; i++)
+			{
+				FD3D12ShaderResourceView *pCurrentView = PipelineState.Common.CurrentShaderResourceViews[stage][i];
+				if (!VerifyViewState(pDebugCommandList, pCurrentView, D3D12_RESOURCE_STATE_PIXEL_SHADER_RESOURCE | D3D12_RESOURCE_STATE_NON_PIXEL_SHADER_RESOURCE))
+				{
+					return false;
+				}
+			}
+		}
+	}
+
+	return true;
+}
+
+template <class TView>
+bool FD3D12StateCacheBase::VerifyViewState(ID3D12DebugCommandList* pDebugCommandList, FD3D12View<TView> *pView, D3D12_RESOURCE_STATES State)
+{
+	// Check the view
+	if (!pView)
+	{
+		// No need to check null views
+		return true;
+	}
+
+	// Get the view's resource
+	FD3D12Resource *pResource = pView->GetResource();
+	if (!pResource)
+	{
+		// Some dynamic SRVs haven't been mapped and updated yet so they actually don't have any backing resources.
+		return true;
+	}
+
+	ID3D12Resource *pD3D12Resource = pResource->GetResource();
+	check(pD3D12Resource);
+
+	// For each subresource in the view...
+	const CViewSubresourceSubset subresourceSubset = pView->GetViewSubresourceSubset();
+	for (CViewSubresourceSubset::CViewSubresourceIterator it = subresourceSubset.begin(); it != subresourceSubset.end(); ++it)
+	{
+		for (uint32 SubresourceIndex = it.StartSubresource(); SubresourceIndex < it.EndSubresource(); SubresourceIndex++)
+		{
+			bool bGoodState = !!pDebugCommandList->AssertResourceState(pD3D12Resource, SubresourceIndex, State);
+			if (!bGoodState)
+			{
+				return false;
+			}
+		}
+	}
+
+	return true;
+}
+
+void FD3D12StateCacheBase::SetUAVs(EShaderFrequency ShaderStage, uint32 UAVStartSlot, uint32 NumSimultaneousUAVs, FD3D12UnorderedAccessView** UAVArray, uint32 *UAVInitialCountArray)
+{
+	if (NumSimultaneousUAVs == 0)
+		return;
+	
+	PipelineState.Common.CurrentUAVStartSlot = FMath::Min(UAVStartSlot, PipelineState.Common.CurrentUAVStartSlot);
+	PipelineState.Common.CurrentUAVStage = ShaderStage;
+	
+	TRefCountPtr<ID3D12Resource> CounterUploadHeap      = GetParentDevice()->GetCounterUploadHeap();
+	uint32&                      CounterUploadHeapIndex = GetParentDevice()->GetCounterUploadHeapIndex();
+	void*                        CounterUploadHeapData  = GetParentDevice()->GetCounterUploadHeapData();
+	
+	for (uint32 i = 0; i < NumSimultaneousUAVs; ++i)
+	{
+		PipelineState.Common.UnorderedAccessViewArray[UAVStartSlot + i] = UAVArray[i];
+
+		if (UAVArray[i] && UAVArray[i]->CounterResource && (!UAVArray[i]->CounterResourceInitialized || UAVInitialCountArray[i] != -1))
+		{
+			// Ensure we don't run off the end of the upload heap.
+			CounterUploadHeapIndex = (CounterUploadHeapIndex % (COUNTER_HEAP_SIZE / 4));
+
+			// Initialize the counter to 0 if it's not been previously initialized and the UAVInitialCount is -1, if not use the value that was passed.
+			((uint32*)CounterUploadHeapData)[CounterUploadHeapIndex] = (!UAVArray[i]->CounterResourceInitialized && UAVInitialCountArray[i] == -1) ? 0 : UAVInitialCountArray[i];
+
+#if SUPPORTS_MEMORY_RESIDENCY
+			UAVArray[i]->CounterResource->UpdateResidency();
+#endif
+
+			CmdContext->CommandListHandle->CopyBufferRegion(
+				UAVArray[i]->CounterResource->GetResource(),
+				0,
+				CounterUploadHeap,
+				CounterUploadHeapIndex * 4,
+				4);
+
+			CounterUploadHeapIndex++;
+			UAVArray[i]->CounterResourceInitialized = true;
+		}
+	}
+
+	PipelineState.Common.CurrentNumberOfSimultaneousUAVs = 0;
+	for (int i = D3D12_PS_CS_UAV_REGISTER_COUNT - 1; i >= 0; --i)
+	{
+		if (PipelineState.Common.UnorderedAccessViewArray[i])
+		{
+			PipelineState.Common.CurrentNumberOfSimultaneousUAVs = i + 1;
+			break;
+		}
+	}
+
+	bNeedSetUAVs = true;
+}
+
+void FD3D12StateCacheBase::SetPrimitiveTopology(D3D_PRIMITIVE_TOPOLOGY PrimitiveTopology)
+{
+	if ((PipelineState.Graphics.CurrentPrimitiveTopology != PrimitiveTopology) || GD3D12SkipStateCaching)
+	{
+		PipelineState.Graphics.CurrentPrimitiveTopology = PrimitiveTopology;
+		switch (PipelineState.Graphics.CurrentPrimitiveTopology)
+		{
+		case D3D_PRIMITIVE_TOPOLOGY_POINTLIST:
+			PipelineState.Graphics.HighLevelDesc.PrimitiveTopologyType = D3D12_PRIMITIVE_TOPOLOGY_TYPE_POINT;
+			break;
+		case D3D_PRIMITIVE_TOPOLOGY_LINELIST:
+		case D3D_PRIMITIVE_TOPOLOGY_LINESTRIP:
+		case D3D_PRIMITIVE_TOPOLOGY_LINELIST_ADJ:
+		case D3D_PRIMITIVE_TOPOLOGY_LINESTRIP_ADJ:
+			PipelineState.Graphics.HighLevelDesc.PrimitiveTopologyType = D3D12_PRIMITIVE_TOPOLOGY_TYPE_LINE;
+			break;
+		case D3D_PRIMITIVE_TOPOLOGY_TRIANGLELIST:
+		case D3D_PRIMITIVE_TOPOLOGY_TRIANGLESTRIP:
+		case D3D_PRIMITIVE_TOPOLOGY_TRIANGLELIST_ADJ:
+		case D3D_PRIMITIVE_TOPOLOGY_TRIANGLESTRIP_ADJ:
+			PipelineState.Graphics.HighLevelDesc.PrimitiveTopologyType = D3D12_PRIMITIVE_TOPOLOGY_TYPE_TRIANGLE;
+			break;
+		default:
+			PipelineState.Graphics.HighLevelDesc.PrimitiveTopologyType = D3D12_PRIMITIVE_TOPOLOGY_TYPE_PATCH;
+			break;
+		}
+		bNeedSetPrimitiveTopology = true;
+		PipelineState.Graphics.bNeedRebuildPSO = true;
+	}
+}
+
+void FD3D12StateCacheBase::SetBlendState(D3D12_BLEND_DESC* State, const float BlendFactor[4], uint32 SampleMask)
+{
+	if (PipelineState.Graphics.HighLevelDesc.BlendState != State || PipelineState.Graphics.HighLevelDesc.SampleMask != SampleMask || GD3D12SkipStateCaching)
+	{
+		PipelineState.Graphics.HighLevelDesc.BlendState = State;
+		PipelineState.Graphics.HighLevelDesc.SampleMask = SampleMask;
+		PipelineState.Graphics.bNeedRebuildPSO = true;
+	}
+
+	if (FMemory::Memcmp(PipelineState.Graphics.CurrentBlendFactor, BlendFactor, sizeof(PipelineState.Graphics.CurrentBlendFactor)))
+	{
+		FMemory::Memcpy(PipelineState.Graphics.CurrentBlendFactor, BlendFactor, sizeof(PipelineState.Graphics.CurrentBlendFactor));
+		bNeedSetBlendFactor = true;
+	}
+}
+
+D3D12_STATE_CACHE_INLINE void FD3D12StateCacheBase::SetDepthStencilState(D3D12_DEPTH_STENCIL_DESC* State, uint32 RefStencil)
+{
+	if (PipelineState.Graphics.HighLevelDesc.DepthStencilState != State || GD3D12SkipStateCaching)
+	{
+		PipelineState.Graphics.HighLevelDesc.DepthStencilState = State;
+		PipelineState.Graphics.bNeedRebuildPSO = true;
+	}
+
+	if (PipelineState.Graphics.CurrentReferenceStencil != RefStencil)
+	{
+		PipelineState.Graphics.CurrentReferenceStencil = RefStencil;
+		bNeedSetStencilRef = true;
+	}
+}
+
+void FD3D12StateCacheBase::InternalSetIndexBuffer(FD3D12ResourceLocation *IndexBufferLocation, DXGI_FORMAT Format, uint32 Offset)
+{
+	if (bAlwaysSetIndexBuffers || (PipelineState.Graphics.CurrentIndexBufferLocation != IndexBufferLocation || PipelineState.Graphics.CurrentIndexFormat != Format || PipelineState.Graphics.CurrentIndexOffset != Offset) || GD3D12SkipStateCaching)
+	{
+		PipelineState.Graphics.CurrentIndexBufferLocation = IndexBufferLocation;
+		PipelineState.Graphics.CurrentIndexFormat = Format;
+		PipelineState.Graphics.CurrentIndexOffset = Offset;
+		bNeedSetIB = true;
+	}
+}
+
+void FD3D12StateCacheBase::InternalSetStreamSource(FD3D12ResourceLocation* VertexBufferLocation, uint32 StreamIndex, uint32 Stride, uint32 Offset, TSetStreamSourceAlternate AlternatePathFunction)
+{
+	check(StreamIndex < ARRAYSIZE(PipelineState.Graphics.CurrentVertexBuffers));
+	FD3D12VertexBufferState& Slot = PipelineState.Graphics.CurrentVertexBuffers[StreamIndex];
+	if ((Slot.VertexBufferLocation != VertexBufferLocation || Slot.Offset != Offset || Slot.Stride != Stride) || GD3D12SkipStateCaching)
+	{
+		Slot.VertexBufferLocation = VertexBufferLocation;
+		Slot.Offset = Offset;
+		Slot.Stride = Stride;
+		if (AlternatePathFunction != nullptr)
+		{
+			(*AlternatePathFunction)(this, VertexBufferLocation->GetResource(), StreamIndex, Stride, Offset + VertexBufferLocation->GetOffset());
+		}
+		else
+		{
+			bNeedSetVB = true;
+		}
+	}
+}
+
+void FD3D12StateCacheBase::SetRenderTargets(uint32 NumSimultaneousRenderTargets, FD3D12RenderTargetView** RTArray, FD3D12DepthStencilView* DSTarget)
+{
+	// Note: We assume that the have been checks to make sure this function is only called when there really are changes being made.
+	// We always update the PSO and set descriptors after calling this function.
+	PipelineState.Graphics.bNeedRebuildPSO = true;
+	bNeedSetRTs = true;
+
+	// Update the depth stencil
+	PipelineState.Graphics.CurrentDepthStencilTarget = DSTarget;
+
+	// Update the render targets
+	FMemory::Memzero(PipelineState.Graphics.RenderTargetArray, sizeof(PipelineState.Graphics.RenderTargetArray));
+	FMemory::Memcpy(PipelineState.Graphics.RenderTargetArray, RTArray, sizeof(FD3D12RenderTargetView*)* NumSimultaneousRenderTargets);
+
+	// In D3D11, the NumSimultaneousRenderTargets count was used even when setting RTV slots to null (to unbind them)
+	// In D3D12, we don't do this. So we need change the count to match the non null views used.
+	uint32 ActiveNumSimultaneousRenderTargets = 0;
+	for (uint32 i = 0; i < NumSimultaneousRenderTargets; i++)
+	{
+		if (RTArray[i] != nullptr)
+		{
+			ActiveNumSimultaneousRenderTargets = i + 1;
+		}
+	}
+	PipelineState.Graphics.HighLevelDesc.NumRenderTargets = ActiveNumSimultaneousRenderTargets;
+}
+
+void FD3D12StateCacheBase::SetStreamOutTargets(uint32 NumSimultaneousStreamOutTargets, FD3D12Resource** SOArray, const uint32* SOOffsets)
+{
+	PipelineState.Graphics.CurrentNumberOfStreamOutTargets = NumSimultaneousStreamOutTargets;
+	if (PipelineState.Graphics.CurrentNumberOfStreamOutTargets > 0)
+	{
+		FMemory::Memcpy(PipelineState.Graphics.CurrentStreamOutTargets, SOArray, sizeof(FD3D12Resource*)* NumSimultaneousStreamOutTargets);
+		FMemory::Memcpy(PipelineState.Graphics.CurrentSOOffsets, SOOffsets, sizeof(uint32*)* NumSimultaneousStreamOutTargets);
+
+		bNeedSetSOs = true;
+	}
+}
+
+template <EShaderFrequency ShaderFrequency>
+void FD3D12StateCacheBase::InternalSetShaderResourceView(FD3D12ShaderResourceView*& SRV, uint32 ResourceIndex, ESRV_Type SrvType, TSetSRVAlternate AlternatePathFunction)
+{
+	check(ResourceIndex < ARRAYSIZE(PipelineState.Common.CurrentShaderResourceViews[ShaderFrequency]));
+	auto& CurrentShaderResourceViews = PipelineState.Common.CurrentShaderResourceViews[ShaderFrequency];
+	if ((CurrentShaderResourceViews[ResourceIndex] != SRV) || GD3D12SkipStateCaching)
+	{
+		// Keep track of the highest bound resource
+		int32& MaxResourceIndex = PipelineState.Common.MaxBoundShaderResourcesIndex[ShaderFrequency];
+		if (SRV != nullptr)
+		{
+			// Mark the SRVs as not cleared
+			bSRVSCleared = false;
+
+			check(ResourceIndex < MAX_SRVS);
+			// Update the max resource index to the highest bound resource index.
+			MaxResourceIndex = FMath::Max(MaxResourceIndex, static_cast<int32>(ResourceIndex));
+		}
+		else
+		{
+			// If this was the highest bound resource...
+			if (MaxResourceIndex == ResourceIndex)
+			{
+				// Adjust the max resource index downwards until we
+				// hit the next non-null slot, or we've run out of slots.
+				do
+				{
+					MaxResourceIndex--;
+				} while (MaxResourceIndex >= 0 && CurrentShaderResourceViews[MaxResourceIndex] == nullptr);
+			}
+		}
+
+		CurrentShaderResourceViews[ResourceIndex] = SRV;
+		bNeedSetSRVsPerShaderStage[ShaderFrequency] = true;
+		bNeedSetSRVs = true;
+
+		if (FD3D12DynamicRHI::ResourceViewsIntersect(PipelineState.Graphics.CurrentDepthStencilTarget, SRV))
+		{
+			check(PipelineState.Graphics.CurrentDepthStencilTarget->GetDesc().Flags & D3D12_DSV_FLAG_READ_ONLY_DEPTH);
+			if (!PipelineState.Common.CurrentShaderResourceViewsIntersectWithDepthRT[ShaderFrequency][ResourceIndex])
+			{
+				PipelineState.Common.CurrentShaderResourceViewsIntersectWithDepthRT[ShaderFrequency][ResourceIndex] = true;
+				PipelineState.Common.ShaderResourceViewsIntersectWithDepthCount++;
+			}
+		}
+		else
+		{
+			if (PipelineState.Common.CurrentShaderResourceViewsIntersectWithDepthRT[ShaderFrequency][ResourceIndex])
+			{
+				PipelineState.Common.CurrentShaderResourceViewsIntersectWithDepthRT[ShaderFrequency][ResourceIndex] = false;
+				PipelineState.Common.ShaderResourceViewsIntersectWithDepthCount--;
+			}
+		}
+
+		if (AlternatePathFunction != nullptr)
+			(*AlternatePathFunction)(this, SRV, ResourceIndex, SrvType);
+	}
+}
+
+template <EShaderFrequency ShaderFrequency>
+void FD3D12StateCacheBase::ClearShaderResourceViews(FD3D12ResourceLocation*& ResourceLocation)
+{
+	int32& MaxResourceIndex = PipelineState.Common.MaxBoundShaderResourcesIndex[ShaderFrequency];
+	auto& CurrentShaderResourceViews = PipelineState.Common.CurrentShaderResourceViews[ShaderFrequency];
+
+	const int32 OriginalMaxResourceIndex = MaxResourceIndex;
+	for (int32 i = 0; i <= OriginalMaxResourceIndex; ++i)
+	{
+		if (CurrentShaderResourceViews[i].GetReference() && CurrentShaderResourceViews[i]->GetResourceLocation() == ResourceLocation)
+		{
+			SetShaderResourceView<ShaderFrequency>(nullptr, i);
+		}
+	}
+
+	// Restore the max index to make sure we set null descriptors in the slots of any SRVs we unbind due to them overlapping with other
+	// resources like RTVs and DSVs.
+	MaxResourceIndex = OriginalMaxResourceIndex;
+}
+
+void FD3D12StateCacheBase::FlushComputeShaderCache(bool bForce)
+{
+	if (bAutoFlushComputeShaderCache || bForce)
+	{
+		D3D12_RESOURCE_BARRIER desc = {};
+		desc.Type = D3D12_RESOURCE_BARRIER_TYPE_UAV;
+		desc.UAV.pResource = nullptr;
+
+		FD3D12CommandListHandle& CommandList = CmdContext->CommandListHandle;
+
+		CmdContext->numBarriers++;
+		CommandList->ResourceBarrier(1, &desc);
+	}
+}
+
+void FD3D12DescriptorCache::HeapRolledOver(D3D12_DESCRIPTOR_HEAP_TYPE Type)
+{
+	TArray<ID3D12DescriptorHeap*> Heaps;
+	Heaps.Add(ViewHeap.Heap);
+	Heaps.Add(SamplerHeap.Heap);
+	CmdContext->SetDescriptorHeaps(Heaps);
+
+	if (Type == D3D12_DESCRIPTOR_HEAP_TYPE_CBV_SRV_UAV)
+	{
+		ViewHeapSequenceNumber++;
+
+		SRVMap.Reset();
+	}
+	else
+	{
+		check(Type == D3D12_DESCRIPTOR_HEAP_TYPE_SAMPLER);
+
+		SamplerMap.Reset();
+	}
+}
+
+void FD3D12DescriptorCache::Init(FD3D12Device* InParent, FD3D12CommandContext* InCmdContext)
+{
+	Parent = InParent;
+	CmdContext = InCmdContext;
+	ViewHeap.SetParent(this);
+	SamplerHeap.SetParent(this);
+
+	D3D12_DESCRIPTOR_HEAP_FLAGS HeapFlags = D3D12_DESCRIPTOR_HEAP_FLAG_SHADER_VISIBLE;
+
+	ViewHeap.Desc.Flags = HeapFlags;
+	SamplerHeap.Desc.Flags = HeapFlags;
+	ViewHeap.Desc.Type = D3D12_DESCRIPTOR_HEAP_TYPE_CBV_SRV_UAV;
+	SamplerHeap.Desc.Type = D3D12_DESCRIPTOR_HEAP_TYPE_SAMPLER;
+
+	D3D12_FEATURE_DATA_D3D12_OPTIONS D3D12Caps;
+	VERIFYD3D11RESULT(GetParentDevice()->GetDevice()->CheckFeatureSupport(D3D12_FEATURE_D3D12_OPTIONS, &D3D12Caps, sizeof(D3D12Caps)));
+	ViewHeap.Desc.NumDescriptors = (D3D12Caps.ResourceBindingTier == D3D12_RESOURCE_BINDING_TIER_1) ? NUM_VIEW_DESCRIPTORS_TIER_1 : NUM_VIEW_DESCRIPTORS_TIER_2;
+	SamplerHeap.Desc.NumDescriptors = NUM_SAMPLER_DESCRIPTORS;
+
+	VERIFYD3D11RESULT(GetParentDevice()->GetDevice()->CreateDescriptorHeap(&ViewHeap.Desc, IID_PPV_ARGS(ViewHeap.Heap.GetInitReference())));
+	VERIFYD3D11RESULT(GetParentDevice()->GetDevice()->CreateDescriptorHeap(&SamplerHeap.Desc, IID_PPV_ARGS(SamplerHeap.Heap.GetInitReference())));
+
+	ViewHeap.CPUBase = ViewHeap.Heap->GetCPUDescriptorHandleForHeapStart();
+	ViewHeap.GPUBase = ViewHeap.Heap->GetGPUDescriptorHandleForHeapStart();
+	ViewHeap.DescriptorSize = GetParentDevice()->GetDevice()->GetDescriptorHandleIncrementSize(ViewHeap.Desc.Type);
+	SamplerHeap.CPUBase = SamplerHeap.Heap->GetCPUDescriptorHandleForHeapStart();
+	SamplerHeap.GPUBase = SamplerHeap.Heap->GetGPUDescriptorHandleForHeapStart();
+	SamplerHeap.DescriptorSize = GetParentDevice()->GetDevice()->GetDescriptorHandleIncrementSize(SamplerHeap.Desc.Type);
+
+	TArray<ID3D12DescriptorHeap*> Heaps;
+	Heaps.Add(ViewHeap.Heap);
+	Heaps.Add(SamplerHeap.Heap);
+	InCmdContext->SetDescriptorHeaps(Heaps);
+
+	// Create default views
+	D3D12_SHADER_RESOURCE_VIEW_DESC SRVDesc = {};
+	SRVDesc.Shader4ComponentMapping = D3D12_DEFAULT_SHADER_4_COMPONENT_MAPPING;
+	SRVDesc.ViewDimension = D3D12_SRV_DIMENSION_TEXTURE2D;
+	SRVDesc.Format = DXGI_FORMAT_B8G8R8A8_UNORM;
+	SRVDesc.Texture2D.MipLevels = 1;
+	SRVDesc.Texture2D.MostDetailedMip = 0;
+	SRVDesc.Texture2D.ResourceMinLODClamp = 0.0f;
+	pNullSRV = new FD3D12ShaderResourceView(GetParentDevice(), &SRVDesc, nullptr);
+
+	D3D12_RENDER_TARGET_VIEW_DESC RTVDesc = {};
+	RTVDesc.ViewDimension = D3D12_RTV_DIMENSION_TEXTURE2D;
+	RTVDesc.Format = DXGI_FORMAT_B8G8R8A8_UNORM;
+	RTVDesc.Texture2D.MipSlice = 0;
+	pNullRTV = new FD3D12RenderTargetView(GetParentDevice(), &RTVDesc, nullptr);
+
+	D3D12_UNORDERED_ACCESS_VIEW_DESC UAVDesc = {};
+	UAVDesc.ViewDimension = D3D12_UAV_DIMENSION_TEXTURE2D;
+	UAVDesc.Format = DXGI_FORMAT_R8G8B8A8_UNORM;
+	UAVDesc.Texture2D.MipSlice = 0;
+	pNullUAV = new FD3D12UnorderedAccessView(GetParentDevice(), &UAVDesc, nullptr);
+
+	// Init offline descriptor allocators
+	CBVAllocator.Init(GetParentDevice()->GetDevice());
+	FDescriptorHeapManager::HeapIndex heapIndex;
+	pNullCBV = CBVAllocator.AllocateHeapSlot(heapIndex);
+	check(pNullCBV.ptr != 0);
+
+	D3D12_CONSTANT_BUFFER_VIEW_DESC CBView = {};
+	CBView.BufferLocation = 0;
+	CBView.SizeInBytes = 0;
+	GetParentDevice()->GetDevice()->CreateConstantBufferView(&CBView, pNullCBV);
+
+	FSamplerStateInitializerRHI SamplerDesc(
+		SF_Trilinear,
+		AM_Clamp,
+		AM_Clamp,
+		AM_Clamp,
+		0,
+		0,
+		0,
+		FLT_MAX
+		);
+
+	FSamplerStateRHIRef Sampler = GetParentDevice()->CreateSamplerState(SamplerDesc);
+
+	pDefaultSampler = static_cast<FD3D12SamplerState*>(Sampler.GetReference());
+
+	// The default sampler must have ID=0
+	// DescriptorCache::SetSamplers relies on this
+	check(pDefaultSampler->ID == 0);
+
+	// Create offline SRV heaps
+	for (uint32 ShaderFrequency = 0; ShaderFrequency < _countof(OfflineHeap); ShaderFrequency++)
+	{
+		OfflineHeap[ShaderFrequency].CurrentSRVSequenceNumber.AddZeroed(D3D12_COMMONSHADER_INPUT_RESOURCE_SLOT_COUNT);
+
+		OfflineHeap[ShaderFrequency].CurrentUniformBufferSequenceNumber.AddZeroed(D3D12_COMMONSHADER_CONSTANT_BUFFER_API_SLOT_COUNT);
+
+		D3D12_DESCRIPTOR_HEAP_DESC HeapDesc = {};
+		HeapDesc.Type = D3D12_DESCRIPTOR_HEAP_TYPE_CBV_SRV_UAV;
+		HeapDesc.NumDescriptors = D3D12_COMMONSHADER_CONSTANT_BUFFER_API_SLOT_COUNT + D3D12_COMMONSHADER_INPUT_RESOURCE_SLOT_COUNT;
+		HeapDesc.Flags = D3D12_DESCRIPTOR_HEAP_FLAG_NONE;
+
+		VERIFYD3D11RESULT(GetParentDevice()->GetDevice()->CreateDescriptorHeap(
+			&HeapDesc,
+			IID_PPV_ARGS(OfflineHeap[ShaderFrequency].Heap.GetInitReference())
+			));
+
+		D3D12_CPU_DESCRIPTOR_HANDLE HeapBase = OfflineHeap[ShaderFrequency].Heap->GetCPUDescriptorHandleForHeapStart();
+
+		// The first 14 descriptors are for CBVs
+		// The next 128 are for SRVs
+		OfflineHeap[ShaderFrequency].CBVBaseHandle = HeapBase;
+
+		OfflineHeap[ShaderFrequency].SRVBaseHandle = CD3DX12_CPU_DESCRIPTOR_HANDLE(HeapBase, D3D12_COMMONSHADER_CONSTANT_BUFFER_API_SLOT_COUNT, ViewHeap.DescriptorSize);
+
+		// Fill the offline heap with null descriptors
+		for (uint32 i = 0; i < D3D12_COMMONSHADER_CONSTANT_BUFFER_API_SLOT_COUNT; i++)
+		{
+			CD3DX12_CPU_DESCRIPTOR_HANDLE OfflineDescriptor(OfflineHeap[ShaderFrequency].CBVBaseHandle, i, ViewHeap.DescriptorSize);
+
+			CD3DX12_CPU_DESCRIPTOR_HANDLE SrcDescriptor = pNullCBV;
+
+			GetParentDevice()->GetDevice()->CopyDescriptors(
+				1, &OfflineDescriptor, nullptr,
+				1, &SrcDescriptor, nullptr,
+				ViewHeap.Desc.Type
+				);
+		}
+
+		for (uint32 i = 0; i < D3D12_COMMONSHADER_INPUT_RESOURCE_SLOT_COUNT; i++)
+		{
+			CD3DX12_CPU_DESCRIPTOR_HANDLE OfflineDescriptor(OfflineHeap[ShaderFrequency].SRVBaseHandle, i, ViewHeap.DescriptorSize);
+			CD3DX12_CPU_DESCRIPTOR_HANDLE SrcDescriptor = pNullSRV->GetView();
+
+			GetParentDevice()->GetDevice()->CopyDescriptors(
+				1, &OfflineDescriptor, nullptr,
+				1, &SrcDescriptor, nullptr,
+				ViewHeap.Desc.Type
+				);
+		}
+	}
+}
+
+void FD3D12DescriptorCache::Clear()
+{
+	pNullSRV = nullptr;
+	pNullUAV = nullptr;
+	pNullRTV = nullptr;
+}
+
+void FD3D12DescriptorCache::EndFrame()
+{
+	SRVMap.Reset();
+}
+
+void FD3D12DescriptorCache::SetIndexBuffer(FD3D12ResourceLocation* IndexBufferLocation, DXGI_FORMAT Format, uint32 Offset)
+{
+	if (IndexBufferLocation != NULL)
+	{
+		FD3D12Resource* IndexBuffer = IndexBufferLocation->GetResource();
+
+		D3D12_INDEX_BUFFER_VIEW IBView;
+		IBView.BufferLocation = IndexBuffer->GetResource()->GetGPUVirtualAddress() + IndexBufferLocation->GetOffset() + Offset;
+		IBView.Format = Format;
+		IBView.SizeInBytes = IndexBufferLocation->GetEffectiveBufferSize();
+
+		CmdContext->CommandListHandle->IASetIndexBuffer(&IBView);
+	}
+}
+
+void FD3D12DescriptorCache::SetVertexBuffers(FD3D12VertexBufferState* VertexBuffers, uint32 Count)
+{
+	// Determine how many slots are really needed, since the Count passed in is a pre-defined maximum
+	uint32 SlotsNeeded = 0;
+	for (int32 i = Count - 1; i >= 0; i--)
+	{
+		FD3D12Resource *Resource = VertexBuffers[i].VertexBufferLocation ? VertexBuffers[i].VertexBufferLocation->GetResource() : nullptr;
+
+		if (Resource)
+		{
+			SlotsNeeded = i + 1;
+			break;
+		}
+	}
+
+	if (0 == SlotsNeeded)
+	{
+		return; // No-op
+	}
+
+	D3D12_VERTEX_BUFFER_VIEW VBViews[D3D12_IA_VERTEX_INPUT_RESOURCE_SLOT_COUNT];
+
+	// Fill heap slots
+	for (uint32 i = 0; i < SlotsNeeded; i++)
+	{
+		ID3D12Resource *VertexBuffer = VertexBuffers[i].VertexBufferLocation ? VertexBuffers[i].VertexBufferLocation->GetResource()->GetResource() : nullptr;
+		if (VertexBuffer != NULL)
+		{
+			D3D12_VERTEX_BUFFER_VIEW &currentView = VBViews[i];
+			currentView.BufferLocation = VertexBuffer->GetGPUVirtualAddress() + VertexBuffers[i].VertexBufferLocation->GetOffset() + VertexBuffers[i].Offset;
+			currentView.StrideInBytes = VertexBuffers[i].Stride;
+			currentView.SizeInBytes = VertexBuffers[i].VertexBufferLocation->GetEffectiveBufferSize() - VertexBuffers[i].Offset;	// Make sure we account for how much we offset into the VB
+		}
+		else
+		{
+			FMemory::Memzero(&VBViews[i], sizeof(VBViews[i]));
+		}
+	}
+
+	CmdContext->CommandListHandle->IASetVertexBuffers(0, SlotsNeeded, VBViews);
+}
+
+void FD3D12DescriptorCache::SetUAVs(EShaderFrequency ShaderStage, uint32 UAVStartSlot, TRefCountPtr<FD3D12UnorderedAccessView>* UnorderedAccessViewArray, uint32 SlotsNeeded, uint32 &HeapSlot)
+{
+	if (0 == SlotsNeeded)
+	{
+		return; // No-op
+	}
+
+	// Reserve heap slots
+	uint32 FirstSlotIndex = HeapSlot;
+	HeapSlot += SlotsNeeded;
+	CD3DX12_CPU_DESCRIPTOR_HANDLE DestDescriptor(ViewHeap.GetCPUSlotHandle(FirstSlotIndex));
+	CD3DX12_GPU_DESCRIPTOR_HANDLE BindDescriptor(ViewHeap.GetGPUSlotHandle(FirstSlotIndex));
+
+	CD3DX12_CPU_DESCRIPTOR_HANDLE SrcDescriptors[D3D12_PS_CS_UAV_REGISTER_COUNT];
+
+	FD3D12CommandListHandle& CommandList = CmdContext->CommandListHandle;
+
+	// Fill heap slots
+	for (uint32 i = 0; i < UAVStartSlot; i++)
+	{
+		SrcDescriptors[i] = pNullUAV->GetView();
+	}
+	for (uint32 i = 0; i < SlotsNeeded; i++)
+	{
+		if (UnorderedAccessViewArray[i] != NULL)
+		{
+			FD3D12DynamicRHI::TransitionResource(CommandList, UnorderedAccessViewArray[i], D3D12_RESOURCE_STATE_UNORDERED_ACCESS);
+			SrcDescriptors[i + UAVStartSlot] = UnorderedAccessViewArray[i]->GetView();
+		}
+		else
+		{
+			SrcDescriptors[i + UAVStartSlot] = pNullUAV->GetView();
+		}
+	}
+
+	// Gather the descriptors from the offline heaps to the online heap
+	uint32 NumCopySlots = SlotsNeeded + UAVStartSlot;
+	GetParentDevice()->GetDevice()->CopyDescriptors(
+		1, &DestDescriptor, &NumCopySlots,
+		NumCopySlots, SrcDescriptors, nullptr /* sizes */,
+		ViewHeap.Desc.Type);
+
+	if (ShaderStage == SF_Pixel)
+	{
+		CommandList->SetGraphicsRootDescriptorTable(Graphics_UAVs, BindDescriptor);
+	}
+	else if (ShaderStage == SF_Compute)
+	{
+		CommandList->SetComputeRootDescriptorTable(Compute_UAVs, BindDescriptor);
+	}
+	else
+	{
+		check(false);
+	}
+
+#if SUPPORTS_MEMORY_RESIDENCY
+	for (uint32 i = 0; i < SlotsNeeded; i++)
+	{
+		if (UnorderedAccessViewArray[i] != NULL)
+			UnorderedAccessViewArray[i]->GetResource()->UpdateResidency();
+	}
+#endif
+
+#ifdef VERBOSE_DESCRIPTOR_HEAP_DEBUG
+	FMsg::Logf(__FILE__, __LINE__, TEXT("DescriptorCache"), ELogVerbosity::Log, TEXT("SetUnorderedAccessViewTable [STAGE %d] to slots %d - %d"), (int32)ShaderStage, FirstSlotIndex, FirstSlotIndex + SlotsNeeded - 1);
+#endif
+
+}
+
+void FD3D12DescriptorCache::SetRenderTargets(FD3D12RenderTargetView **RenderTargetViewArray, uint32 Count, FD3D12DepthStencilView* DepthStencilTarget, bool bDepthIsBoundAsSRV)
+{
+	// NOTE: For this function, setting zero render targets might not be a no-op, since this is also used
+	//       sometimes for only setting a depth stencil.
+
+	D3D12_CPU_DESCRIPTOR_HANDLE RTVDescriptors[D3D12_SIMULTANEOUS_RENDER_TARGET_COUNT];
+
+	FD3D12CommandListHandle& CommandList = CmdContext->CommandListHandle;
+
+	// Fill heap slots
+	for (uint32 i = 0; i < Count; i++)
+	{
+		if (RenderTargetViewArray[i] != NULL)
+		{
+			// RTV should already be in the correct state. It is transitioned in RHISetRenderTargets.
+			FD3D12DynamicRHI::TransitionResource(CommandList, RenderTargetViewArray[i], D3D12_RESOURCE_STATE_RENDER_TARGET);
+			RTVDescriptors[i] = RenderTargetViewArray[i]->GetView();
+		}
+		else
+		{
+			RTVDescriptors[i] = pNullRTV->GetView();
+		}
+	}
+
+	if (DepthStencilTarget != NULL)
+	{
+		if (bDepthIsBoundAsSRV)
+		{
+			// Depth is also bound as an SRV. Need to add some SRV state to the depth plane, so transition the depth and stencil separately.
+			const D3D12_DEPTH_STENCIL_VIEW_DESC &desc = DepthStencilTarget->GetDesc();
+			check(desc.Flags & D3D12_DSV_FLAG_READ_ONLY_DEPTH);
+			const D3D12_RESOURCE_STATES depthState = D3D12_RESOURCE_STATE_DEPTH_READ | D3D12_RESOURCE_STATE_PIXEL_SHADER_RESOURCE | D3D12_RESOURCE_STATE_NON_PIXEL_SHADER_RESOURCE;
+			FD3D12DynamicRHI::TransitionResource(CommandList, DepthStencilTarget->GetResource(), depthState, DepthStencilTarget->GetDepthOnlyViewSubresourceSubset());
+
+			if (DepthStencilTarget->HasStencil())
+			{
+				const D3D12_RESOURCE_STATES stencilState = (desc.Flags & D3D12_DSV_FLAG_READ_ONLY_STENCIL) ? D3D12_RESOURCE_STATE_DEPTH_READ : D3D12_RESOURCE_STATE_DEPTH_WRITE;
+				FD3D12DynamicRHI::TransitionResource(CommandList, DepthStencilTarget->GetResource(), stencilState, DepthStencilTarget->GetStencilOnlyViewSubresourceSubset());
+			}
+		}
+		else
+		{
+			check(!bDepthIsBoundAsSRV);
+			FD3D12DynamicRHI::TransitionResource(CommandList, DepthStencilTarget);
+		}
+
+		D3D12_CPU_DESCRIPTOR_HANDLE DSVDescriptor = DepthStencilTarget->GetView();
+		CommandList->OMSetRenderTargets(Count, RTVDescriptors, 0, &DSVDescriptor);
+	}
+	else
+	{
+		CommandList->OMSetRenderTargets(Count, RTVDescriptors, 0, nullptr);
+	}
+}
+
+void FD3D12DescriptorCache::SetStreamOutTargets(FD3D12Resource **Buffers, uint32 Count, const uint32* Offsets)
+{
+	// Determine how many slots are really needed, since the Count passed in is a pre-defined maximum
+	uint32 SlotsNeeded = 0;
+	for (int32 i = Count - 1; i >= 0; i--)
+	{
+		if (Buffers[i] != NULL)
+		{
+			SlotsNeeded = i + 1;
+			break;
+		}
+	}
+
+	if (0 == SlotsNeeded)
+	{
+		return; // No-op
+	}
+
+	D3D12_STREAM_OUTPUT_BUFFER_VIEW SOViews[D3D12_SO_BUFFER_SLOT_COUNT] = { 0 };
+
+	FD3D12CommandListHandle& CommandList = CmdContext->CommandListHandle;
+
+	// Fill heap slots
+	for (uint32 i = 0; i < SlotsNeeded; i++)
+	{
+		ID3D12Resource *StreamOutBuffer = Buffers[i] ? Buffers[i]->GetResource() : nullptr;
+		
+		D3D12_STREAM_OUTPUT_BUFFER_VIEW &currentView = SOViews[i];
+		currentView.BufferLocation = (StreamOutBuffer != nullptr) ? StreamOutBuffer->GetGPUVirtualAddress() : 0;
+		
+		// MS - The following view members are not correct
+		check(0);
+		currentView.BufferFilledSizeLocation = 0;
+		currentView.SizeInBytes = -1;
+
+		if (Buffers[i] != nullptr)
+		{
+			FD3D12DynamicRHI::TransitionResource(CommandList, Buffers[i], D3D12_RESOURCE_STATE_STREAM_OUT, D3D12_RESOURCE_BARRIER_ALL_SUBRESOURCES);
+		}
+	}
+
+	CommandList->SOSetTargets(0, SlotsNeeded, SOViews);
+}
+
+void FD3D12DescriptorCache::SetSamplers(EShaderFrequency ShaderStage, FD3D12SamplerState** Samplers, uint32 SlotsNeeded, uint32 &HeapSlot)
+{
+	if (0 == SlotsNeeded)
+	{
+		return; // No-op
+	}
+
+	D3D12_GPU_DESCRIPTOR_HANDLE BindDescriptor = { 0 };
+	bool CacheHit = false;
+
+	// Check to see if the sampler configuration is already in the sampler heap
+	FD3D12SamplerArrayDesc Desc;
+	if (SlotsNeeded <= _countof(Desc.SamplerID))
+	{
+		Desc.Count = SlotsNeeded;
+
+		for (uint32 i = 0; i < SlotsNeeded; i++)
+		{
+			Desc.SamplerID[i] = Samplers[i] ? Samplers[i]->ID : 0;
+		}
+
+		// The hash uses all of the bits
+		for (uint32 i = SlotsNeeded; i < _countof(Desc.SamplerID); i++)
+		{
+			Desc.SamplerID[i] = 0;
+		}
+
+		D3D12_GPU_DESCRIPTOR_HANDLE* FoundDescriptor = SamplerMap.Find(Desc);
+
+		if (FoundDescriptor)
+		{
+			BindDescriptor = *FoundDescriptor;
+			CacheHit = true;
+		}
+	}
+
+	if (!CacheHit)
+	{
+		// Reserve heap slots
+		uint32 FirstSlotIndex = HeapSlot;
+		HeapSlot += SlotsNeeded;
+		D3D12_CPU_DESCRIPTOR_HANDLE DestDescriptor = SamplerHeap.GetCPUSlotHandle(FirstSlotIndex);
+		BindDescriptor = SamplerHeap.GetGPUSlotHandle(FirstSlotIndex);
+
+		CD3DX12_CPU_DESCRIPTOR_HANDLE SrcDescriptors[D3D12_COMMONSHADER_SAMPLER_SLOT_COUNT];
+		// Fill heap slots
+		for (uint32 i = 0; i < SlotsNeeded; i++)
+		{
+			if (Samplers[i] != NULL)
+			{
+				SrcDescriptors[i] = Samplers[i]->Descriptor;
+			}
+			else
+			{
+				SrcDescriptors[i] = pDefaultSampler->Descriptor;
+			}
+		}
+
+		GetParentDevice()->GetDevice()->CopyDescriptors(
+			1, &DestDescriptor, &SlotsNeeded,
+			SlotsNeeded, SrcDescriptors, nullptr /* sizes */,
+			SamplerHeap.Desc.Type);
+
+		// Remember the locations of the samplers in the sampler map
+		if (SlotsNeeded <= _countof(Desc.SamplerID))
+		{
+			SamplerMap.Add(Desc, BindDescriptor);
+		}
+	}
+
+	FD3D12CommandListHandle& CommandList = CmdContext->CommandListHandle;
+
+	if (ShaderStage == SF_Compute)
+	{
+		CommandList->SetComputeRootDescriptorTable(Compute_Samplers, BindDescriptor);
+	}
+	else
+	{
+		uint32 RTLIndex = SamplerRTLBindSlot(ShaderStage);
+		CommandList->SetGraphicsRootDescriptorTable(RTLIndex, BindDescriptor);
+	}
+
+#ifdef VERBOSE_DESCRIPTOR_HEAP_DEBUG
+	FMsg::Logf(__FILE__, __LINE__, TEXT("DescriptorCache"), ELogVerbosity::Log, TEXT("SetSamplerTable [STAGE %d] to slots %d - %d"), (int32)ShaderStage, FirstSlotIndex, FirstSlotIndex + SlotsNeeded - 1);
+#endif
+}
+
+void FD3D12DescriptorCache::SetSRVs(EShaderFrequency ShaderStage, TRefCountPtr<FD3D12ShaderResourceView> * SRVs, uint32 SlotsNeeded, bool* CurrentShaderResourceViewsIntersectWithDepthRT, uint32 &HeapSlot)
+{
+	if (0 == SlotsNeeded)
+	{
+		return; // No-op
+	}
+
+	// Check hash table for hit
+
+	D3D12_GPU_DESCRIPTOR_HANDLE BindDescriptor = { 0 };
+
+	FD3D12CommandListHandle& CommandList = CmdContext->CommandListHandle;
+
+	D3D12_RESOURCE_STATES ResourceStateFlag = D3D12_RESOURCE_STATE_PIXEL_SHADER_RESOURCE | D3D12_RESOURCE_STATE_NON_PIXEL_SHADER_RESOURCE;
+
+	// Check to see if the srv configuration is already in the srv heap
+	FD3D12SRVArrayDesc Desc = {0};
+	check (SlotsNeeded <= _countof(Desc.SRVSequenceNumber));
+	Desc.Count = SlotsNeeded;
+
+	for (uint32 i = 0; i < SlotsNeeded; i++)
+	{
+		if (SRVs[i] != NULL)
+		{
+			Desc.SRVSequenceNumber[i] = SRVs[i]->GetSequenceNumber();
+	
+			if (CurrentShaderResourceViewsIntersectWithDepthRT[i] == true)
+			{
+				FD3D12DynamicRHI::TransitionResource(CommandList, SRVs[i], ResourceStateFlag | D3D12_RESOURCE_STATE_DEPTH_READ);
+			}
+			else
+			{
+				FD3D12DynamicRHI::TransitionResource(CommandList, SRVs[i], ResourceStateFlag);
+			}
+		}
+	}
+
+	D3D12_GPU_DESCRIPTOR_HANDLE* FoundDescriptor = SRVMap.Find(Desc);
+
+	if (FoundDescriptor)
+		BindDescriptor = *FoundDescriptor;
+
+	else
+	{
+		// Reserve heap slots
+		uint32 FirstSlotIndex = HeapSlot;
+		HeapSlot += SlotsNeeded;
+
+		TArray<uint64>& CurrentSequenceNumber = OfflineHeap[ShaderStage].CurrentSRVSequenceNumber;
+		CD3DX12_CPU_DESCRIPTOR_HANDLE OfflineBase = OfflineHeap[ShaderStage].SRVBaseHandle;
+
+		uint32 DescriptorSize = ViewHeap.DescriptorSize;
+
+		// Copy to offline heap (only those descriptors which have changed)
+		CD3DX12_CPU_DESCRIPTOR_HANDLE OfflineDescriptor = OfflineBase;
+
+		for (uint32 i = 0; i < SlotsNeeded; i++)
+		{
+			CD3DX12_CPU_DESCRIPTOR_HANDLE SrcDescriptor;
+			uint64 SequenceNumber = 0;
+
+			if (SRVs[i] != NULL)
+			{
+				SequenceNumber = SRVs[i]->GetSequenceNumber();
+				SrcDescriptor = SRVs[i]->GetView();
+			}
+			else
+			{
+				SrcDescriptor = pNullSRV->GetView();
+			}
+			check(SrcDescriptor.ptr != 0);
+
+			if (SequenceNumber != CurrentSequenceNumber[i])
+			{
+				GetParentDevice()->GetDevice()->CopyDescriptors(
+					1, &OfflineDescriptor, nullptr,
+					1, &SrcDescriptor, nullptr,
+					ViewHeap.Desc.Type
+					);
+
+				CurrentSequenceNumber[i] = SequenceNumber;
+			}
+
+			OfflineDescriptor.Offset(DescriptorSize);
+		}
+
+		// Copy from offline heap to online heap
+		D3D12_CPU_DESCRIPTOR_HANDLE DestDescriptor = ViewHeap.GetCPUSlotHandle(FirstSlotIndex);
+
+		// Copy the descriptors from the offline heaps to the online heap
+		GetParentDevice()->GetDevice()->CopyDescriptors(
+			1, &DestDescriptor, &SlotsNeeded,
+			1, &OfflineBase, &SlotsNeeded,
+			ViewHeap.Desc.Type
+			);
+
+		BindDescriptor = ViewHeap.GetGPUSlotHandle(FirstSlotIndex);
+
+		// Remember the locations of the srvs in the map
+		SRVMap.Add(Desc, BindDescriptor);
+	}
+
+	if (ShaderStage == SF_Compute)
+	{
+		CommandList->SetComputeRootDescriptorTable(Compute_SRVs, BindDescriptor);
+	}
+	else
+	{
+		uint32 RTLIndex = SRVRTLBindSlot(ShaderStage);
+		CommandList->SetGraphicsRootDescriptorTable(RTLIndex, BindDescriptor);
+	}
+
+#if SUPPORTS_MEMORY_RESIDENCY
+	for (uint32 i = 0; i < SlotsNeeded; i++)
+	{
+		if (SRVs[i] != NULL)
+			SRVs[i]->GetResource()->UpdateResidency();
+	}
+#endif
+
+#ifdef VERBOSE_DESCRIPTOR_HEAP_DEBUG
+	FMsg::Logf(__FILE__, __LINE__, TEXT("DescriptorCache"), ELogVerbosity::Log, TEXT("SetShaderResourceViewTable [STAGE %d] to slots %d - %d"), (int32)ShaderStage, FirstSlotIndex, FirstSlotIndex + SlotsNeeded - 1);
+#endif
+}
+
+// Updates the offline heap
+void FD3D12DescriptorCache::SetConstantBuffer(
+	EShaderFrequency ShaderStage,
+	uint32 SlotIndex,
+	FD3D12UniformBuffer* UniformBuffer
+	)
+{
+	check(SlotIndex < D3D12_COMMONSHADER_CONSTANT_BUFFER_API_SLOT_COUNT);
+	check(UniformBuffer);
+	check(UniformBuffer->OfflineDescriptorHandle.ptr != 0);
+
+	CD3DX12_CPU_DESCRIPTOR_HANDLE DestCPUHandle(OfflineHeap[ShaderStage].CBVBaseHandle, SlotIndex, ViewHeap.DescriptorSize);
+
+	TArray<uint64>& CurrentSequenceNumber = OfflineHeap[ShaderStage].CurrentUniformBufferSequenceNumber;
+
+	// Only call CopyDescriptors if this is a new uniform buffer for this slot
+	if (CurrentSequenceNumber[SlotIndex] != UniformBuffer->SequenceNumber)
+	{
+		GetParentDevice()->GetDevice()->CopyDescriptors(
+			1, &DestCPUHandle, nullptr /*sizes, null means assume 1*/,
+			1, &UniformBuffer->OfflineDescriptorHandle, nullptr/*sizes, null means assume 1*/,
+			ViewHeap.Desc.Type);
+
+		CurrentSequenceNumber[SlotIndex] = UniformBuffer->SequenceNumber;
+	}
+}
+
+void FD3D12DescriptorCache::SetConstantBuffer(
+	EShaderFrequency ShaderStage,
+	uint32 SlotIndex,
+	ID3D12Resource* Resource,
+	uint32 OffsetInBytes,
+	uint32 SizeInBytes
+	)
+{
+	check(SlotIndex < D3D12_COMMONSHADER_CONSTANT_BUFFER_API_SLOT_COUNT);
+	check(Resource);
+
+	CD3DX12_CPU_DESCRIPTOR_HANDLE DestCPUHandle(OfflineHeap[ShaderStage].CBVBaseHandle, SlotIndex, ViewHeap.DescriptorSize);
+
+	TArray<uint64>& CurrentSequenceNumber = OfflineHeap[ShaderStage].CurrentUniformBufferSequenceNumber;
+
+	D3D12_CONSTANT_BUFFER_VIEW_DESC CBView;
+	CBView.BufferLocation = Resource->GetGPUVirtualAddress() + OffsetInBytes;
+	CBView.SizeInBytes = SizeInBytes;
+	check(OffsetInBytes % 256 == 0);
+	check(CBView.SizeInBytes <= 4096 * 16);
+	check(CBView.SizeInBytes % 256 == 0);
+
+	GetParentDevice()->GetDevice()->CreateConstantBufferView(&CBView, DestCPUHandle);
+
+	CurrentSequenceNumber[SlotIndex] = 0;
+}
+
+void FD3D12DescriptorCache::ClearConstantBuffer(
+	EShaderFrequency ShaderStage,
+	uint32 SlotIndex
+	)
+{
+	check(SlotIndex < D3D12_COMMONSHADER_CONSTANT_BUFFER_API_SLOT_COUNT);
+
+	CD3DX12_CPU_DESCRIPTOR_HANDLE DestCPUHandle(OfflineHeap[ShaderStage].CBVBaseHandle, SlotIndex, ViewHeap.DescriptorSize);
+
+	TArray<uint64>& CurrentSequenceNumber = OfflineHeap[ShaderStage].CurrentUniformBufferSequenceNumber;
+
+	// Only call CopyDescriptors if this is a new uniform buffer for this slot
+	if (CurrentSequenceNumber[SlotIndex] != 0)
+	{
+		GetParentDevice()->GetDevice()->CopyDescriptors(
+			1, &DestCPUHandle, nullptr /*sizes, null means assume 1*/,
+			1, &pNullCBV, nullptr/*sizes, null means assume 1*/,
+			ViewHeap.Desc.Type);
+	
+		CurrentSequenceNumber[SlotIndex] = 0;
+	}
+}
+
+void FD3D12DescriptorCache::SetConstantBuffers(EShaderFrequency ShaderStage, uint32 SlotsNeeded, uint32 &HeapSlot)
+{
+	if (0 == SlotsNeeded)
+	{
+		return; // No-op
+	}
+
+	// Reserve heap slots
+	uint32 FirstSlotIndex = HeapSlot;
+	HeapSlot += SlotsNeeded;
+
+	// Copy from offline heap to online heap
+
+	D3D12_CPU_DESCRIPTOR_HANDLE OnlineHeapHandle = ViewHeap.GetCPUSlotHandle(FirstSlotIndex);
+
+	D3D12_CPU_DESCRIPTOR_HANDLE OfflineHeapHandle = OfflineHeap[ShaderStage].CBVBaseHandle;
+
+	GetParentDevice()->GetDevice()->CopyDescriptors(
+		1, &OnlineHeapHandle, &SlotsNeeded,
+		1, &OfflineHeapHandle, &SlotsNeeded,
+		ViewHeap.Desc.Type
+		);
+
+	D3D12_GPU_DESCRIPTOR_HANDLE BindDescriptor = ViewHeap.GetGPUSlotHandle(FirstSlotIndex);
+
+	FD3D12CommandListHandle& CommandList = CmdContext->CommandListHandle;
+
+	if (ShaderStage == SF_Compute)
+	{
+		CommandList->SetComputeRootDescriptorTable(Compute_CBs, BindDescriptor);
+	}
+	else
+	{
+		uint32 RTLIndex = CBCRTLBindSlot(ShaderStage);
+		CommandList->SetGraphicsRootDescriptorTable(RTLIndex, BindDescriptor);
+	}
+
+#ifdef VERBOSE_DESCRIPTOR_HEAP_DEBUG
+	FMsg::Logf(__FILE__, __LINE__, TEXT("DescriptorCache"), ELogVerbosity::Log, TEXT("ConstantBufferViewTable set to slots %d - %d"), FirstSlotIndex, FirstSlotIndex + SlotsNeeded - 1);
+#endif
+}
+
+bool FD3D12DescriptorCache::OnlineHeap::CanReserveSlots(uint32 NumSlots)
+{
+	const uint32 HeapSize = Desc.NumDescriptors;
+
+	// Sanity checks
+	if (0 == NumSlots)
+	{
+		return true;
+	}
+	if (NumSlots > HeapSize)
+	{
+		throw E_OUTOFMEMORY;
+	}
+	uint32 FirstRequestedSlot = NextSlotIndex;
+	uint32 SlotAfterReservation = NextSlotIndex + NumSlots;
+
+	if (SlotAfterReservation > HeapSize)
+	{
+		return false;
+	}
+	return true;
+}
+
+void FD3D12DescriptorCache::OnlineHeap::RollOver()
+{
+	// Roll over
+	PoolEntry Entry;
+	Entry.Heap.Swap(Heap);
+	Entry.SyncPoint = CurrentCommandList;
+	ReclaimPool.Enqueue(Entry);
+
+	if (ReclaimPool.Peek(Entry) &&
+		GetParentDevice()->GetCommandListManager().IsReady() &&
+		GetParentDevice()->GetCommandListManager().GetCommandListState(Entry.SyncPoint) == CommandListState::kFinished)
+	{
+		ReclaimPool.Dequeue(Entry);
+		Heap.Swap(Entry.Heap);
+	}
+	else
+	{
+		VERIFYD3D11RESULT(GetParentDevice()->GetDevice()->CreateDescriptorHeap(&Desc, IID_PPV_ARGS(Heap.GetInitReference())));
+	}
+
+	NextSlotIndex = 0;
+
+	// Notify other layers of heap change
+	CPUBase = Heap->GetCPUDescriptorHandleForHeapStart();
+	GPUBase = Heap->GetGPUDescriptorHandleForHeapStart();
+	Parent->HeapRolledOver(Desc.Type);
+}
+
+uint32 FD3D12DescriptorCache::OnlineHeap::ReserveSlots(uint32 NumSlotsRequested)
+{
+#ifdef VERBOSE_DESCRIPTOR_HEAP_DEBUG
+	FMsg::Logf(__FILE__, __LINE__, TEXT("DescriptorCache"), ELogVerbosity::Log, TEXT("Requesting reservation [TYPE %d] with %d slots, required fence is %d"),
+		(int32)Desc.Type, NumSlotsRequested, RequiredFenceForCurrentCL);
+#endif
+
+	const uint32 HeapSize = Desc.NumDescriptors;
+
+	// Sanity checks
+	if (NumSlotsRequested > HeapSize)
+	{
+		throw E_OUTOFMEMORY;
+	}
+
+	// Decide which slots will be reserved and what needs to be cleaned up
+	uint32 FirstRequestedSlot = NextSlotIndex;
+	uint32 SlotAfterReservation = NextSlotIndex + NumSlotsRequested;
+
+	// CanReserveSlots should have ben called first
+	check(SlotAfterReservation <= HeapSize);
+
+	// Note where to start looking next time
+	NextSlotIndex = SlotAfterReservation;
+
+	return FirstRequestedSlot;
+}
+
+void FD3D12DescriptorCache::OnlineHeap::SetNextSlot(uint32 NextSlot)
+{
+	// For samplers, ReserveSlots will be called with a conservative estimate
+	// This is used to correct for the actual number of heap slots used
+	check(NextSlot <= NextSlotIndex);
+
+	NextSlotIndex = NextSlot;
+}
+
+uint32 GetTypeHash(const FD3D12SamplerArrayDesc& Key)
+{
+	return uint32(FD3D12PipelineStateCache::HashData((void*)Key.SamplerID, Key.Count * sizeof(Key.SamplerID[0])));
+}
+
+uint32 GetTypeHash(const FD3D12SRVArrayDesc& Key)
+{
+	return uint32(FD3D12PipelineStateCache::HashData((void*)Key.SRVSequenceNumber, Key.Count * sizeof(Key.SRVSequenceNumber[0])));
+}
