UE4 - D3D12RHI (not working)

diff --git a/Engine/Source/Runtime/Windows/D3D12RHI/Private/D3D12StateCachePrivate.h b/Engine/Source/Runtime/Windows/D3D12RHI/Private/D3D12StateCachePrivate.h
new file mode 100644
--- /dev/null
+++ b/Engine/Source/Runtime/Windows/D3D12RHI/Private/D3D12StateCachePrivate.h
@@ -0,0 +1,1469 @@
+// Copyright 1998-2014 Epic Games, Inc. All Rights Reserved.
+
+// Implementation of Device Context State Caching to improve draw
+//	thread performance by removing redundant device context calls.
+
+#pragma once
+#include "Queue.h"
+#include "D3D12DirectCommandListManager.h"
+
+//-----------------------------------------------------------------------------
+//	Configuration
+//-----------------------------------------------------------------------------
+
+// If set, includes a runtime toggle console command for debugging D3D11 state caching.
+// ("TOGGLESTATECACHE")
+#define D3D12_STATE_CACHE_RUNTIME_TOGGLE 0
+
+// If set, includes a cache state verification check.
+// After each state set call, the cached state is compared against the actual state of the ID3D11DeviceContext.
+// This is *very slow* and should only be enabled to debug the state caching system.
+#ifndef D3D11_STATE_CACHE_DEBUG
+#define D3D11_STATE_CACHE_DEBUG 0
+#endif
+
+#define MAX_SRVS 16
+#define MAX_CBS 8
+
+// Uncomment only for debugging of the descriptor heap management; this is very noisy
+//#define VERBOSE_DESCRIPTOR_HEAP_DEBUG 1
+
+// The number of view descriptors available per (online) descriptor heap, depending on hardware tier
+#define NUM_VIEW_DESCRIPTORS_TIER_1 250000
+#define NUM_VIEW_DESCRIPTORS_TIER_2 250000
+#define NUM_SAMPLER_DESCRIPTORS 2000
+
+// Heap for updating UAV counter values.
+#define COUNTER_HEAP_SIZE 1024 * 64
+
+enum EGraphicsDescriptorTables
+{
+	Graphics_PS_SRVs,
+	Graphics_PS_CBs,
+	Graphics_PS_Samplers,
+	Graphics_VS_SRVs,
+	Graphics_VS_CBs,
+	Graphics_VS_Samplers,
+	Graphics_GS_SRVs,
+	Graphics_GS_CBs,
+	Graphics_GS_Samplers,
+	Graphics_HS_SRVs,
+	Graphics_HS_CBs,
+	Graphics_HS_Samplers,
+	Graphics_DS_SRVs,
+	Graphics_DS_CBs,
+	Graphics_DS_Samplers,
+	Graphics_UAVs,
+	Graphics_DescriptorTableCount,
+};
+
+enum EComputeDescriptorTables
+{
+	Compute_SRVs,
+	Compute_CBs,
+	Compute_Samplers,
+	Compute_UAVs,
+	Compute_DescriptorTableCount,
+};
+
+enum EShaderStageDescriptorTables
+{
+	Descriptors_SRVs,
+	Descriptors_CBs,
+	Descriptors_Samplers,
+	Descriptors_ShaderStageTableCount,
+};
+
+#define NUM_CB_RTL_SLOTS 4
+
+inline uint32 SamplerRTLBindSlot(EShaderFrequency SF)
+{
+	switch (SF)
+	{
+	case SF_Vertex: return Graphics_VS_Samplers;
+	case SF_Pixel: return Graphics_PS_Samplers;
+	case SF_Geometry: return Graphics_GS_Samplers;
+	case SF_Hull: return Graphics_HS_Samplers;
+	case SF_Domain: return Graphics_DS_Samplers;
+
+	default: check(false); // Caller should check compute first if it's an option
+	}
+	return UINT_MAX;
+}
+
+inline uint32 SRVRTLBindSlot(EShaderFrequency SF)
+{
+	switch (SF)
+	{
+	case SF_Vertex: return Graphics_VS_SRVs;
+	case SF_Pixel: return Graphics_PS_SRVs;
+	case SF_Geometry: return Graphics_GS_SRVs;
+	case SF_Hull: return Graphics_HS_SRVs;
+	case SF_Domain: return Graphics_DS_SRVs;
+	default: check(false); // Caller should check compute first if it's an option
+	}
+	return UINT_MAX;
+}
+
+inline uint32 CBCRTLBindSlot(EShaderFrequency SF)
+{
+	switch (SF)
+	{
+	case SF_Vertex: return Graphics_VS_CBs;
+	case SF_Pixel: return Graphics_PS_CBs;
+	case SF_Geometry: return Graphics_GS_CBs;
+	case SF_Hull: return Graphics_HS_CBs;
+	case SF_Domain: return Graphics_DS_CBs;
+	default: check(false); // Caller should check compute first if it's an option
+		return UINT_MAX;
+	}
+}
+
+// Keep set state functions inline to reduce call overhead
+#define D3D12_STATE_CACHE_INLINE FORCEINLINE
+
+#if D3D12_STATE_CACHE_RUNTIME_TOGGLE
+extern bool GD3D12SkipStateCaching;
+#else
+static const bool GD3D12SkipStateCaching = false;
+#endif
+template <typename T>
+inline void hash_combine(SIZE_T & seed, const T & v)
+{
+	seed ^= GetTypeHash(v) + 0x9e3779b9 + (seed << 6) + (seed >> 2);
+}
+
+#define PSO_IF_NOT_EQUAL_RETURN_FALSE( value ) if(lhs.##value != rhs.##value){ return false; }
+
+#define PSO_IF_MEMCMP_FAILS_RETURN_FALSE( value ) if(FMemory::Memcmp(&lhs.##value, &rhs.##value, sizeof(rhs.##value)) != 0){ return false; }
+
+#define PSO_IF_STRING_COMPARE_FAILS_RETURN_FALSE( value ) \
+	const char* lhString = lhs.##value##; \
+	const char* rhString = rhs.##value##; \
+	if (lhString != rhString) \
+	{ \
+		if (strcmp(lhString, rhString) != 0) \
+		{ \
+			return false; \
+		} \
+	}
+
+template <typename TDesc> struct equality_pipeline_state_desc;
+template <> struct equality_pipeline_state_desc<FD3D12HighLevelGraphicsPipelineStateDesc>
+{
+	bool operator()(const FD3D12HighLevelGraphicsPipelineStateDesc& lhs, const FD3D12HighLevelGraphicsPipelineStateDesc& rhs)
+	{
+		PSO_IF_NOT_EQUAL_RETURN_FALSE(BoundShaderState);
+		PSO_IF_NOT_EQUAL_RETURN_FALSE(BlendState);
+		PSO_IF_NOT_EQUAL_RETURN_FALSE(DepthStencilState);
+		PSO_IF_NOT_EQUAL_RETURN_FALSE(RasterizerState);
+		PSO_IF_NOT_EQUAL_RETURN_FALSE(SampleMask);
+		PSO_IF_NOT_EQUAL_RETURN_FALSE(PrimitiveTopologyType);
+		PSO_IF_NOT_EQUAL_RETURN_FALSE(NumRenderTargets);
+		for (SIZE_T i = 0; i < lhs.NumRenderTargets; i++)
+		{
+			PSO_IF_NOT_EQUAL_RETURN_FALSE(RTVFormats[i]);
+		}
+		PSO_IF_NOT_EQUAL_RETURN_FALSE(DSVFormat);
+		PSO_IF_NOT_EQUAL_RETURN_FALSE(SampleDesc.Count);
+		PSO_IF_NOT_EQUAL_RETURN_FALSE(SampleDesc.Quality);
+		return true;
+	}
+};
+
+template <> struct equality_pipeline_state_desc<FD3D12LowLevelGraphicsPipelineStateDesc>
+{
+	bool operator()(const FD3D12LowLevelGraphicsPipelineStateDesc& lhs, const FD3D12LowLevelGraphicsPipelineStateDesc& rhs)
+	{
+		// Order from most likely to change to least
+		PSO_IF_NOT_EQUAL_RETURN_FALSE(Desc.PS.BytecodeLength)
+		PSO_IF_NOT_EQUAL_RETURN_FALSE(Desc.VS.BytecodeLength)
+		PSO_IF_NOT_EQUAL_RETURN_FALSE(Desc.GS.BytecodeLength)
+		PSO_IF_NOT_EQUAL_RETURN_FALSE(Desc.DS.BytecodeLength)
+		PSO_IF_NOT_EQUAL_RETURN_FALSE(Desc.HS.BytecodeLength)
+		PSO_IF_NOT_EQUAL_RETURN_FALSE(Desc.InputLayout.NumElements)
+		PSO_IF_NOT_EQUAL_RETURN_FALSE(Desc.NumRenderTargets)
+		PSO_IF_NOT_EQUAL_RETURN_FALSE(Desc.DSVFormat)
+		PSO_IF_NOT_EQUAL_RETURN_FALSE(Desc.PrimitiveTopologyType)
+		PSO_IF_NOT_EQUAL_RETURN_FALSE(Desc.Flags)
+		PSO_IF_NOT_EQUAL_RETURN_FALSE(Desc.pRootSignature)
+		PSO_IF_NOT_EQUAL_RETURN_FALSE(Desc.SampleMask)
+		PSO_IF_NOT_EQUAL_RETURN_FALSE(Desc.IBStripCutValue)
+		PSO_IF_NOT_EQUAL_RETURN_FALSE(Desc.NodeMask)
+		PSO_IF_NOT_EQUAL_RETURN_FALSE(Desc.StreamOutput.RasterizedStream)
+		PSO_IF_NOT_EQUAL_RETURN_FALSE(Desc.StreamOutput.NumEntries)
+		PSO_IF_NOT_EQUAL_RETURN_FALSE(Desc.StreamOutput.NumStrides)
+		PSO_IF_NOT_EQUAL_RETURN_FALSE(Desc.SampleDesc.Count)
+		PSO_IF_NOT_EQUAL_RETURN_FALSE(Desc.SampleDesc.Quality)
+
+		PSO_IF_MEMCMP_FAILS_RETURN_FALSE(Desc.BlendState)
+		PSO_IF_MEMCMP_FAILS_RETURN_FALSE(Desc.RasterizerState)
+		PSO_IF_MEMCMP_FAILS_RETURN_FALSE(Desc.DepthStencilState)
+		for (SIZE_T i = 0; i < lhs.Desc.NumRenderTargets; i++)
+		{
+			PSO_IF_NOT_EQUAL_RETURN_FALSE(Desc.RTVFormats[i]);
+		}
+
+		// Shader byte code is hashed with SHA1 (160 bit) so the chances of collision
+		// should be tiny i.e if there were 1 quadrillion shaders the chance of a 
+		// collision is ~ 1 in 10^18. so only do a full check on debug builds
+		PSO_IF_NOT_EQUAL_RETURN_FALSE(VSHash)
+		PSO_IF_NOT_EQUAL_RETURN_FALSE(PSHash)
+		PSO_IF_NOT_EQUAL_RETURN_FALSE(GSHash)
+		PSO_IF_NOT_EQUAL_RETURN_FALSE(HSHash)
+		PSO_IF_NOT_EQUAL_RETURN_FALSE(DSHash)
+
+#if UE_BUILD_DEBUG
+		const D3D12_SHADER_BYTECODE* lhByteCode = &lhs.Desc.VS;
+		const D3D12_SHADER_BYTECODE* rhByteCode = &rhs.Desc.VS;
+		for (SIZE_T i = 0; i < 5; i++)
+		{
+			if (lhByteCode[i].pShaderBytecode != rhByteCode[i].pShaderBytecode &&
+				lhByteCode[i].BytecodeLength)
+			{
+				if (FMemory::Memcmp(lhByteCode[i].pShaderBytecode, rhByteCode[i].pShaderBytecode, rhByteCode[i].BytecodeLength) != 0)
+				{
+					UE_LOG(LogD3D12RHI, Error, TEXT("Error! there is a collision with the SHA1. This should never happen but checking for completeness."));
+					return false;
+				}
+			}
+		}
+#endif
+
+		if (lhs.Desc.StreamOutput.pSODeclaration != rhs.Desc.StreamOutput.pSODeclaration &&
+			lhs.Desc.StreamOutput.NumEntries)
+		{
+			for (SIZE_T i = 0; i < lhs.Desc.StreamOutput.NumEntries; i++)
+			{
+				PSO_IF_NOT_EQUAL_RETURN_FALSE(Desc.StreamOutput.pSODeclaration[i].Stream)
+				PSO_IF_NOT_EQUAL_RETURN_FALSE(Desc.StreamOutput.pSODeclaration[i].SemanticIndex )
+				PSO_IF_NOT_EQUAL_RETURN_FALSE(Desc.StreamOutput.pSODeclaration[i].StartComponent )
+				PSO_IF_NOT_EQUAL_RETURN_FALSE(Desc.StreamOutput.pSODeclaration[i].ComponentCount )
+				PSO_IF_NOT_EQUAL_RETURN_FALSE(Desc.StreamOutput.pSODeclaration[i].OutputSlot )
+				PSO_IF_STRING_COMPARE_FAILS_RETURN_FALSE(Desc.StreamOutput.pSODeclaration[i].SemanticName )
+			}
+		}
+
+		if (lhs.Desc.StreamOutput.pBufferStrides != rhs.Desc.StreamOutput.pBufferStrides &&
+			lhs.Desc.StreamOutput.NumStrides)
+		{
+			for (SIZE_T i = 0; i < lhs.Desc.StreamOutput.NumStrides; i++)
+			{
+				PSO_IF_NOT_EQUAL_RETURN_FALSE(Desc.StreamOutput.pBufferStrides[i])
+			}
+		}
+
+		if (lhs.Desc.InputLayout.pInputElementDescs != rhs.Desc.InputLayout.pInputElementDescs &&
+			lhs.Desc.InputLayout.NumElements)
+		{
+			for (SIZE_T i = 0; i < lhs.Desc.InputLayout.NumElements; i++)
+			{
+				PSO_IF_NOT_EQUAL_RETURN_FALSE(Desc.InputLayout.pInputElementDescs[i].SemanticIndex)
+				PSO_IF_NOT_EQUAL_RETURN_FALSE(Desc.InputLayout.pInputElementDescs[i].Format)
+				PSO_IF_NOT_EQUAL_RETURN_FALSE(Desc.InputLayout.pInputElementDescs[i].InputSlot)
+				PSO_IF_NOT_EQUAL_RETURN_FALSE(Desc.InputLayout.pInputElementDescs[i].AlignedByteOffset)
+				PSO_IF_NOT_EQUAL_RETURN_FALSE(Desc.InputLayout.pInputElementDescs[i].InputSlotClass)
+				PSO_IF_NOT_EQUAL_RETURN_FALSE(Desc.InputLayout.pInputElementDescs[i].InstanceDataStepRate)
+				PSO_IF_STRING_COMPARE_FAILS_RETURN_FALSE(Desc.InputLayout.pInputElementDescs[i].SemanticName)
+			}
+		}
+		return true;
+	}
+};
+template <> struct equality_pipeline_state_desc<FD3D12ComputePipelineStateDesc>
+{
+	bool operator()(const FD3D12ComputePipelineStateDesc& lhs, const FD3D12ComputePipelineStateDesc& rhs)
+	{
+		PSO_IF_NOT_EQUAL_RETURN_FALSE(Desc.CS.BytecodeLength)
+		PSO_IF_NOT_EQUAL_RETURN_FALSE(Desc.Flags)
+		PSO_IF_NOT_EQUAL_RETURN_FALSE(Desc.pRootSignature)
+		PSO_IF_NOT_EQUAL_RETURN_FALSE(Desc.NodeMask)
+		
+		if (lhs.Desc.CS.pShaderBytecode != rhs.Desc.CS.pShaderBytecode &&
+			lhs.Desc.CS.BytecodeLength)
+		{
+			if (FMemory::Memcmp(lhs.Desc.CS.pShaderBytecode, rhs.Desc.CS.pShaderBytecode, lhs.Desc.CS.BytecodeLength) != 0)
+			{
+				return false;
+			}
+		}
+
+		return true;
+	}
+};
+
+// Like a TMap<KeyType, ValueType>
+// Faster lookup performance, but possibly has false negatives
+template<typename KeyType, typename ValueType>
+class FD3D12ConservativeMap
+{
+public:
+	FD3D12ConservativeMap(uint32 Size)
+	{
+		Table.AddUninitialized(Size);
+
+		Reset();
+	}
+
+	void Add(const KeyType& Key, const ValueType& Value)
+	{
+		uint32 Index = GetIndex(Key);
+
+		Entry& Pair = Table[Index];
+
+		Pair.Valid = true;
+		Pair.Key = Key;
+		Pair.Value = Value;
+	}
+
+	ValueType* Find(const KeyType& Key)
+	{
+		uint32 Index = GetIndex(Key);
+
+		Entry& Pair = Table[Index];
+
+		if (Pair.Valid &&
+			(Pair.Key == Key))
+		{
+			return &Pair.Value;
+		}
+		else
+		{
+			return nullptr;
+		}
+	}
+
+	void Reset()
+	{
+		for (int32 i = 0; i < Table.Num(); i++)
+		{
+			Table[i].Valid = false;
+		}
+	}
+
+private:
+	uint32 GetIndex(const KeyType& Key)
+	{
+		uint32 Hash = GetTypeHash(Key);
+
+		return Hash % static_cast<uint32>(Table.Num());
+	}
+
+	struct Entry
+	{
+		bool Valid;
+		KeyType Key;
+		ValueType Value;
+	};
+
+	TArray<Entry> Table;
+};
+
+struct FD3D12SamplerArrayDesc
+{
+	uint32 Count;
+	uint16 SamplerID[16];
+	inline bool operator==(const FD3D12SamplerArrayDesc& rhs) const
+	{
+		check(Count <= _countof(SamplerID));
+		check(rhs.Count <= _countof(rhs.SamplerID));
+
+		if (Count != rhs.Count)
+		{
+			return false;
+		}
+		else
+		{
+			// It is safe to compare pointers, because samplers are kept alive for the lifetime of the RHI
+			return 0 == FMemory::Memcmp(SamplerID, rhs.SamplerID, sizeof(SamplerID[0]) * Count);
+		}
+	}
+};
+uint32 GetTypeHash(const FD3D12SamplerArrayDesc& Key);
+typedef FD3D12ConservativeMap<FD3D12SamplerArrayDesc, D3D12_GPU_DESCRIPTOR_HANDLE> FD3D12SamplerMap;
+
+struct FD3D12SRVArrayDesc
+{
+	uint32 Count;
+	uint64 SRVSequenceNumber[MAX_SRVS];
+
+	inline bool operator==(const FD3D12SRVArrayDesc& rhs) const
+	{
+		check(Count <= _countof(SRVSequenceNumber));
+		check(rhs.Count <= _countof(rhs.SRVSequenceNumber));
+
+		if (Count != rhs.Count)
+		{
+			return false;
+		}
+		else
+		{
+			return 0 == FMemory::Memcmp(SRVSequenceNumber, rhs.SRVSequenceNumber, sizeof(SRVSequenceNumber[0]) * Count);
+		}
+	}
+};
+uint32 GetTypeHash(const FD3D12SRVArrayDesc& Key);
+typedef FD3D12ConservativeMap<FD3D12SRVArrayDesc, D3D12_GPU_DESCRIPTOR_HANDLE> FD3D12SRVMap;
+
+
+class FD3D12ResourceHelper : public FD3D12DeviceChild
+{
+public:
+	HRESULT CreateCommittedResource(const D3D12_RESOURCE_DESC& Desc, const D3D12_HEAP_PROPERTIES& HeapProps, FD3D12Resource **ppOutResource);
+	HRESULT CreateCommittedResource(const D3D12_RESOURCE_DESC& Desc, const D3D12_HEAP_PROPERTIES& HeapProps, const D3D12_RESOURCE_STATES& InitialUsage, FD3D12Resource **ppOutResource);
+	HRESULT CreateDefaultResource(const D3D12_RESOURCE_DESC& Desc, FD3D12Resource **ppOutResource);
+	HRESULT CreateBuffer(D3D12_HEAP_TYPE heapType, uint64 heapSize, FD3D12Resource **ppOutResource);
+
+	FD3D12ResourceHelper(FD3D12Device* InParent);
+
+};
+
+// Vertex Buffer State
+struct FD3D12VertexBufferState
+{
+	TRefCountPtr<FD3D12ResourceLocation> VertexBufferLocation;
+	uint32 Stride;
+	uint32 Offset;
+};
+struct FD3D12ConstantBufferState
+{
+	ID3D12Resource *Resource;
+	uint32 FirstConstant;
+	uint32 NumConstants;
+	FD3D12UniformBuffer* UniformBuffer;
+};
+
+class FD3D12DynamicRHI;
+
+//-----------------------------------------------------------------------------
+//	FD3D12DescriptorCache Class Definition
+//-----------------------------------------------------------------------------
+class FD3D12DescriptorCache : public FD3D12DeviceChild
+{
+protected:
+	FD3D12CommandContext* CmdContext;
+
+public:
+	class OnlineHeap
+	{
+		OnlineHeap(FD3D12DescriptorCache* _Parent) 
+		: Parent(_Parent)
+		, Desc({})
+		, DescriptorSize(0)
+		, NextSlotIndex(0)
+		{ }
+
+		friend class FD3D12DescriptorCache;
+		FD3D12DescriptorCache* Parent;
+
+		void SetParent(FD3D12DescriptorCache* InParent) { Parent = InParent; }
+		FD3D12Device* GetParentDevice() { return Parent->GetParentDevice(); }
+
+		// Desc contains the number of slots and allows for easy recreation
+		D3D12_DESCRIPTOR_HEAP_DESC Desc;
+
+		// Keeping this ptr around is basically just for lifetime management
+		TRefCountPtr<ID3D12DescriptorHeap> Heap;
+
+		// Handles for manipulation of the heap
+		uint32 DescriptorSize;
+		D3D12_CPU_DESCRIPTOR_HANDLE CPUBase;
+		D3D12_GPU_DESCRIPTOR_HANDLE GPUBase;
+
+		// This index indicate where the next set of descriptors should be placed *if* there's room
+		uint32 NextSlotIndex;
+
+		FD3D12CommandListHandle CurrentCommandList;
+
+		struct PoolEntry
+		{
+			TRefCountPtr<ID3D12DescriptorHeap> Heap;
+			FD3D12CLSyncPoint SyncPoint;
+		};
+		TQueue<PoolEntry> ReclaimPool;
+
+	public:
+		// Call this to reserve descriptor heap slots for use by the command list you are currently recording. This will wait if
+		// necessary until slots are free (if they are currently in use by another command list.) If the reservation can be
+		// fulfilled, the index of the first reserved slot is returned (all reserved slots are consecutive.) If not, it will 
+		// throw an exception.
+		bool CanReserveSlots(uint32 NumSlots);
+		void RollOver();
+		uint32 ReserveSlots(uint32 NumSlotsRequested);
+		void SetNextSlot(uint32 NextSlot);
+		D3D12_CPU_DESCRIPTOR_HANDLE GetCPUSlotHandle(uint32 Slot) { return{ CPUBase.ptr + Slot * DescriptorSize }; }
+		D3D12_GPU_DESCRIPTOR_HANDLE GetGPUSlotHandle(uint32 Slot) { return{ GPUBase.ptr + Slot * DescriptorSize }; }
+	};
+
+	OnlineHeap ViewHeap;
+	OnlineHeap SamplerHeap;
+
+	FD3D12DescriptorCache()
+		: ViewHeap(this)
+		, SamplerHeap(this)
+		, CBVAllocator(D3D12_DESCRIPTOR_HEAP_TYPE_CBV_SRV_UAV, 1)
+		, ViewHeapSequenceNumber(1) // starts at 1, because 0 means "is not in any heap"
+		, SamplerMap(271) // Prime numbers for better hashing
+		, SRVMap(271)
+		, FD3D12DeviceChild(nullptr)
+	{
+		CmdContext = nullptr;
+	}
+
+	~FD3D12DescriptorCache(){}
+
+	inline ID3D12DescriptorHeap *GetViewDescriptorHeap()
+	{
+		return ViewHeap.Heap;
+	}
+
+	inline ID3D12DescriptorHeap *GetSamplerDescriptorHeap()
+	{
+		return SamplerHeap.Heap;
+	}
+
+	// Notify the descriptor cache of the current fence value every time you start recording a command list; this allows
+	// us to avoid querying DX12 for that value thousands of times per frame, which can be costly.
+	inline void NotifyCurrentCommandList(FD3D12CommandListHandle CommandListHandle)
+	{
+		ViewHeap.CurrentCommandList = CommandListHandle;
+		SamplerHeap.CurrentCommandList = CommandListHandle;
+	}
+
+	// ------------------------------------------------------
+	// end Descriptor Slot Reservation stuff
+
+	// null views
+	FDescriptorHeapManager CBVAllocator;	// CBV allocator of null CBV
+	CD3DX12_CPU_DESCRIPTOR_HANDLE pNullCBV;
+	TRefCountPtr<FD3D12ShaderResourceView> pNullSRV;
+	TRefCountPtr<FD3D12UnorderedAccessView> pNullUAV;
+	TRefCountPtr<FD3D12RenderTargetView> pNullRTV;
+	TRefCountPtr<FD3D12SamplerState> pDefaultSampler;
+
+	void SetIndexBuffer(FD3D12ResourceLocation* IndexBufferLocation, DXGI_FORMAT Format, uint32 Offset);
+	void SetVertexBuffers(FD3D12VertexBufferState* VertexBuffers, uint32 Count);
+	void SetUAVs(EShaderFrequency ShaderStage,uint32 UAVStartSlot, TRefCountPtr<FD3D12UnorderedAccessView>* UnorderedAccessViewArray, uint32 Count, uint32 &HeapSlot);
+	void SetRenderTargets(FD3D12RenderTargetView **RenderTargetViewArray, uint32 Count, FD3D12DepthStencilView* DepthStencilTarget, bool bDepthIsBoundAsSRV);
+	void SetSamplers(EShaderFrequency ShaderStage, FD3D12SamplerState** Samplers, uint32 Count, uint32 &HeapSlot);
+	void SetSRVs(EShaderFrequency ShaderStage, TRefCountPtr<FD3D12ShaderResourceView> * SRVs, uint32 Count, bool* CurrentShaderResourceViewsIntersectWithDepthRT, uint32 &HeapSlot);
+	void SetConstantBuffers(EShaderFrequency ShaderStage, uint32 Count, uint32 &HeapSlot);
+	void SetStreamOutTargets(FD3D12Resource **Buffers, uint32 Count, const uint32* Offsets);
+
+	void SetConstantBuffer(EShaderFrequency ShaderStage, uint32 SlotIndex, FD3D12UniformBuffer* UniformBuffer);
+	void SetConstantBuffer(EShaderFrequency ShaderStage, uint32 SlotIndex, ID3D12Resource* Resource, uint32 OffsetInBytes, uint32 SizeInBytes);
+	void ClearConstantBuffer(EShaderFrequency ShaderStage, uint32 SlotIndex);
+
+	void HeapRolledOver(D3D12_DESCRIPTOR_HEAP_TYPE Type);
+	void Init(FD3D12Device* InParent, FD3D12CommandContext* InCmdContext);
+	void Clear();
+	void EndFrame();
+
+	struct
+	{
+		TRefCountPtr<ID3D12DescriptorHeap>  Heap;
+
+		TArray<uint64>				CurrentSRVSequenceNumber;
+		CD3DX12_CPU_DESCRIPTOR_HANDLE SRVBaseHandle;
+
+		TArray<uint64>				CurrentUniformBufferSequenceNumber;
+		CD3DX12_CPU_DESCRIPTOR_HANDLE CBVBaseHandle;
+	} OfflineHeap[SF_NumFrequencies];
+
+private:
+	FD3D12SamplerMap SamplerMap;
+	FD3D12SRVMap SRVMap;
+	uint64 ViewHeapSequenceNumber;
+};
+
+class FDiskCacheInterface
+{
+	// Increment this if changes are made to the
+	// disk caches so stale caches get updated correctly
+	static const uint32 mCurrentHeaderVersion = 2;
+	struct FDiskCacheHeader
+	{
+		uint32 mHeaderVersion;
+		uint32 mNumPsos;
+		bool   mUsesAPILibraries;
+	};
+
+private:
+	FString mFileName;
+	byte*   mFileStart;
+	HANDLE  hFile;
+	HANDLE  hMemoryMap;
+	HANDLE  hMapAddress;
+	SIZE_T  mCurrentFileMapSize;
+	SIZE_T  mCurrentOffset;
+	bool    mCacheExists;
+	bool    mInErrorState;
+	FDiskCacheHeader mHeader;
+
+	// There is the potential for the file mapping to grow
+	// in that case all of the pointers will be invalid. Back
+	// some of the pointers we might read again (i.e. shade byte
+	// code for PSO mapping) in persitent system memory.
+	TArray<void*> mBackedMemory;
+
+	static const SIZE_T mFileGrowSize = (1024 * 1024); // 1 megabyte;
+
+	void GrowMapping(SIZE_T size, bool firstrun);
+
+public:
+	bool AppendData(void* pData, size_t size);
+	bool SetPointerAndAdvanceFilePosition(void** pDest, size_t size, bool backWithSystemMemory = false);
+	void Reset();
+	void Init(FString &filename);
+	void Close(uint32 numberOfPSOs);
+	void Flush(uint32 numberOfPSOs);
+	void ClearDiskCache();
+	uint32 GetNumPSOs(){
+		return mHeader.mNumPsos;
+	}
+	bool IsInErrorState(){
+		return mInErrorState;
+	}
+
+	~FDiskCacheInterface()
+	{
+		for (void* memory : mBackedMemory)
+		{
+			if (memory)
+			{
+				FMemory::Free(memory);
+			}
+		}
+	}
+};
+
+static bool bCPUSupportsSSE4;
+
+class FD3D12PipelineStateCache : public FD3D12DeviceChild
+{
+private:
+	enum PSO_CACHE_TYPE {
+		PSO_CACHE_GRAPHICS,
+		PSO_CACHE_COMPUTE,
+		NUM_PSO_CACHE_TYPES
+	};
+
+	template <typename TDesc, typename TValue>
+	struct TStateCacheKeyFuncs : BaseKeyFuncs<TPair<TDesc, TValue>, TDesc, false>
+	{
+		typedef typename TTypeTraits<TDesc>::ConstPointerType KeyInitType;
+		typedef const typename TPairInitializer<typename TTypeTraits<TDesc>::ConstInitType, typename TTypeTraits<TValue>::ConstInitType>& ElementInitType;
+
+		static FORCEINLINE KeyInitType GetSetKey(ElementInitType Element)
+		{
+			return Element.Key;
+		}
+		static FORCEINLINE bool Matches(KeyInitType A, KeyInitType B)
+		{
+			equality_pipeline_state_desc<TDesc> equal;
+			return equal(A, B);
+		}
+		static FORCEINLINE uint32 GetKeyHash(KeyInitType Key)
+		{
+			return Key.CombinedHash;
+		}
+	};
+
+	template <typename TDesc, typename TValue = TRefCountPtr<ID3D12PipelineState>>
+	using TPipelineCache = TMap<TDesc, TValue, FDefaultSetAllocator, TStateCacheKeyFuncs<TDesc, TValue>>;
+
+	TPipelineCache<FD3D12HighLevelGraphicsPipelineStateDesc, TPair<ID3D12PipelineState*, uint64>> HighLevelGraphicsPipelineStateCache;
+	TPipelineCache<FD3D12LowLevelGraphicsPipelineStateDesc> LowLevelGraphicsPipelineStateCache;
+	TPipelineCache<FD3D12ComputePipelineStateDesc> ComputePipelineStateCache;
+
+	FCriticalSection CS;
+	FDiskCacheInterface DiskCaches[NUM_PSO_CACHE_TYPES];
+
+	ID3D12PipelineState* Add(FD3D12LowLevelGraphicsPipelineStateDesc &graphicsPSODesc, bool insertIntoDiskCache = false);
+	ID3D12PipelineState* Add(FD3D12ComputePipelineStateDesc &computePSODesc, bool insertIntoDiskCache = false);
+
+	ID3D12PipelineState* FindGraphicsLowLevel(FD3D12LowLevelGraphicsPipelineStateDesc &graphicsPSODesc, bool insertIntoDiskCache = false);
+
+#if UE_BUILD_DEBUG
+	uint64 GraphicsCacheRequestCount = 0;
+	uint64 HighLevelCacheFulfillCount = 0;
+	uint64 HighLevelCacheStaleCount = 0;
+	uint64 HighLevelCacheMissCount = 0;
+#endif
+
+public:
+	void RebuildFromDiskCache(ID3D12RootSignature* pGraphicsRootSig, ID3D12RootSignature* pComputeRootSignature);
+
+	ID3D12PipelineState* FindGraphics(FD3D12HighLevelGraphicsPipelineStateDesc &graphicsPSODesc, ID3D12RootSignature* rootSignature, bool insertIntoDiskCache = false);
+	ID3D12PipelineState* FindCompute(FD3D12ComputePipelineStateDesc &computePSODesc, bool insertIntoDiskCache = false);
+
+	void Close();
+
+	void Init(FString &GraphicsCacheFilename, FString &ComputeCacheFilename);
+
+	static SIZE_T HashPSODesc(const FD3D12HighLevelGraphicsPipelineStateDesc &psoDesc);
+	static SIZE_T HashPSODesc(const FD3D12LowLevelGraphicsPipelineStateDesc &psoDesc);
+	static SIZE_T HashPSODesc(const FD3D12ComputePipelineStateDesc &psoDesc);
+
+	static inline SIZE_T HashData(void* data, SIZE_T numBytes);
+
+	FD3D12PipelineStateCache(FD3D12Device* InParent);
+	~FD3D12PipelineStateCache();
+
+	static const bool bUseAPILibaries = false;
+};
+
+class FD3D12BitArray
+{
+public:
+
+	FD3D12BitArray() : Array (0) {}
+	FD3D12BitArray (uint32 InArray) : Array (InArray) {}
+
+	void Clear()
+	{
+		Array = 0;
+	}
+
+	void SetIndex (uint32 Index)
+	{
+		Array |= (1 << Index);
+	}
+
+	void ClearIndex (uint32 Index)
+	{
+		Array &= ~(1 << Index);
+	}
+
+	void Set (uint32 Bits)
+	{
+		for (uint32 i = 0; i < Bits; ++i)
+			SetIndex (i);
+	}
+
+	bool AnySet()
+	{
+		return (Array != 0);
+	}
+
+	uint32 LastSet()
+	{
+		return 32 - FMath::CountLeadingZeros (Array);
+	}
+
+	uint32 GetValue()
+	{
+		return Array;
+	}
+
+private:
+
+	uint32 Array;
+};
+
+class FD3D12StateArray
+{
+public:
+
+	void SetIndex (uint32 Index)
+	{
+		SetIndices.SetIndex     (Index);
+		ChangedIndices.SetIndex (Index);
+	}
+
+	void ClearIndex (uint32 Index)
+	{
+		SetIndices.ClearIndex   (Index);
+		ChangedIndices.SetIndex (Index);
+	}
+
+	uint32 LastSet()
+	{
+		FD3D12BitArray Merged (SetIndices.GetValue() | ChangedIndices.GetValue());
+
+		return Merged.LastSet();
+	}
+
+	void Clear()
+	{
+		SetIndices.Clear();
+		ChangedIndices.Clear();
+	}
+
+	void ClearChanges()
+	{
+		ChangedIndices.Clear();
+	}
+
+private:
+
+	FD3D12BitArray SetIndices;
+	FD3D12BitArray ChangedIndices;
+};
+
+//-----------------------------------------------------------------------------
+//	FD3D12StateCache Class Definition
+//-----------------------------------------------------------------------------
+class FD3D12StateCacheBase : public FD3D12DeviceChild
+{
+	friend class FD3D12DynamicRHI;
+
+public:
+	enum ESRV_Type
+	{
+		SRV_Unknown,
+		SRV_Dynamic,
+		SRV_Static,
+	};
+protected:
+	FD3D12CommandContext* CmdContext;
+
+	bool bNeedSetVB;
+	bool bNeedSetIB;
+	bool bNeedSetUAVs;
+	bool bNeedSetRTs;
+	bool bNeedSetSOs;
+	bool bNeedSetSamplersPerShaderStage[SF_NumFrequencies];
+	bool bNeedSetSamplers;
+	bool bNeedSetSRVsPerShaderStage[SF_NumFrequencies];
+	bool bSRVSCleared;
+	bool bNeedSetSRVs;
+	bool bNeedSetConstantBuffersPerShaderStage[SF_NumFrequencies];
+	bool bNeedSetConstantBuffers;
+	bool bNeedSetViewports;
+	bool bNeedSetScissorRects;
+	bool bNeedSetPrimitiveTopology;
+	bool bNeedSetBlendFactor;
+	bool bNeedSetStencilRef;
+	bool bAutoFlushComputeShaderCache;
+
+	struct
+	{
+		struct
+		{
+			// Cache
+			ID3D12PipelineState* CurrentPipelineStateObject;
+			bool bNeedRebuildPSO;
+
+			// Full high level PSO desc
+			FD3D12HighLevelGraphicsPipelineStateDesc HighLevelDesc;
+
+			// Depth Stencil State Cache
+			uint32 CurrentReferenceStencil;
+
+			// Blend State Cache
+			float CurrentBlendFactor[4];
+
+			// Viewport
+			uint32	CurrentNumberOfViewports;
+			D3D12_VIEWPORT CurrentViewport[D3D12_VIEWPORT_AND_SCISSORRECT_OBJECT_COUNT_PER_PIPELINE];
+
+			// Vertex Buffer State
+			FD3D12VertexBufferState CurrentVertexBuffers[D3D12_IA_VERTEX_INPUT_RESOURCE_SLOT_COUNT];
+
+			// SO
+			uint32			CurrentNumberOfStreamOutTargets;
+			FD3D12Resource* CurrentStreamOutTargets[D3D12_SO_STREAM_COUNT];
+			uint32			CurrentSOOffsets[D3D12_SO_STREAM_COUNT];
+
+			// Index Buffer State
+			TRefCountPtr<FD3D12ResourceLocation> CurrentIndexBufferLocation;
+			DXGI_FORMAT CurrentIndexFormat;
+			uint32 CurrentIndexOffset;
+
+			// Primitive Topology State
+			D3D_PRIMITIVE_TOPOLOGY CurrentPrimitiveTopology;
+
+			// Input Layout State
+			D3D12_RECT CurrentScissorRects[D3D12_VIEWPORT_AND_SCISSORRECT_OBJECT_COUNT_PER_PIPELINE];
+			D3D12_RECT CurrentViewportScissorRects[D3D12_VIEWPORT_AND_SCISSORRECT_OBJECT_COUNT_PER_PIPELINE];
+			uint32 CurrentNumberOfScissorRects;
+
+			FD3D12RenderTargetView* RenderTargetArray[D3D12_SIMULTANEOUS_RENDER_TARGET_COUNT];
+			
+			FD3D12DepthStencilView* CurrentDepthStencilTarget;
+		} Graphics;
+
+		struct
+		{
+			// Cache
+			ID3D12PipelineState* CurrentPipelineStateObject;
+			bool bNeedRebuildPSO;
+
+			// Compute
+			FD3D12ComputeShader* CurrentComputeShader;
+		} Compute;
+
+		struct
+		{
+			// UAVs
+			TRefCountPtr<FD3D12UnorderedAccessView> UnorderedAccessViewArray[D3D12_PS_CS_UAV_REGISTER_COUNT];
+			EShaderFrequency CurrentUAVStage;
+			uint32	CurrentUAVStartSlot;
+			uint32	CurrentNumberOfSimultaneousUAVs;
+
+			// Shader Resource Views Cache
+			TRefCountPtr<FD3D12ShaderResourceView> CurrentShaderResourceViews[SF_NumFrequencies][D3D12_COMMONSHADER_INPUT_RESOURCE_SLOT_COUNT];
+			bool CurrentShaderResourceViewsIntersectWithDepthRT[SF_NumFrequencies][MAX_SRVS];
+			uint32 ShaderResourceViewsIntersectWithDepthCount;
+			int32 MaxBoundShaderResourcesIndex[SF_NumFrequencies];
+
+			// Sampler State
+			FD3D12SamplerState* CurrentSamplerStates[SF_NumFrequencies][D3D12_COMMONSHADER_SAMPLER_SLOT_COUNT];
+
+			// PSO
+			ID3D12PipelineState* CurrentPipelineStateObject;
+			bool bNeedSetPSO;
+
+			uint32 CurrentShaderSamplerCounts[SF_NumFrequencies];
+			uint32 CurrentShaderSRVCounts[SF_NumFrequencies];
+			uint32 CurrentShaderCBCounts[SF_NumFrequencies];
+			uint32 CurrentShaderUAVCounts[SF_NumFrequencies];
+		} Common;
+	} PipelineState;
+
+	FD3D12DescriptorCache DescriptorCache;
+	bool bAlwaysSetIndexBuffers;
+
+	D3D12_STATE_CACHE_INLINE void InternalSetIndexBuffer(FD3D12ResourceLocation *IndexBufferLocation, DXGI_FORMAT Format, uint32 Offset);
+
+	typedef void(*TSetSRVAlternate)(FD3D12StateCacheBase* StateCache, FD3D12ShaderResourceView* SRV, uint32 ResourceIndex, ESRV_Type SrvType);
+	template <EShaderFrequency ShaderFrequency>
+	void D3D12_STATE_CACHE_INLINE InternalSetShaderResourceView(FD3D12ShaderResourceView*& SRV, uint32 ResourceIndex, ESRV_Type SrvType, TSetSRVAlternate AlternatePathFunction);
+
+	typedef void(*TSetStreamSourceAlternate)(FD3D12StateCacheBase* StateCache, FD3D12Resource* VertexBuffer, uint32 StreamIndex, uint32 Stride, uint32 Offset);
+	D3D12_STATE_CACHE_INLINE void InternalSetStreamSource(FD3D12ResourceLocation* VertexBufferLocation, uint32 StreamIndex, uint32 Stride, uint32 Offset, TSetStreamSourceAlternate AlternatePathFunction);
+
+	typedef void(*TSetSamplerStateAlternate)(FD3D12StateCacheBase* StateCache, FD3D12SamplerState* SamplerState, uint32 SamplerIndex);
+	template <EShaderFrequency ShaderFrequency>
+	D3D12_STATE_CACHE_INLINE void InternalSetSamplerState(FD3D12SamplerState* SamplerState, uint32 SamplerIndex, TSetSamplerStateAlternate AlternatePathFunction)
+	{
+		check(SamplerIndex < ARRAYSIZE(PipelineState.Common.CurrentSamplerStates[ShaderFrequency]));;
+		if ((PipelineState.Common.CurrentSamplerStates[ShaderFrequency][SamplerIndex] != SamplerState) || GD3D12SkipStateCaching)
+		{
+			PipelineState.Common.CurrentSamplerStates[ShaderFrequency][SamplerIndex] = SamplerState;
+			bNeedSetSamplersPerShaderStage[ShaderFrequency] = true;
+			bNeedSetSamplers = true;
+			if (AlternatePathFunction != nullptr)
+				(*AlternatePathFunction)(this, SamplerState, SamplerIndex);
+		}
+	}
+
+	// Shorthand for typing/reading convenience
+	D3D12_STATE_CACHE_INLINE FD3D12BoundShaderState* BSS()
+	{
+		return PipelineState.Graphics.HighLevelDesc.BoundShaderState;
+	}
+
+	template <typename TShader> struct StateCacheShaderTraits;
+#define DECLARE_SHADER_TRAITS(Name) \
+	template <> struct StateCacheShaderTraits<FD3D12##Name##Shader> \
+	{ \
+		static const EShaderFrequency Frequency = SF_##Name; \
+		static FD3D12##Name##Shader* GetShader(FD3D12BoundShaderState* BSS) { return BSS ? BSS->Get##Name##Shader() : nullptr; } \
+	}
+	DECLARE_SHADER_TRAITS(Vertex);
+	DECLARE_SHADER_TRAITS(Pixel);
+	DECLARE_SHADER_TRAITS(Domain);
+	DECLARE_SHADER_TRAITS(Hull);
+	DECLARE_SHADER_TRAITS(Geometry);
+#undef DECLARE_SHADER_TRAITS
+
+	template <typename TShader> D3D12_STATE_CACHE_INLINE void SetShader(TShader* Shader)
+	{
+		typedef StateCacheShaderTraits<TShader> Traits;
+		TShader* OldShader = Traits::GetShader(BSS());
+		if (OldShader != Shader)
+		{
+			PipelineState.Common.CurrentShaderSamplerCounts[Traits::Frequency] = (Shader) ? Shader->SamplerCount : 0;
+			PipelineState.Common.CurrentShaderSRVCounts[Traits::Frequency]     = (Shader) ? Shader->SRVCount     : 0;
+			PipelineState.Common.CurrentShaderCBCounts[Traits::Frequency]      = (Shader) ? Shader->CBCount      : 0;
+			PipelineState.Common.CurrentShaderUAVCounts[Traits::Frequency]     = (Shader) ? Shader->UAVCount     : 0;
+		}
+	}
+
+	template <typename TShader> D3D12_STATE_CACHE_INLINE void GetShader(TShader** Shader)
+	{
+		*Shader = StateCacheShaderTraits<TShader>::GetShader(BSS());
+	}
+
+public:
+
+	void InheritState( const FD3D12StateCacheBase& AncestralCache )
+	{
+		FMemory::Memcpy(&PipelineState, &AncestralCache.PipelineState, sizeof(PipelineState));
+		RestoreState();
+	}
+
+	FD3D12DescriptorCache* GetDescriptorCache() { return &DescriptorCache; }
+
+	ID3D12PipelineState* GetPipelineStateObject()
+	{
+		return PipelineState.Common.CurrentPipelineStateObject;
+	}
+
+	void ClearSamplers()
+	{
+		ZeroMemory(PipelineState.Common.CurrentSamplerStates, sizeof(PipelineState.Common.CurrentSamplerStates));
+		for (uint32 ShaderFrequency = 0; ShaderFrequency < _countof(bNeedSetSamplersPerShaderStage); ShaderFrequency++)
+		{
+			bNeedSetSamplersPerShaderStage[ShaderFrequency] = true;
+		}
+		bNeedSetSamplers = true;
+	}
+
+	void ClearSRVs();
+
+	template <EShaderFrequency ShaderFrequency>
+	void ClearShaderResourceViews(FD3D12ResourceLocation*& ResourceLocation);
+
+	template <EShaderFrequency ShaderFrequency>
+	D3D12_STATE_CACHE_INLINE void SetShaderResourceView(FD3D12ShaderResourceView* SRV, uint32 ResourceIndex, ESRV_Type SrvType = SRV_Unknown)
+	{
+		InternalSetShaderResourceView<ShaderFrequency>(SRV, ResourceIndex, SrvType, nullptr);
+	}
+
+	template <EShaderFrequency ShaderFrequency>
+	D3D12_STATE_CACHE_INLINE void GetShaderResourceViews(uint32 StartResourceIndex, uint32& NumResources, FD3D12ShaderResourceView** SRV)
+	{
+		{
+			uint32 NumLoops = D3D12_COMMONSHADER_INPUT_RESOURCE_SLOT_COUNT - StartResourceIndex;
+			NumResources = 0;
+			for (uint32 ResourceLoop = 0; ResourceLoop < NumLoops; ResourceLoop++)
+			{
+				SRV[ResourceLoop] = PipelineState.Common.CurrentShaderResourceViews[ShaderFrequency][ResourceLoop + StartResourceIndex];
+				if (SRV[ResourceLoop])
+				{
+					SRV[ResourceLoop]->AddRef();
+					NumResources = ResourceLoop;
+				}
+			}
+		}
+	}
+
+	void UpdateViewportScissorRects()
+	{
+		for (uint32 i = 0; i < PipelineState.Graphics.CurrentNumberOfScissorRects; ++i)
+		{
+			D3D12_VIEWPORT& Viewport		    = PipelineState.Graphics.CurrentViewport[FMath::Min (i, PipelineState.Graphics.CurrentNumberOfViewports)];
+			D3D12_RECT&     ScissorRect         = PipelineState.Graphics.CurrentScissorRects[i];
+			D3D12_RECT&     ViewportScissorRect = PipelineState.Graphics.CurrentViewportScissorRects[i];
+
+			ViewportScissorRect.top    = FMath::Max(ScissorRect.top,    (LONG)Viewport.TopLeftY);
+			ViewportScissorRect.left   = FMath::Max(ScissorRect.left,   (LONG)Viewport.TopLeftX);
+			ViewportScissorRect.bottom = FMath::Min(ScissorRect.bottom, (LONG)Viewport.TopLeftY + (LONG)Viewport.Height);
+			ViewportScissorRect.right  = FMath::Min(ScissorRect.right,  (LONG)Viewport.TopLeftX + (LONG)Viewport.Width);
+		}
+
+		bNeedSetScissorRects = true;
+	}
+
+	D3D12_STATE_CACHE_INLINE void SetScissorRects(uint32 Count, D3D12_RECT* ScissorRects);
+	D3D12_STATE_CACHE_INLINE void SetScissorRect(D3D12_RECT ScissorRect);
+
+	D3D12_STATE_CACHE_INLINE void GetScissorRect(D3D12_RECT *ScissorRect) const
+	{
+		check(ScissorRect);
+		FMemory::Memcpy(ScissorRect, &PipelineState.Graphics.CurrentScissorRects, sizeof(D3D12_RECT));
+	}
+
+
+	D3D12_STATE_CACHE_INLINE void SetViewport(D3D12_VIEWPORT Viewport);
+	D3D12_STATE_CACHE_INLINE void SetViewports(uint32 Count, D3D12_VIEWPORT* Viewports);
+
+	D3D12_STATE_CACHE_INLINE void GetViewport(D3D12_VIEWPORT *Viewport) const
+	{
+		check(Viewport);
+		FMemory::Memcpy(Viewport, &PipelineState.Graphics.CurrentViewport, sizeof(D3D12_VIEWPORT));
+	}
+
+	D3D12_STATE_CACHE_INLINE void GetViewports(uint32* Count, D3D12_VIEWPORT *Viewports) const
+	{
+		check(*Count);
+		if (Viewports) //NULL is legal if you just want count
+		{
+			//as per d3d spec
+			int32 StorageSizeCount = (int32)(*Count);
+			int32 CopyCount = FMath::Min(FMath::Min(StorageSizeCount, (int32)PipelineState.Graphics.CurrentNumberOfViewports), D3D12_VIEWPORT_AND_SCISSORRECT_OBJECT_COUNT_PER_PIPELINE);
+			if (CopyCount > 0)
+			{
+				FMemory::Memcpy(Viewports, &PipelineState.Graphics.CurrentViewport[0], sizeof(D3D12_VIEWPORT) * CopyCount);
+			}
+			//remaining viewports in supplied array must be set to zero
+			if (StorageSizeCount > CopyCount)
+			{
+				FMemory::Memset(&Viewports[CopyCount], 0, sizeof(D3D12_VIEWPORT) * (StorageSizeCount - CopyCount));
+			}
+		}
+		*Count = PipelineState.Graphics.CurrentNumberOfViewports;
+	}
+
+	template <EShaderFrequency ShaderFrequency>
+	D3D12_STATE_CACHE_INLINE void SetSamplerState(FD3D12SamplerState* SamplerState, uint32 SamplerIndex)
+	{
+		InternalSetSamplerState<ShaderFrequency>(SamplerState, SamplerIndex, nullptr);
+	}
+
+	template <EShaderFrequency ShaderFrequency>
+	D3D12_STATE_CACHE_INLINE void GetSamplerState(uint32 StartSamplerIndex, uint32 NumSamplerIndexes, FD3D12SamplerState** SamplerStates)
+	{
+		check(StartSamplerIndex + NumSamplerIndexes <= D3D12_COMMONSHADER_INPUT_RESOURCE_SLOT_COUNT);
+		for (uint32 StateLoop = 0; StateLoop < NumSamplerIndexes; StateLoop++)
+		{
+			SamplerStates[StateLoop] = CurrentShaderResourceViews[ShaderFrequency][StateLoop + StartSamplerIndex];
+			if (SamplerStates[StateLoop])
+			{
+				SamplerStates[StateLoop]->AddRef();
+			}
+		}
+	}
+
+	template <EShaderFrequency ShaderFrequency>
+	void D3D12_STATE_CACHE_INLINE SetConstantBuffer(uint32 SlotIndex, FD3D12ResourceLocation* ResourceLocation, FD3D12UniformBuffer* UniformBuffer)
+	{
+		// Update the high water mark
+		check(SlotIndex < MAX_CBS);
+
+		if (UniformBuffer && (UniformBuffer->OfflineDescriptorHandle.ptr != 0))
+		{
+			check(!ResourceLocation);
+			DescriptorCache.SetConstantBuffer(ShaderFrequency, SlotIndex, UniformBuffer);
+		}
+		else if (ResourceLocation)
+		{
+			check(!UniformBuffer);
+			const uint32 AlignedSize = Align(ResourceLocation->GetEffectiveBufferSize(), D3D12_CONSTANT_BUFFER_DATA_PLACEMENT_ALIGNMENT);
+			DescriptorCache.SetConstantBuffer(ShaderFrequency, SlotIndex, ResourceLocation->GetResource()->GetResource(), ResourceLocation->GetOffset(), AlignedSize);
+		}
+		else
+		{
+			DescriptorCache.ClearConstantBuffer(ShaderFrequency, SlotIndex);
+		}
+
+		bNeedSetConstantBuffersPerShaderStage[ShaderFrequency] = true;
+		bNeedSetConstantBuffers = true;
+	}
+
+	D3D12_STATE_CACHE_INLINE void SetRasterizerState(D3D12_RASTERIZER_DESC* State)
+	{
+		if (PipelineState.Graphics.HighLevelDesc.RasterizerState != State || GD3D12SkipStateCaching)
+		{
+			PipelineState.Graphics.HighLevelDesc.RasterizerState = State;
+			PipelineState.Graphics.bNeedRebuildPSO = true;
+		}
+	}
+
+	D3D12_STATE_CACHE_INLINE void GetRasterizerState(D3D12_RASTERIZER_DESC** RasterizerState)
+	{
+		*RasterizerState = PipelineState.Graphics.HighLevelDesc.RasterizerState;
+	}
+
+	D3D12_STATE_CACHE_INLINE void SetBlendState(D3D12_BLEND_DESC* State, const float BlendFactor[4], uint32 SampleMask);
+
+	D3D12_STATE_CACHE_INLINE void GetBlendState(D3D12_BLEND_DESC** BlendState, float BlendFactor[4], uint32* SampleMask)
+	{
+		*BlendState = PipelineState.Graphics.HighLevelDesc.BlendState;
+		*SampleMask = PipelineState.Graphics.HighLevelDesc.SampleMask;
+		FMemory::Memcpy(BlendFactor, PipelineState.Graphics.CurrentBlendFactor, sizeof(PipelineState.Graphics.CurrentBlendFactor));
+	}
+
+	D3D12_STATE_CACHE_INLINE void SetDepthStencilState(D3D12_DEPTH_STENCIL_DESC* State, uint32 RefStencil);
+
+	D3D12_STATE_CACHE_INLINE void GetDepthStencilState(D3D12_DEPTH_STENCIL_DESC** DepthStencilState, uint32* StencilRef)
+	{
+		*DepthStencilState = PipelineState.Graphics.HighLevelDesc.DepthStencilState;
+		*StencilRef = PipelineState.Graphics.CurrentReferenceStencil;
+	}
+
+	D3D12_STATE_CACHE_INLINE void GetVertexShader(FD3D12VertexShader** Shader)
+	{
+		GetShader(Shader);
+	}
+
+	D3D12_STATE_CACHE_INLINE void GetHullShader(FD3D12HullShader** Shader)
+	{
+		GetShader(Shader);
+	}
+
+	D3D12_STATE_CACHE_INLINE void GetDomainShader(FD3D12DomainShader** Shader)
+	{
+		GetShader(Shader);
+	}
+
+	D3D12_STATE_CACHE_INLINE void GetGeometryShader(FD3D12GeometryShader** Shader)
+	{
+		GetShader(Shader);
+	}
+
+	D3D12_STATE_CACHE_INLINE void GetPixelShader(FD3D12PixelShader** Shader)
+	{
+		GetShader(Shader);
+	}
+
+	D3D12_STATE_CACHE_INLINE void SetBoundShaderState(FD3D12BoundShaderState* BoundShaderState)
+	{
+		if (BoundShaderState)
+		{
+			SetShader(BoundShaderState->GetVertexShader());
+			SetShader(BoundShaderState->GetPixelShader());
+			SetShader(BoundShaderState->GetDomainShader());
+			SetShader(BoundShaderState->GetHullShader());
+			SetShader(BoundShaderState->GetGeometryShader());
+		}
+		else
+		{
+			SetShader<FD3D12VertexShader>(nullptr);
+			SetShader<FD3D12PixelShader>(nullptr);
+			SetShader<FD3D12HullShader>(nullptr);
+			SetShader<FD3D12DomainShader>(nullptr);
+			SetShader<FD3D12GeometryShader>(nullptr);
+		}
+
+		if (PipelineState.Graphics.HighLevelDesc.BoundShaderState != BoundShaderState)
+		{
+			PipelineState.Graphics.HighLevelDesc.BoundShaderState = BoundShaderState;
+			PipelineState.Graphics.bNeedRebuildPSO = true;
+		}
+	}
+
+	D3D12_STATE_CACHE_INLINE void GetBoundShaderState(FD3D12BoundShaderState** BoundShaderState)
+	{
+		*BoundShaderState = PipelineState.Graphics.HighLevelDesc.BoundShaderState;
+	}
+
+	D3D12_STATE_CACHE_INLINE void SetComputeShader(FD3D12ComputeShader* Shader)
+	{
+		if (PipelineState.Compute.CurrentComputeShader != Shader)
+		{
+			PipelineState.Compute.CurrentComputeShader                  = Shader;
+			PipelineState.Compute.bNeedRebuildPSO = true;
+			PipelineState.Common.CurrentShaderSamplerCounts[SF_Compute] = (Shader) ? Shader->SamplerCount : 0;
+			PipelineState.Common.CurrentShaderSRVCounts[SF_Compute]     = (Shader) ? Shader->SRVCount     : 0;
+			PipelineState.Common.CurrentShaderCBCounts[SF_Compute]      = (Shader) ? Shader->CBCount      : 0;
+			PipelineState.Common.CurrentShaderUAVCounts[SF_Compute]     = (Shader) ? Shader->UAVCount     : 0;
+		}
+	}
+
+	D3D12_STATE_CACHE_INLINE void GetComputeShader(FD3D12ComputeShader** ComputeShader)
+	{
+		*ComputeShader = PipelineState.Compute.CurrentComputeShader;
+	}
+
+	D3D12_STATE_CACHE_INLINE void GetInputLayout(D3D12_INPUT_LAYOUT_DESC* InputLayout)
+	{
+		*InputLayout = PipelineState.Graphics.HighLevelDesc.BoundShaderState->InputLayout;
+	}
+
+	D3D12_STATE_CACHE_INLINE void SetStreamSource(FD3D12ResourceLocation* VertexBufferLocation, uint32 StreamIndex, uint32 Stride, uint32 Offset)
+	{
+		InternalSetStreamSource(VertexBufferLocation, StreamIndex, Stride, Offset, nullptr);
+	}
+
+	D3D12_STATE_CACHE_INLINE void GetStreamSources(uint32 StartStreamIndex, uint32 NumStreams, FD3D12ResourceLocation** VertexBuffers, uint32* Strides, uint32* Offsets)
+	{
+		check(StartStreamIndex + NumStreams <= D3D12_IA_VERTEX_INPUT_RESOURCE_SLOT_COUNT);
+		for (uint32 StreamLoop = 0; StreamLoop < NumStreams; StreamLoop++)
+		{
+			FD3D12VertexBufferState& Slot = PipelineState.Graphics.CurrentVertexBuffers[StreamLoop + StartStreamIndex];
+			VertexBuffers[StreamLoop] = Slot.VertexBufferLocation;
+			if (Strides)
+			{
+				Strides[StreamLoop] = Slot.Stride;
+			}
+
+			if (Offsets)
+			{
+				Offsets[StreamLoop] = Slot.Offset + Slot.VertexBufferLocation->GetOffset();
+			}
+
+			if (Slot.VertexBufferLocation)
+			{
+				Slot.VertexBufferLocation->AddRef();
+			}
+		}
+	}
+
+	D3D12_STATE_CACHE_INLINE bool IsShaderResource(const FD3D12ResourceLocation* VertexBufferLocation) const
+	{
+		for (int i = 0; i < SF_NumFrequencies; i++)
+		{
+			for (uint32 j = 0; j < D3D12_COMMONSHADER_INPUT_RESOURCE_SLOT_COUNT; ++j)
+			{
+				if (PipelineState.Common.CurrentShaderResourceViews[i][j] && PipelineState.Common.CurrentShaderResourceViews[i][j]->GetResourceLocation())
+				{
+					if (PipelineState.Common.CurrentShaderResourceViews[i][j]->GetResourceLocation() == VertexBufferLocation)
+					{
+						return true;
+					}
+				}
+			}
+		}
+
+		return false;
+	}
+
+	D3D12_STATE_CACHE_INLINE bool IsStreamSource(const FD3D12ResourceLocation* VertexBufferLocation) const
+	{
+		for (uint32 index = 0; index < D3D12_IA_VERTEX_INPUT_RESOURCE_SLOT_COUNT; ++index)
+		{
+			if (PipelineState.Graphics.CurrentVertexBuffers[index].VertexBufferLocation.GetReference() == VertexBufferLocation)
+			{
+				return true;
+			}
+		}
+
+		return false;
+	}
+
+public:
+
+	D3D12_STATE_CACHE_INLINE void SetIndexBuffer(FD3D12ResourceLocation* IndexBufferLocation, DXGI_FORMAT Format, uint32 Offset)
+	{
+		InternalSetIndexBuffer(IndexBufferLocation, Format, Offset);
+	}
+
+	D3D12_STATE_CACHE_INLINE const FD3D12ResourceLocation *GetIndexBufferLocation() const
+	{
+		return PipelineState.Graphics.CurrentIndexBufferLocation;
+	}
+
+	D3D12_STATE_CACHE_INLINE bool IsIndexBuffer(const FD3D12ResourceLocation *ResourceLocation) const
+	{
+		return PipelineState.Graphics.CurrentIndexBufferLocation &&
+			PipelineState.Graphics.CurrentIndexBufferLocation->GetResource() == ResourceLocation->GetResource()
+			&& PipelineState.Graphics.CurrentIndexBufferLocation->GetOffset() == ResourceLocation->GetOffset();
+	}
+
+	D3D12_STATE_CACHE_INLINE void SetPrimitiveTopology(D3D12_PRIMITIVE_TOPOLOGY PrimitiveTopology);
+
+
+	D3D12_STATE_CACHE_INLINE void GetPrimitiveTopology(D3D12_PRIMITIVE_TOPOLOGY* PrimitiveTopology)
+	{
+		*PrimitiveTopology = PipelineState.Graphics.CurrentPrimitiveTopology;
+	}
+
+	FD3D12StateCacheBase()
+	{
+		// Make all the bound shader resources -1
+		for (auto& Index : PipelineState.Common.MaxBoundShaderResourcesIndex)
+		{
+			Index = INDEX_NONE;
+		}
+	}
+
+	void Init(FD3D12Device* InParent, FD3D12CommandContext* InCmdContext, const FD3D12StateCacheBase* AncestralState, bool bInAlwaysSetIndexBuffers = false);
+
+	~FD3D12StateCacheBase()
+	{
+	}
+	void ApplyState(bool IsCompute = false);
+	void D3D12_STATE_CACHE_INLINE RestoreState();
+	bool VerifyResourceStates(const bool IsCompute);
+
+	template <class ViewT>
+	bool VerifyViewState(ID3D12DebugCommandList* pDebugCommandList, FD3D12View<ViewT> *pView, D3D12_RESOURCE_STATES State);
+
+	void D3D12_STATE_CACHE_INLINE SetRenderTargets(uint32 NumSimultaneousRenderTargets, FD3D12RenderTargetView** RTArray, FD3D12DepthStencilView* DSTarget);
+	D3D12_STATE_CACHE_INLINE void GetRenderTargets(FD3D12RenderTargetView **RTArray, uint32* NumSimultaneousRTs, FD3D12DepthStencilView** DepthStencilTarget)
+	{
+		if (RTArray) //NULL is legal
+		{
+			FMemory::Memcpy(RTArray, PipelineState.Graphics.RenderTargetArray, sizeof(FD3D12RenderTargetView*)* D3D12_SIMULTANEOUS_RENDER_TARGET_COUNT);
+			*NumSimultaneousRTs = PipelineState.Graphics.HighLevelDesc.NumRenderTargets;
+		}
+
+		if (DepthStencilTarget)
+		{
+			*DepthStencilTarget = PipelineState.Graphics.CurrentDepthStencilTarget;
+		}
+	}
+
+	void D3D12_STATE_CACHE_INLINE SetStreamOutTargets(uint32 NumSimultaneousStreamOutTargets, FD3D12Resource** SOArray, const uint32* SOOffsets);
+
+	void D3D12_STATE_CACHE_INLINE SetUAVs(EShaderFrequency ShaderStage, uint32 UAVStartSlot, uint32 NumSimultaneousUAVs, FD3D12UnorderedAccessView** UAVArray, uint32 *UAVInitialCountArray);
+
+	D3D12_STATE_CACHE_INLINE void GetUAVs(FD3D12UnorderedAccessView** UAVArray, uint32* UAVStartSlot, uint32* NumSimultaneousUAVs)
+	{
+		if (UAVArray)
+		{
+			FMemory::Memcpy(UAVArray, PipelineState.Common.UnorderedAccessViewArray, sizeof(FD3D12UnorderedAccessView*) * PipelineState.Common.CurrentNumberOfSimultaneousUAVs);
+			*UAVStartSlot = PipelineState.Common.CurrentUAVStartSlot;
+			*NumSimultaneousUAVs = PipelineState.Common.CurrentNumberOfSimultaneousUAVs;
+		}
+	}
+
+	D3D12_STATE_CACHE_INLINE FD3D12UnorderedAccessView* GetUAV(const uint32 &UAVIndex)
+	{
+		check(UAVIndex < _countof(PipelineState.Common.UnorderedAccessViewArray));
+		return PipelineState.Common.UnorderedAccessViewArray[UAVIndex];
+	}
+
+	D3D12_STATE_CACHE_INLINE void AutoFlushComputeShaderCache(bool bEnable)
+	{
+		bAutoFlushComputeShaderCache = bEnable;
+	}
+
+	void FlushComputeShaderCache(bool bForce = false);
+
+	/**
+	 * Clears all D3D11 State, setting all input/output resource slots, shaders, input layouts,
+	 * predications, scissor rectangles, depth-stencil state, rasterizer state, blend state,
+	 * sampler state, and viewports to NULL
+	 */
+	virtual void ClearState();
+
+	/**
+	 * Releases any object references held by the state cache
+	 */
+	void Clear();
+
+	void ForceRebuildGraphicsPSO() { PipelineState.Graphics.bNeedRebuildPSO = true; }
+	void ForceRebuildComputePSO() { PipelineState.Compute.bNeedRebuildPSO = true; }
+	void ForceSetVB() { bNeedSetVB = true; }
+	void ForceSetIB() { bNeedSetIB = true; }
+	void ForceSetUAVs() { bNeedSetUAVs = true; }
+	void ForceSetRTs() { bNeedSetRTs = true; }
+	void ForceSetSOs() { bNeedSetSOs = true; }
+	void ForceSetSamplersPerShaderStage(uint32 Frequency) { bNeedSetSamplersPerShaderStage[Frequency] = true; }
+	void ForceSetSamplers() { bNeedSetSamplers = true; }
+	void ForceSetSRVsPerShaderStage(uint32 Frequency) { bNeedSetSRVsPerShaderStage[Frequency] = true; }
+	void ForceSetSRVs() { bNeedSetSRVs = true; }
+	void ForceSetConstantBuffers() { bNeedSetConstantBuffers = true; }
+	void ForceSetViewports() { bNeedSetViewports = true; }
+	void ForceSetScissorRects() { bNeedSetScissorRects = true; }
+	void ForceSetPrimitiveTopology() { bNeedSetPrimitiveTopology = true; }
+	void ForceSetBlendFactor() { bNeedSetBlendFactor = true; }
+	void ForceSetStencilRef() { bNeedSetStencilRef = true; }
+
+	bool GetForceRebuildGraphicsPSO() { return PipelineState.Graphics.bNeedRebuildPSO; }
+	bool GetForceRebuildComputePSO() { return PipelineState.Compute.bNeedRebuildPSO; }
+	bool GetForceSetVB() { return bNeedSetVB; }
+	bool GetForceSetIB() { return bNeedSetIB; }
+	bool GetForceSetUAVs() { return bNeedSetUAVs; }
+	bool GetForceSetRTs() { return bNeedSetRTs; }
+	bool GetForceSetSOs() { return bNeedSetSOs; }
+	bool GetForceSetSamplersPerShaderStage(uint32 Frequency) { return bNeedSetSamplersPerShaderStage[Frequency]; }
+	bool GetForceSetSamplers() { return bNeedSetSamplers; }
+	bool GetForceSetSRVsPerShaderStage(uint32 Frequency) { return bNeedSetSRVsPerShaderStage[Frequency]; }
+	bool GetForceSetSRVs() { return bNeedSetSRVs; }
+	bool GetForceSetConstantBuffers() { return bNeedSetConstantBuffers; }
+	bool GetForceSetViewports() { return bNeedSetViewports; }
+	bool GetForceSetScissorRects() { return bNeedSetScissorRects; }
+	bool GetForceSetPrimitiveTopology() { return bNeedSetPrimitiveTopology; }
+	bool GetForceSetBlendFactor() { return bNeedSetBlendFactor; }
+	bool GetForceSetStencilRef() { return bNeedSetStencilRef; }
+
+
+#if D3D11_STATE_CACHE_DEBUG
+protected:
+	// Debug helper methods to verify cached state integrity.
+	template <EShaderFrequency ShaderFrequency>
+	void VerifySamplerStates();
+
+	template <EShaderFrequency ShaderFrequency>
+	void VerifyConstantBuffers();
+
+	template <EShaderFrequency ShaderFrequency>
+	void VerifyShaderResourceViews();
+#endif
+};
diff --git a/Engine/Source/Runtime/Windows/D3D12RHI/Private/D3D12StructuredBuffer.cpp b/Engine/Source/Runtime/Windows/D3D12RHI/Private/D3D12StructuredBuffer.cpp
new file mode 100644
--- /dev/null
+++ b/Engine/Source/Runtime/Windows/D3D12RHI/Private/D3D12StructuredBuffer.cpp
@@ -0,0 +1,227 @@
+// Copyright 1998-2014 Epic Games, Inc. All Rights Reserved.
+
+#include "D3D12RHIPrivate.h"
+
+FStructuredBufferRHIRef FD3D12DynamicRHI::RHICreateStructuredBuffer(uint32 Stride, uint32 Size, uint32 InUsage, FRHIResourceCreateInfo& CreateInfo)
+{
+	// Explicitly check that the size is nonzero before allowing CreateStructuredBuffer to opaquely fail.
+	check(Size > 0);
+	// Check for values that will cause D3D calls to fail
+	check(Size / Stride > 0 && Size % Stride == 0);
+
+    // Describe the structured buffer.
+    const bool bIsDynamic = (InUsage & BUF_AnyDynamic) ? true : false;
+    D3D12_RESOURCE_DESC Desc = CD3DX12_RESOURCE_DESC::Buffer(Size);
+
+    if ((InUsage & BUF_ShaderResource) == 0)
+    {
+        Desc.Flags |= D3D12_RESOURCE_FLAG_DENY_SHADER_RESOURCE;
+    }
+
+    if (InUsage & BUF_UnorderedAccess)
+    {
+        Desc.Flags |= D3D12_RESOURCE_FLAG_ALLOW_UNORDERED_ACCESS;
+    }
+
+	if (FPlatformProperties::SupportsFastVRAMMemory())
+	{
+		if (InUsage & BUF_FastVRAM)
+		{
+			FFastVRAMAllocator::GetFastVRAMAllocator()->AllocUAVBuffer(Desc);
+		}
+	}
+
+	// If a resource array was provided for the resource, create the resource pre-populated
+	D3D12_SUBRESOURCE_DATA InitData = { 0 };
+	D3D12_SUBRESOURCE_DATA* pInitData = nullptr;
+	if (CreateInfo.ResourceArray)
+	{
+		check(Size == CreateInfo.ResourceArray->GetResourceDataSize());
+		InitData.pData = CreateInfo.ResourceArray->GetResourceData();
+		InitData.RowPitch = Size;
+		InitData.SlicePitch = 0;
+		pInitData = &InitData;
+	}
+
+	TRefCountPtr<FD3D12ResourceLocation> StructuredBufferResource;
+    if (bIsDynamic)
+	{
+		// Use an upload heap for dynamic resources and map its memory for writing.
+		void* pData = GetRHIDevice()->GetDefaultUploadHeapAllocator().Alloc(Desc.Width, D3D12_TEXTURE_DATA_PLACEMENT_ALIGNMENT, StructuredBufferResource.GetInitReference());
+
+		// Add the lock to the lock map.
+		FD3D12LockedKey LockedKey(StructuredBufferResource.GetReference());
+		FD3D12LockedData LockedData;
+		LockedData.SetData(pData);
+		LockedData.Pitch = Desc.Width;
+		OutstandingLocks.Add(LockedKey, LockedData);
+
+		if (pInitData)
+		{
+			memcpy(pData, InitData.pData, Desc.Width);
+		}
+	}
+	else
+	{
+#if SUB_ALLOCATED_DEFAULT_ALLOCATIONS
+		// Structured buffers, non-byte address buffers, need to be aligned to their stride to ensure that they
+		// can be addressed correctly with element based offsets.
+		if ((InUsage & (BUF_ByteAddressBuffer | BUF_DrawIndirect)) == 0)
+			GetRHIDevice()->GetDefaultBufferAllocator().AllocDefaultResource(Desc, pInitData, StructuredBufferResource.GetInitReference(), Stride);
+
+		else
+#endif
+		GetRHIDevice()->GetDefaultBufferAllocator().AllocDefaultResource(Desc, pInitData, StructuredBufferResource.GetInitReference());
+	}
+
+	UpdateBufferStats(StructuredBufferResource, true, D3D12_BUFFER_TYPE_STRUCTURED);
+
+	if (CreateInfo.ResourceArray)
+	{
+		// Discard the resource array's contents.
+		CreateInfo.ResourceArray->Discard();
+	}
+
+	return new FD3D12StructuredBuffer(GetRHIDevice(), StructuredBufferResource, Stride, Size, InUsage);
+}
+
+FD3D12StructuredBuffer::~FD3D12StructuredBuffer()
+{
+	UpdateBufferStats(ResourceLocation, false, D3D12_BUFFER_TYPE_STRUCTURED);
+
+	// Dynamic resources have a permanent entry in the OutstandingLocks that needs to be cleaned up. 
+	if (GetUsage() & BUF_AnyDynamic)
+	{
+		FD3D12LockedKey LockedKey(ResourceLocation.GetReference());
+		FD3D12DynamicRHI::GetD3DRHI()->OutstandingLocks.Remove (LockedKey);
+	}
+}
+
+void* FD3D12DynamicRHI::RHILockStructuredBuffer(FStructuredBufferRHIParamRef StructuredBufferRHI, uint32 Offset, uint32 Size, EResourceLockMode LockMode)
+{
+	FD3D12StructuredBuffer*  StructuredBuffer = FD3D12DynamicRHI::ResourceCast(StructuredBufferRHI);
+
+    FD3D12CommandContext& DefaultContext = GetRHIDevice()->GetDefaultCommandContext();
+
+	// If this resource is bound to the device, unbind it
+	DefaultContext.ConditionalClearShaderResource(StructuredBuffer->ResourceLocation);
+
+	// Determine whether the Structured buffer is dynamic or not.
+	D3D12_RESOURCE_DESC const& Desc = StructuredBuffer->Resource->GetDesc();
+    const uint32 BufferUsage = StructuredBuffer->GetUsage();
+    const bool bIsDynamic = (BufferUsage & BUF_AnyDynamic) ? true : false;
+
+	FD3D12LockedKey LockedKey(StructuredBuffer->Resource.GetReference());
+	FD3D12LockedData LockedData;
+
+	if (bIsDynamic)
+	{
+		check(LockMode == RLM_WriteOnly);
+
+		// The buffer is dynamic, so we've already mapped it. Just return the offset pointer.
+		FD3D12LockedData* pLockedData = OutstandingLocks.Find(LockedKey);
+		check(pLockedData != nullptr);
+		return (void*)((uint8*)pLockedData->GetData() + Offset);
+	}
+	else
+	{
+		if (LockMode == RLM_ReadOnly)
+		{
+			// If the static buffer is being locked for reading, create a staging buffer.
+			TRefCountPtr<FD3D12Resource> StagingStructuredBuffer;
+            VERIFYD3D11RESULT(GetRHIDevice()->GetResourceHelper().CreateBuffer(D3D12_HEAP_TYPE_READBACK, Offset + Size, StagingStructuredBuffer.GetInitReference()));
+			LockedData.StagingResource = StagingStructuredBuffer;
+
+			// Copy the contents of the Structured buffer to the staging buffer.
+			DefaultContext.numCopies++;
+			DefaultContext.CommandListHandle->CopyResource(StagingStructuredBuffer->GetResource(), StructuredBuffer->Resource->GetResource());
+
+            GetRHIDevice()->GetDefaultCommandContext().FlushCommands(true);
+
+			// Map the staging buffer's memory for reading.
+			void* pData;
+			VERIFYD3D11RESULT(StagingStructuredBuffer->GetResource()->Map(0, nullptr, &pData));
+			LockedData.SetData(pData);
+			LockedData.Pitch = Offset + Size;
+		}
+		else
+		{
+			// If the static buffer is being locked for writing, allocate memory for the contents to be written to.
+			check(StructuredBuffer->ResourceLocation->GetEffectiveBufferSize() >= Size);
+
+			// Use an upload heap to copy data to a default resource.
+			TRefCountPtr<FD3D12ResourceLocation> UploadBuffer = new FD3D12ResourceLocation();
+            void *pData = GetRHIDevice()->GetDefaultUploadHeapAllocator().FastAlloc(Offset + Size, D3D12_TEXTURE_DATA_PLACEMENT_ALIGNMENT, UploadBuffer);
+
+			// Add the lock to the lock map.
+			FD3D12LockedKey LockedKey(UploadBuffer);
+			LockedData.SetData(pData);
+			LockedData.Pitch = Offset + Size;
+
+			// Keep track of the underlying resource for the upload heap so it can be referenced during Unmap.
+			LockedData.UploadHeapLocation = UploadBuffer;
+		}
+	}
+
+	// Add the lock to the lock map.
+	OutstandingLocks.Add(LockedKey, LockedData);
+
+	// Return the offset pointer
+	return (void*)((uint8*)LockedData.GetData() + Offset);
+}
+
+void FD3D12DynamicRHI::RHIUnlockStructuredBuffer(FStructuredBufferRHIParamRef StructuredBufferRHI)
+{
+	FD3D12StructuredBuffer*  StructuredBuffer = FD3D12DynamicRHI::ResourceCast(StructuredBufferRHI);
+
+	// Determine whether the Structured buffer is dynamic or not.
+	D3D12_RESOURCE_DESC const& Desc = StructuredBuffer->Resource->GetDesc();
+    const uint32 BufferUsage = StructuredBuffer->GetUsage();
+    const bool bIsDynamic = (BufferUsage & BUF_AnyDynamic) ? true : false;
+
+	// Find the outstanding lock for this VB.
+	FD3D12LockedKey LockedKey(StructuredBuffer->Resource.GetReference());
+	FD3D12LockedData* LockedData = OutstandingLocks.Find(LockedKey);
+	check(LockedData);
+
+	if (bIsDynamic)
+	{
+		// If the VB is dynamic, its upload heap memory can always stay mapped. Don't do anything.
+	}
+	else
+	{
+		// If the static VB lock involved a staging resource, it was locked for reading.
+		if (LockedData->StagingResource)
+		{
+			// Unmap the staging buffer's memory.
+			ID3D12Resource* StagingBuffer = LockedData->StagingResource.GetReference()->GetResource();
+			StagingBuffer->Unmap(0, nullptr);
+		}
+		else
+		{
+			check(LockedData->UploadHeapLocation != nullptr);
+
+			// Copy the contents of the temporary memory buffer allocated for writing into the VB.
+			FD3D12ResourceLocation* UploadHeap = LockedData->UploadHeapLocation;
+
+			FD3D12CommandListHandle& hCommandList = GetRHIDevice()->GetDefaultCommandContext().CommandListHandle;
+
+			// Writable structured bufferes are sometimes unlocked which means they sometimes need tracking.
+			FConditionalScopeResourceBarrier ConditionalScopeResourceBarrier(hCommandList, StructuredBuffer->Resource, D3D12_RESOURCE_STATE_COPY_DEST, 0);
+			// Don't need to transition upload heaps
+
+			GetRHIDevice()->GetDefaultCommandContext().numCopies++;
+			hCommandList->CopyBufferRegion(
+                StructuredBuffer->Resource->GetResource(),
+				StructuredBuffer->ResourceLocation->GetOffset(),
+				UploadHeap->GetResource()->GetResource(),
+				UploadHeap->GetOffset(),
+				LockedData->Pitch);
+		}
+	}
+
+	// Remove the FD3D12LockedData from the lock map.
+	// If the lock involved a staging resource, this releases it.
+	// If the lock involved an upload heap resource, this releases it.
+	OutstandingLocks.Remove(LockedKey);
+}
