UE4 - D3D12RHI (not working)

diff --git a/Engine/Source/Runtime/Windows/D3D12RHI/Private/D3D12Texture.cpp b/Engine/Source/Runtime/Windows/D3D12RHI/Private/D3D12Texture.cpp
new file mode 100644
--- /dev/null
+++ b/Engine/Source/Runtime/Windows/D3D12RHI/Private/D3D12Texture.cpp
@@ -0,0 +1,2037 @@
+// Copyright 1998-2014 Epic Games, Inc. All Rights Reserved.
+
+/*=============================================================================
+	D3D12VertexBuffer.cpp: D3D texture RHI implementation.
+	=============================================================================*/
+
+#include "D3D12RHIPrivate.h"
+
+/*-----------------------------------------------------------------------------
+	Texture allocator support.
+	-----------------------------------------------------------------------------*/
+namespace D3D12RHI
+{
+	// in bytes, never change after RHI, needed to scale game features
+	int64 GDedicatedVideoMemory = 0;
+	// in bytes, never change after RHI, needed to scale game features
+	int64 GDedicatedSystemMemory = 0;
+	// in bytes, never change after RHI, needed to scale game features
+	int64 GSharedSystemMemory = 0;
+	// In bytes. Never changed after RHI init. Our estimate of the amount of memory that we can use for graphics resources in total.
+	int64 GTotalGraphicsMemory = 0;
+
+	static bool ShouldCountAsTextureMemory(uint32 MiscFlags)
+	{
+        // Shouldn't be used for DEPTH, RENDER TARGET, or UNORDERED ACCESS
+        return (0 == (MiscFlags & (D3D12_RESOURCE_FLAG_ALLOW_DEPTH_STENCIL | D3D12_RESOURCE_FLAG_ALLOW_RENDER_TARGET | D3D12_RESOURCE_FLAG_ALLOW_UNORDERED_ACCESS)));
+	}
+
+	// @param b3D true:3D, false:2D or cube map
+	static TStatId GetD3D11StatEnum(uint32 MiscFlags, bool bCubeMap, bool b3D)
+	{
+#if STATS
+		if (ShouldCountAsTextureMemory(MiscFlags))
+		{
+			// normal texture
+			if (bCubeMap)
+			{
+				return GET_STATID(STAT_TextureMemoryCube);
+			}
+			else if (b3D)
+			{
+				return GET_STATID(STAT_TextureMemory3D);
+			}
+			else
+			{
+				return GET_STATID(STAT_TextureMemory2D);
+			}
+		}
+		else
+		{
+			// render target
+			if (bCubeMap)
+			{
+				return GET_STATID(STAT_RenderTargetMemoryCube);
+			}
+			else if (b3D)
+			{
+				return GET_STATID(STAT_RenderTargetMemory3D);
+			}
+			else
+			{
+				return GET_STATID(STAT_RenderTargetMemory2D);
+			}
+		}
+#endif
+		return TStatId();
+	}
+
+	// Note: This function can be called from many different threads
+	// @param TextureSize >0 to allocate, <0 to deallocate
+	// @param b3D true:3D, false:2D or cube map
+    void UpdateD3D11TextureStats(const D3D12_RESOURCE_DESC& Desc, int64 TextureSize, bool b3D, bool bCubeMap)
+	{
+		if (TextureSize == 0)
+		{
+			return;
+		}
+
+		// MSFT: This code (AlignedSize) is in DX11's RHI but it appears to regress texture quality in infiltrator. For now, keeping the existing code
+		if (ShouldCountAsTextureMemory(Desc.Flags))
+		{
+			FPlatformAtomics::InterlockedAdd(&GCurrentTextureMemorySize, Align(D3D12_TEXTURE_DATA_PLACEMENT_ALIGNMENT, 1024) / 1024);
+		}
+		else
+		{
+			FPlatformAtomics::InterlockedAdd(&GCurrentRendertargetMemorySize, Align(D3D12_TEXTURE_DATA_PLACEMENT_ALIGNMENT, 1024) / 1024);
+		}
+		/*int64 AlignedSize = (TextureSize > 0) ? Align(TextureSize, 1024) / 1024 : -(Align(-TextureSize, 1024) / 1024);
+        if (ShouldCountAsTextureMemory(Desc.Flags))
+		{
+			FPlatformAtomics::InterlockedAdd(&GCurrentTextureMemorySize, AlignedSize);
+		}
+		else
+		{
+			FPlatformAtomics::InterlockedAdd(&GCurrentRendertargetMemorySize, AlignedSize);
+		}*/
+
+		INC_MEMORY_STAT_BY_FName(GetD3D11StatEnum(Desc.Flags, bCubeMap, b3D).GetName(), TextureSize);
+
+		if (TextureSize > 0)
+		{
+			INC_DWORD_STAT(STAT_D3D12TexturesAllocated);
+		}
+		else
+		{
+			INC_DWORD_STAT(STAT_D3D12TexturesReleased);
+		}
+	}
+
+	template<typename BaseResourceType>
+	void D3D11TextureAllocated(TD3D12Texture2D<BaseResourceType>& Texture)
+	{
+		FD3D12Resource* D3D11Texture2D = Texture.GetResource();
+
+		if (D3D11Texture2D)
+		{
+			if ((Texture.Flags & TexCreate_Virtual) == TexCreate_Virtual)
+			{
+				Texture.SetMemorySize(0);
+			}
+			else
+			{
+				const D3D12_RESOURCE_DESC& Desc = D3D11Texture2D->GetDesc();
+
+				int64 TextureSize = CalcTextureSize(Desc.Width, Desc.Height, Texture.GetFormat(), Desc.MipLevels) * Desc.DepthOrArraySize;
+
+				Texture.SetMemorySize(TextureSize);
+                UpdateD3D11TextureStats(Desc, TextureSize, false, Texture.IsCubemap());
+			}
+		}
+	}
+
+	template<typename BaseResourceType>
+	void D3D11TextureDeleted(TD3D12Texture2D<BaseResourceType>& Texture)
+	{
+		FD3D12Resource* D3D11Texture2D = Texture.GetResource();
+
+		if (D3D11Texture2D)
+		{
+			const D3D12_RESOURCE_DESC& Desc = D3D11Texture2D->GetDesc();
+
+			// When using virtual textures use the current memory size, which is the number of physical pages allocated, not virtual
+			int64 TextureSize = 0;
+			if ((Texture.GetFlags() & TexCreate_Virtual) == TexCreate_Virtual)
+			{
+				TextureSize = Texture.GetMemorySize();
+			}
+			else
+			{
+                TextureSize = CalcTextureSize(Desc.Width, Desc.Height, Texture.GetFormat(), Desc.MipLevels) * Desc.DepthOrArraySize;
+			}
+
+			UpdateD3D11TextureStats(Desc, -TextureSize, false, Texture.IsCubemap());
+		}
+	}
+
+	void D3D11TextureAllocated2D(FD3D12Texture2D& Texture)
+	{
+		D3D11TextureAllocated(Texture);
+	}
+
+	void D3D11TextureAllocated(FD3D12Texture3D& Texture)
+	{
+		FD3D12Resource* D3D11Texture3D = Texture.GetResource();
+
+		if (D3D11Texture3D)
+		{
+            const D3D12_RESOURCE_DESC& Desc = D3D11Texture3D->GetDesc();
+
+            int64 TextureSize = CalcTextureSize3D(Desc.Width, Desc.Height, Desc.DepthOrArraySize, Texture.GetFormat(), Desc.MipLevels);
+
+			Texture.SetMemorySize(TextureSize);
+
+			UpdateD3D11TextureStats(Desc, TextureSize, true, false);
+		}
+	}
+
+	void D3D11TextureDeleted(FD3D12Texture3D& Texture)
+	{
+		FD3D12Resource* D3D11Texture3D = Texture.GetResource();
+
+		if (D3D11Texture3D)
+		{
+            const D3D12_RESOURCE_DESC& Desc = D3D11Texture3D->GetDesc();
+
+            int64 TextureSize = CalcTextureSize3D(Desc.Width, Desc.Height, Desc.DepthOrArraySize, Texture.GetFormat(), Desc.MipLevels);
+
+			UpdateD3D11TextureStats(Desc, -TextureSize, true, false);
+		}
+	}
+}
+using namespace D3D12RHI;
+
+
+template<typename BaseResourceType>
+TD3D12Texture2D<BaseResourceType>::~TD3D12Texture2D()
+{
+	D3D11TextureDeleted(*this);
+	if (bPooled)
+	{
+		ReturnPooledTexture2D(GetNumMips(), GetFormat(), GetResource());
+	}
+
+#if PLATFORM_SUPPORTS_VIRTUAL_TEXTURES
+	D3DRHI->DestroyVirtualTexture(GetFlags(), GetRawTextureMemory());
+#endif
+}
+
+FD3D12Texture3D::~FD3D12Texture3D()
+{
+	D3D11TextureDeleted(*this);
+}
+
+uint64 FD3D12DynamicRHI::RHICalcTexture2DPlatformSize(uint32 SizeX, uint32 SizeY, uint8 Format, uint32 NumMips, uint32 NumSamples, uint32 Flags, uint32& OutAlign)
+{
+	OutAlign = 0;
+	return CalcTextureSize(SizeX, SizeY, (EPixelFormat)Format, NumMips);
+}
+
+uint64 FD3D12DynamicRHI::RHICalcTexture3DPlatformSize(uint32 SizeX, uint32 SizeY, uint32 SizeZ, uint8 Format, uint32 NumMips, uint32 Flags, uint32& OutAlign)
+{
+	OutAlign = 0;
+	return CalcTextureSize3D(SizeX, SizeY, SizeZ, (EPixelFormat)Format, NumMips);
+}
+
+uint64 FD3D12DynamicRHI::RHICalcTextureCubePlatformSize(uint32 Size, uint8 Format, uint32 NumMips, uint32 Flags, uint32& OutAlign)
+{
+	OutAlign = 0;
+	return CalcTextureSize(Size, Size, (EPixelFormat)Format, NumMips) * 6;
+}
+
+/**
+ * Retrieves texture memory stats.
+ */
+void FD3D12DynamicRHI::RHIGetTextureMemoryStats(FTextureMemoryStats& OutStats)
+{
+	OutStats.DedicatedVideoMemory = GDedicatedVideoMemory;
+	OutStats.DedicatedSystemMemory = GDedicatedSystemMemory;
+	OutStats.SharedSystemMemory = GSharedSystemMemory;
+	bool bIsWARP = FParse::Param(FCommandLine::Get(), TEXT("warp"));
+
+	// Hack to handle a kernel bug where QueryVideomMemoryInfo can't be called on WARP
+	if (!bIsWARP)
+	{
+		DXGI_QUERY_VIDEO_MEMORY_INFO LocalVideoMemoryInfo = {};
+		VERIFYD3D11RESULT(GetRHIDevice()->GetAdapter3()->QueryVideoMemoryInfo(0, DXGI_MEMORY_SEGMENT_GROUP_LOCAL, &LocalVideoMemoryInfo));
+
+		int64 budget = LocalVideoMemoryInfo.Budget;
+		OutStats.TotalGraphicsMemory = budget;
+
+		// Leave 5% unused to account for sudden budget cuts
+		const int64 BudgetPadding = budget * 0.05f;
+		// Note: AvailableSpace can be negative
+		int64 AvailableSpace = budget - int64(LocalVideoMemoryInfo.CurrentUsage);
+		int64 PreviousPoolSize = GTexturePoolSize;
+
+		// Only change the pool size if over-budget, or a reasonable amount of memory is available
+		if (AvailableSpace < 0 || AvailableSpace > BudgetPadding)
+		{
+			int64 BudgettedPoolSize = PreviousPoolSize > 0 ? PreviousPoolSize + AvailableSpace - BudgetPadding : INT64_MAX;
+
+			float MaxPoolSize = float(GPoolSizeVRAMPercentage) * 0.01f * budget;
+			int64 TruncMaxPoolSize = int64(FGenericPlatformMath::TruncToFloat(MaxPoolSize / 1024.0f / 1024.0f)) * 1024 * 1024;
+
+			if (sizeof(SIZE_T) < 8)
+			{
+				// Clamp to 1 GB if we're less than 64-bit
+				TruncMaxPoolSize = FMath::Min(TruncMaxPoolSize, 1024ll * 1024ll * 1024ll);
+			}
+			else
+			{
+				// Clamp to 1.9 GB if we're 64-bit
+				TruncMaxPoolSize = FMath::Min(TruncMaxPoolSize, 1945ll * 1024ll * 1024ll);
+			}
+			
+			// Clamp to the pre-defined texture pool size, if it's been set.
+			TruncMaxPoolSize = (GTexturePoolSize == 0) ? TruncMaxPoolSize : FMath::Min(TruncMaxPoolSize, GTexturePoolSize);
+
+			const int64 MinPoolSize = int64(200 * 1024 * 1024);
+			OutStats.TexturePoolSize = FMath::Max(FMath::Min(BudgettedPoolSize, TruncMaxPoolSize), MinPoolSize);
+		}
+		else
+		{
+			OutStats.TexturePoolSize = PreviousPoolSize;
+		}
+
+		// Truncate GTexturePoolSize to MB (but still counted in bytes)
+		OutStats.AllocatedMemorySize = int64(GCurrentTextureMemorySize) * 1024;
+		OutStats.PendingMemoryAdjustment = 0;
+	}
+	else
+	{
+		OutStats.TotalGraphicsMemory = GTotalGraphicsMemory;
+	}
+
+
+}
+
+/**
+ * Fills a texture with to visualize the texture pool memory.
+ *
+ * @param	TextureData		Start address
+ * @param	SizeX			Number of pixels along X
+ * @param	SizeY			Number of pixels along Y
+ * @param	Pitch			Number of bytes between each row
+ * @param	PixelSize		Number of bytes each pixel represents
+ *
+ * @return true if successful, false otherwise
+ */
+bool FD3D12DynamicRHI::RHIGetTextureMemoryVisualizeData(FColor* /*TextureData*/, int32 /*SizeX*/, int32 /*SizeY*/, int32 /*Pitch*/, int32 /*PixelSize*/)
+{
+	// currently only implemented for console (Note: Keep this function for further extension. Talk to NiklasS for more info.)
+	return false;
+}
+
+/*------------------------------------------------------------------------------
+	Texture pooling.
+	------------------------------------------------------------------------------*/
+
+/** Define to 1 to enable the pooling of 2D texture resources. */
+#define USE_TEXTURE_POOLING 0
+
+namespace D3D12RHI
+{
+	/** A texture resource stored in the pool. */
+	struct FPooledTexture2D
+	{
+		/** The texture resource. */
+		TRefCountPtr<FD3D12Resource> Resource;
+	};
+
+	/** A pool of D3D texture resources. */
+	struct FTexturePool
+	{
+		TArray<FPooledTexture2D> Textures;
+	};
+
+	/** The global texture pool. */
+	struct FGlobalTexturePool
+	{
+		/** Formats stored in the pool. */
+		enum EInternalFormat
+		{
+			IF_DXT1,
+			IF_DXT5,
+			IF_BC5,
+			IF_Max
+		};
+
+		enum
+		{
+			/** Minimum mip count for which to pool textures. */
+			MinMipCount = 7,
+			/** Maximum mip count for which to pool textures. */
+			MaxMipCount = 13,
+			/** The number of pools based on mip levels. */
+			MipPoolCount = MaxMipCount - MinMipCount,
+		};
+
+		/** The individual texture pools. */
+		FTexturePool Pools[MipPoolCount][IF_Max];
+	};
+	FGlobalTexturePool GTexturePool;
+
+	/**
+	 * Releases all pooled textures.
+	 */
+	void ReleasePooledTextures()
+	{
+		for (int32 MipPoolIndex = 0; MipPoolIndex < FGlobalTexturePool::MipPoolCount; ++MipPoolIndex)
+		{
+			for (int32 FormatPoolIndex = 0; FormatPoolIndex < FGlobalTexturePool::IF_Max; ++FormatPoolIndex)
+			{
+				GTexturePool.Pools[MipPoolIndex][FormatPoolIndex].Textures.Empty();
+			}
+		}
+	}
+
+	/**
+	 * Retrieves the texture pool for the specified mip count and format.
+	 */
+	FTexturePool* GetTexturePool(int32 MipCount, EPixelFormat PixelFormat)
+	{
+		FTexturePool* Pool = NULL;
+		int32 MipPool = MipCount - FGlobalTexturePool::MinMipCount;
+		if (MipPool >= 0 && MipPool < FGlobalTexturePool::MipPoolCount)
+		{
+			int32 FormatPool = -1;
+			switch (PixelFormat)
+			{
+			case PF_DXT1: FormatPool = FGlobalTexturePool::IF_DXT1; break;
+			case PF_DXT5: FormatPool = FGlobalTexturePool::IF_DXT5; break;
+			case PF_BC5: FormatPool = FGlobalTexturePool::IF_BC5; break;
+			}
+			if (FormatPool >= 0 && FormatPool < FGlobalTexturePool::IF_Max)
+			{
+				Pool = &GTexturePool.Pools[MipPool][FormatPool];
+			}
+		}
+		return Pool;
+	}
+
+	/**
+	 * Retrieves a texture from the pool if one exists.
+	 */
+	bool GetPooledTexture2D(int32 MipCount, EPixelFormat PixelFormat, FPooledTexture2D* OutTexture)
+	{
+#if USE_TEXTURE_POOLING
+		FTexturePool* Pool = GetTexturePool(MipCount, PixelFormat);
+		if (Pool && Pool->Textures.Num() > 0)
+		{
+			*OutTexture = Pool->Textures.Last();
+
+			{
+                const D3D12_RESOURCE_DESC& Desc = OutTexture->Resource->GetDesc();
+				check(Desc.Format == GPixelFormats[PixelFormat].PlatformFormat);
+				check(MipCount == Desc.MipLevels);
+				check(Desc.Width == Desc.Height);
+				check(Desc.Width == (1 << (MipCount - 1)));
+				int32 TextureSize = CalcTextureSize(Desc.Width, Desc.Height, PixelFormat, Desc.MipLevels);
+				DEC_MEMORY_STAT_BY(STAT_D3D12TexturePoolMemory, TextureSize);
+			}
+
+			Pool->Textures.RemoveAt(Pool->Textures.Num() - 1);
+			return true;
+		}
+#endif // #if USE_TEXTURE_POOLING
+		return false;
+	}
+
+	/**
+	 * Returns a texture to its pool.
+	 */
+	void ReturnPooledTexture2D(int32 MipCount, EPixelFormat PixelFormat, FD3D12Resource* InResource)
+	{
+#if USE_TEXTURE_POOLING
+		FTexturePool* Pool = GetTexturePool(MipCount, PixelFormat);
+		if (Pool)
+		{
+			FPooledTexture2D* PooledTexture = new(Pool->Textures) FPooledTexture2D;
+			PooledTexture->Resource = InResource;
+			{
+                const D3D12_RESOURCE_DESC& Desc = PooledTexture->Resource->GetDesc();
+				check(Desc.Format == GPixelFormats[PixelFormat].PlatformFormat);
+				check(MipCount == Desc.MipLevels);
+				check(Desc.Width == Desc.Height);
+				check(Desc.Width == (1 << (MipCount - 1)));
+				int32 TextureSize = CalcTextureSize(Desc.Width, Desc.Height, PixelFormat, Desc.MipLevels);
+				INC_MEMORY_STAT_BY(STAT_D3D12TexturePoolMemory, TextureSize);
+			}
+		}
+#endif // #if USE_TEXTURE_POOLING
+	}
+}
+using namespace D3D12RHI;
+
+#if WITH_D3DX_LIBS
+DXGI_FORMAT FD3D12DynamicRHI::GetPlatformTextureResourceFormat(DXGI_FORMAT InFormat, uint32 InFlags)
+{
+	// DX 11 Shared textures must be B8G8R8A8_UNORM
+	if (InFlags & TexCreate_Shared)
+	{
+		return DXGI_FORMAT_B8G8R8A8_UNORM;
+	}
+	return InFormat;
+}
+#endif	//WITH_D3DX_LIBS
+
+/** If true, guard texture creates with SEH to log more information about a driver crash we are seeing during texture streaming. */
+#define GUARDED_TEXTURE_CREATES (PLATFORM_WINDOWS && !(UE_BUILD_SHIPPING || UE_BUILD_TEST))
+
+void InitializeTexture2DData(
+	const D3D12_RESOURCE_DESC& TextureDesc,
+	D3D12_SUBRESOURCE_DATA* SubResourceData,
+	FD3D12CommandListHandle& hCommandList,
+	FD3D12DynamicHeapAllocator &DynamicHeapAllocator,
+	uint32 BlockBytes,
+	uint32 BlockSizeX,
+	uint32 BlockSizeY,
+	uint32 NumBlocksX,
+	uint32 NumBlocksY,
+	FD3D12Resource *pDestTexture2D)
+{
+    uint32 NumSubresources = TextureDesc.MipLevels * TextureDesc.DepthOrArraySize;
+    uint64 IntermediateBufferSize = GetRequiredIntermediateSize(pDestTexture2D->GetResource(), 0, NumSubresources);
+
+    FD3D12ResourceLocation SrcResourceLoc;
+    void *pData = DynamicHeapAllocator.Alloc(IntermediateBufferSize, D3D12_TEXTURE_DATA_PLACEMENT_ALIGNMENT, &SrcResourceLoc);
+    check(pData);
+
+	{
+		FScopeResourceBarrier ScopeResourceBarrier(hCommandList, pDestTexture2D, D3D12_RESOURCE_STATE_COMMON, D3D12_RESOURCE_STATE_COPY_DEST, D3D12_RESOURCE_BARRIER_ALL_SUBRESOURCES);
+
+		// Using the stack allocated implementation of UpdateSubresources, using 15 max subresources as that should be sufficient for large textures
+		const uint32 maxSubresource = 15;
+		check(NumSubresources <= maxSubresource);
+
+		hCommandList.GetCurrentOwningContext()->numCopies++;
+		uint64 result = UpdateSubresources<maxSubresource>(
+			hCommandList.CommandList(),
+			pDestTexture2D->GetResource(),
+			SrcResourceLoc.GetResource()->GetResource(), SrcResourceLoc.GetOffset(),
+			(uint32)0, NumSubresources,
+			SubResourceData);
+	}
+}
+
+/**
+ * Creates a 2D texture optionally guarded by a structured exception handler.
+ */
+void SafeCreateTexture2D(FD3D12Device *pDevice, const D3D12_RESOURCE_DESC& TextureDesc, D3D12_SUBRESOURCE_DATA* SubResourceData, FD3D12CommandListHandle& hCommandList, FD3D12DynamicHeapAllocator &DynamicHeapAllocator, FD3D12Resource** OutTexture2D, uint8 Format, uint32 Flags)
+{
+#if GUARDED_TEXTURE_CREATES
+	bool bDriverCrash = true;
+	__try
+	{
+#endif // #if GUARDED_TEXTURE_CREATES
+
+        check(hCommandList != NULL);
+
+		const D3D12_HEAP_TYPE heapType = (Flags & TexCreate_CPUReadback) ? D3D12_HEAP_TYPE_READBACK : D3D12_HEAP_TYPE_DEFAULT;
+        const uint64 BlockSizeX = GPixelFormats[Format].BlockSizeX;
+        const uint64 BlockSizeY = GPixelFormats[Format].BlockSizeY;
+        const uint64 BlockBytes = GPixelFormats[Format].BlockBytes;
+        const uint64 MipSizeX = FMath::Max(TextureDesc.Width, BlockSizeX);
+        const uint64 MipSizeY = FMath::Max((uint64)TextureDesc.Height, BlockSizeY);
+        const uint64 NumBlocksX = (MipSizeX + BlockSizeX - 1) / BlockSizeX;
+        const uint64 NumBlocksY = (MipSizeY + BlockSizeY - 1) / BlockSizeY;
+        const uint64 XBytesAligned = Align(NumBlocksX * BlockBytes, D3D12_TEXTURE_DATA_PITCH_ALIGNMENT);
+        const uint64 MipBytesAligned = Align(NumBlocksY * XBytesAligned, D3D12_TEXTURE_DATA_PLACEMENT_ALIGNMENT);
+
+        switch (heapType)
+		{
+        case D3D12_HEAP_TYPE_READBACK:
+			VERIFYD3D11CREATETEXTURERESULT(
+                pDevice->GetResourceHelper().CreateBuffer(heapType, MipBytesAligned, OutTexture2D),
+				TextureDesc.Width,
+				TextureDesc.Height,
+                TextureDesc.DepthOrArraySize,
+				TextureDesc.Format,
+				TextureDesc.MipLevels,
+				TextureDesc.Flags
+				);
+
+			break;
+
+        case D3D12_HEAP_TYPE_DEFAULT:
+			VERIFYD3D11CREATETEXTURERESULT(
+                pDevice->GetResourceHelper().CreateDefaultResource(TextureDesc, OutTexture2D),
+				TextureDesc.Width,
+				TextureDesc.Height,
+				TextureDesc.DepthOrArraySize,
+				TextureDesc.Format,
+				TextureDesc.MipLevels,
+				TextureDesc.Flags
+				);
+
+			break;
+
+		default:
+			check(false);	// Need to create a resource here
+		}
+
+		if (SubResourceData)
+		{
+			// SubResourceData is only used in async texture creation (RHIAsyncCreateTexture2D). We need to manually transition the resource to
+			// its 'default state', which is what the rest of the RHI (including InitializeTexture2DData) expects for SRV-only resources.
+
+            check((TextureDesc.Flags & D3D12_RESOURCE_FLAG_DENY_SHADER_RESOURCE) == 0);
+			FD3D12Resource *pDestTexture2D = (*OutTexture2D);
+			InitializeTexture2DData(
+				TextureDesc,
+				SubResourceData,
+				hCommandList,
+				DynamicHeapAllocator,
+				BlockBytes,
+				BlockSizeX,
+				BlockSizeY,
+				NumBlocksX,
+				NumBlocksY,
+				pDestTexture2D);
+		}
+
+#if GUARDED_TEXTURE_CREATES
+		bDriverCrash = false;
+	}
+	__finally
+	{
+		if (bDriverCrash)
+		{
+			UE_LOG(LogD3D12RHI, Error,
+				TEXT("Driver crashed while creating texture: %ux%ux%u %s(0x%08x) with %u mips"),
+				TextureDesc.Width,
+				TextureDesc.Height,
+				TextureDesc.DepthOrArraySize,
+				GetD3D11TextureFormatString(TextureDesc.Format),
+				(uint32)TextureDesc.Format,
+				TextureDesc.MipLevels
+				);
+		}
+	}
+#endif // #if GUARDED_TEXTURE_CREATES
+}
+
+
+template<typename BaseResourceType>
+TD3D12Texture2D<BaseResourceType>* FD3D12DynamicRHI::CreateD3D11Texture2D(uint32 SizeX, uint32 SizeY, uint32 SizeZ, bool bTextureArray, bool bCubeTexture, uint8 Format,
+	uint32 NumMips, uint32 NumSamples, uint32 Flags, FRHIResourceCreateInfo& CreateInfo)
+{
+	check(SizeX > 0 && SizeY > 0 && NumMips > 0);
+
+	if (bCubeTexture)
+	{
+		check(SizeX <= GetMaxCubeTextureDimension());
+		check(SizeX == SizeY);
+	}
+	else
+	{
+		check(SizeX <= GetMax2DTextureDimension());
+		check(SizeY <= GetMax2DTextureDimension());
+	}
+
+	if (bTextureArray)
+	{
+		check(SizeZ <= GetMaxTextureArrayLayers());
+	}
+
+	// Render target allocation with UAV flag will silently fail in feature level 10
+	check(FeatureLevel >= D3D_FEATURE_LEVEL_11_0 || !(Flags & TexCreate_UAV));
+
+	SCOPE_CYCLE_COUNTER(STAT_D3D12CreateTextureTime);
+
+	bool bPooledTexture = true;
+
+	if (GMaxRHIFeatureLevel <= ERHIFeatureLevel::ES2)
+	{
+		// Remove sRGB read flag when not supported
+		Flags &= ~TexCreate_SRGB;
+	}
+
+	const bool bSRGB = (Flags & TexCreate_SRGB) != 0;
+
+	const DXGI_FORMAT PlatformResourceFormat = FD3D12DynamicRHI::GetPlatformTextureResourceFormat((DXGI_FORMAT)GPixelFormats[Format].PlatformFormat, Flags);
+	const DXGI_FORMAT PlatformShaderResourceFormat = FindShaderResourceDXGIFormat(PlatformResourceFormat, bSRGB);
+	const DXGI_FORMAT PlatformRenderTargetFormat = FindShaderResourceDXGIFormat(PlatformResourceFormat, bSRGB);
+
+	// Determine the MSAA settings to use for the texture.
+	D3D12_DSV_DIMENSION DepthStencilViewDimension = D3D12_DSV_DIMENSION_TEXTURE2D;
+	D3D12_RTV_DIMENSION RenderTargetViewDimension = D3D12_RTV_DIMENSION_TEXTURE2D;
+	D3D12_SRV_DIMENSION ShaderResourceViewDimension = D3D12_SRV_DIMENSION_TEXTURE2D;
+	bool bCreateShaderResource = true;
+
+	uint32 ActualMSAACount = NumSamples;
+
+	uint32 ActualMSAAQuality = GetMaxMSAAQuality(ActualMSAACount);
+
+	// 0xffffffff means not supported
+	if (ActualMSAAQuality == 0xffffffff || (Flags & TexCreate_Shared) != 0)
+	{
+		// no MSAA
+		ActualMSAACount = 1;
+		ActualMSAAQuality = 0;
+	}
+
+	if (ActualMSAACount > 1)
+	{
+		DepthStencilViewDimension = D3D12_DSV_DIMENSION_TEXTURE2DMS;
+		RenderTargetViewDimension = D3D12_RTV_DIMENSION_TEXTURE2DMS;
+		ShaderResourceViewDimension = D3D12_SRV_DIMENSION_TEXTURE2DMS;
+		bPooledTexture = false;
+	}
+
+	if (NumMips < 1 || SizeX != SizeY || (1 << (NumMips - 1)) != SizeX || (Flags & TexCreate_Shared) != 0)
+	{
+		bPooledTexture = false;
+	}
+
+	if (Flags & TexCreate_CPUReadback)
+	{
+		check(!(Flags & TexCreate_RenderTargetable));
+		check(!(Flags & TexCreate_DepthStencilTargetable));
+		check(!(Flags & TexCreate_ShaderResource));
+		bCreateShaderResource = false;
+	}
+
+	// Describe the texture.
+    D3D12_RESOURCE_DESC TextureDesc = CD3DX12_RESOURCE_DESC::Tex2D(
+        PlatformResourceFormat,
+        SizeX,
+        SizeY,
+        SizeZ,  // Array size
+        NumMips,
+        ActualMSAACount,
+        ActualMSAAQuality,
+        D3D12_RESOURCE_FLAG_NONE);  // Add misc flags later
+
+	if (Flags & TexCreate_GenerateMipCapable)
+	{
+		// Set the flag that allows us to call GenerateMips on this texture later
+		bPooledTexture = false;
+	}
+
+	// Set up the texture bind flags.
+	bool bCreateRTV = false;
+	bool bCreateDSV = false;
+	bool bCreatedRTVPerSlice = false;
+
+	if (Flags & TexCreate_RenderTargetable)
+	{
+		check(!(Flags & TexCreate_DepthStencilTargetable));
+		check(!(Flags & TexCreate_ResolveTargetable));
+        TextureDesc.Flags |= D3D12_RESOURCE_FLAG_ALLOW_RENDER_TARGET;
+		bCreateRTV = true;
+	}
+	else if (Flags & TexCreate_DepthStencilTargetable)
+	{
+		check(!(Flags & TexCreate_RenderTargetable));
+		check(!(Flags & TexCreate_ResolveTargetable));
+        TextureDesc.Flags |= D3D12_RESOURCE_FLAG_ALLOW_DEPTH_STENCIL;
+		bCreateDSV = true;
+	}
+	else if (Flags & TexCreate_ResolveTargetable)
+	{
+		check(!(Flags & TexCreate_RenderTargetable));
+		check(!(Flags & TexCreate_DepthStencilTargetable));
+		if (Format == PF_DepthStencil || Format == PF_ShadowDepth || Format == PF_D24)
+		{
+            TextureDesc.Flags |= D3D12_RESOURCE_FLAG_ALLOW_DEPTH_STENCIL;
+			bCreateDSV = true;
+		}
+		else
+		{
+            TextureDesc.Flags |= D3D12_RESOURCE_FLAG_ALLOW_RENDER_TARGET;
+			bCreateRTV = true;
+		}
+	}
+
+	if (Flags & TexCreate_UAV)
+	{ 
+        TextureDesc.Flags |= D3D12_RESOURCE_FLAG_ALLOW_UNORDERED_ACCESS;
+		bPooledTexture = false;
+	}
+
+	if (bCreateDSV || bCreateRTV || bCubeTexture || bTextureArray)
+	{
+		bPooledTexture = false;
+	}
+
+	FVRamAllocation VRamAllocation;
+
+	if (FPlatformProperties::SupportsFastVRAMMemory())
+	{
+		if (Flags & TexCreate_FastVRAM)
+		{
+			VRamAllocation = FFastVRAMAllocator::GetFastVRAMAllocator()->AllocTexture2D(TextureDesc);
+		}
+	}
+
+	TRefCountPtr<FD3D12Resource> TextureResource;
+	TArray<TRefCountPtr<FD3D12RenderTargetView> > RenderTargetViews;
+	TRefCountPtr<FD3D12DepthStencilView> DepthStencilViews[FExclusiveDepthStencil::MaxIndex];
+
+#if PLATFORM_SUPPORTS_VIRTUAL_TEXTURES
+	// Turn off pooling when we are using virtual textures or the texture is offline processed as we control when the memory is released
+	if ((Flags & (TexCreate_Virtual | TexCreate_OfflineProcessed)) != 0)
+	{
+		bPooledTexture = false;
+	}
+	void* RawTextureMemory = nullptr;
+#else
+	Flags &= ~TexCreate_Virtual;
+#endif
+
+	if (bPooledTexture)
+	{
+		FPooledTexture2D PooledTexture;
+		if (GetPooledTexture2D(NumMips, (EPixelFormat)Format, &PooledTexture))
+		{
+			TextureResource = PooledTexture.Resource;
+		}
+	}
+
+	if (!IsValidRef(TextureResource))
+	{
+		TArray<D3D12_SUBRESOURCE_DATA> SubResourceData;
+
+		if (CreateInfo.BulkData)
+		{
+			uint8* Data = (uint8*)CreateInfo.BulkData->GetResourceBulkData();
+
+			// each mip of each array slice counts as a subresource
+			SubResourceData.AddZeroed(NumMips * SizeZ);
+
+			uint32 SliceOffset = 0;
+			for (uint32 ArraySliceIndex = 0; ArraySliceIndex < SizeZ; ++ArraySliceIndex)
+			{
+				uint32 MipOffset = 0;
+				for (uint32 MipIndex = 0; MipIndex < NumMips; ++MipIndex)
+				{
+					uint32 DataOffset = SliceOffset + MipOffset;
+					uint32 SubResourceIndex = ArraySliceIndex * NumMips + MipIndex;
+
+					uint32 NumBlocksX = FMath::Max<uint32>(1, (SizeX >> MipIndex) / GPixelFormats[Format].BlockSizeX);
+					uint32 NumBlocksY = FMath::Max<uint32>(1, (SizeY >> MipIndex) / GPixelFormats[Format].BlockSizeY);
+
+					SubResourceData[SubResourceIndex].pData = &Data[DataOffset];
+					SubResourceData[SubResourceIndex].RowPitch = NumBlocksX * GPixelFormats[Format].BlockBytes;
+					SubResourceData[SubResourceIndex].SlicePitch = NumBlocksX * NumBlocksY * SubResourceData[MipIndex].RowPitch;
+
+					MipOffset += NumBlocksY * SubResourceData[MipIndex].RowPitch;
+				}
+				SliceOffset += MipOffset;
+			}
+		}
+
+#if PLATFORM_SUPPORTS_VIRTUAL_TEXTURES
+		if ((Flags & (TexCreate_Virtual | TexCreate_OfflineProcessed)) != 0)
+		{
+			RawTextureMemory = CreateVirtualTexture(SizeX, SizeY, SizeZ, NumMips, bCubeTexture, Flags, &TextureDesc, &TextureResource);
+		}
+		else
+#endif
+		{
+			FD3D12CommandListHandle& hCommandList = GetRHIDevice()->GetDefaultCommandContext().CommandListHandle;
+            FD3D12DynamicHeapAllocator& UploadHeap = GetRHIDevice()->GetDefaultUploadHeapAllocator();
+
+			SafeCreateTexture2D(GetRHIDevice(), TextureDesc, nullptr, hCommandList, UploadHeap, TextureResource.GetInitReference(), Format, Flags);
+
+			if (CreateInfo.BulkData != NULL)
+			{
+				uint64 Size = GetRequiredIntermediateSize(TextureResource->GetResource(), 0, NumMips * SizeZ);
+
+				FD3D12ResourceLocation TempResourceLocation;
+				void *pData = UploadHeap.Alloc(Size, D3D12_TEXTURE_DATA_PLACEMENT_ALIGNMENT, &TempResourceLocation);
+
+				{
+					FScopeResourceBarrier ScopeResourceBarrierDest(hCommandList, TextureResource, TextureResource->GetDefaultResourceState(), D3D12_RESOURCE_STATE_COPY_DEST, D3D12_RESOURCE_BARRIER_ALL_SUBRESOURCES);
+
+					hCommandList.GetCurrentOwningContext()->numCopies++;
+					UpdateSubresources(hCommandList.CommandList(),
+						TextureResource->GetResource(),
+						TempResourceLocation.GetResource()->GetResource(),
+						TempResourceLocation.GetOffset(),
+						0, NumMips * SizeZ,
+						SubResourceData.GetData());
+				}
+			}
+		}
+
+		if (bCreateRTV)
+		{
+			// Create a render target view for each mip
+			for (uint32 MipIndex = 0; MipIndex < NumMips; MipIndex++)
+			{
+				if ((Flags & TexCreate_TargetArraySlicesIndependently) && (bTextureArray || bCubeTexture))
+				{
+					bCreatedRTVPerSlice = true;
+
+					for (uint32 SliceIndex = 0; SliceIndex < TextureDesc.DepthOrArraySize; SliceIndex++)
+					{
+						D3D12_RENDER_TARGET_VIEW_DESC RTVDesc;
+						FMemory::Memzero(&RTVDesc, sizeof(RTVDesc));
+						RTVDesc.Format = PlatformRenderTargetFormat;
+						RTVDesc.ViewDimension = D3D12_RTV_DIMENSION_TEXTURE2DARRAY;
+						RTVDesc.Texture2DArray.FirstArraySlice = SliceIndex;
+						RTVDesc.Texture2DArray.ArraySize = 1;
+						RTVDesc.Texture2DArray.MipSlice = MipIndex;
+						RTVDesc.Texture2DArray.PlaneSlice = GetPlaneSliceFromViewFormat(PlatformResourceFormat, RTVDesc.Format);
+
+						TRefCountPtr<FD3D12RenderTargetView> RenderTargetView = new FD3D12RenderTargetView(GetRHIDevice(), &RTVDesc, TextureResource);
+						RenderTargetViews.Add(RenderTargetView);
+					}
+				}
+				else
+				{
+					D3D12_RENDER_TARGET_VIEW_DESC RTVDesc;
+					FMemory::Memzero(&RTVDesc, sizeof(RTVDesc));
+					RTVDesc.Format = PlatformRenderTargetFormat;
+					if (bTextureArray || bCubeTexture)
+					{
+						RTVDesc.ViewDimension = D3D12_RTV_DIMENSION_TEXTURE2DARRAY;
+						RTVDesc.Texture2DArray.FirstArraySlice = 0;
+                        RTVDesc.Texture2DArray.ArraySize = TextureDesc.DepthOrArraySize;
+						RTVDesc.Texture2DArray.MipSlice = MipIndex;
+						RTVDesc.Texture2DArray.PlaneSlice = GetPlaneSliceFromViewFormat(PlatformResourceFormat, RTVDesc.Format);
+					}
+					else
+					{
+						RTVDesc.ViewDimension = RenderTargetViewDimension;
+						RTVDesc.Texture2D.MipSlice = MipIndex;
+						RTVDesc.Texture2D.PlaneSlice = GetPlaneSliceFromViewFormat(PlatformResourceFormat, RTVDesc.Format);
+					}
+
+					TRefCountPtr<FD3D12RenderTargetView> RenderTargetView = new FD3D12RenderTargetView(GetRHIDevice(), &RTVDesc, TextureResource);
+					RenderTargetViews.Add(RenderTargetView);
+				}
+			}
+		}
+
+		if (bCreateDSV)
+		{
+			// Create a depth-stencil-view for the texture.
+			D3D12_DEPTH_STENCIL_VIEW_DESC DSVDesc;
+			FMemory::Memzero(&DSVDesc, sizeof(DSVDesc));
+			DSVDesc.Format = FindDepthStencilDXGIFormat(PlatformResourceFormat);
+			if (bTextureArray || bCubeTexture)
+			{
+				DSVDesc.ViewDimension = D3D12_DSV_DIMENSION_TEXTURE2DARRAY;
+				DSVDesc.Texture2DArray.FirstArraySlice = 0;
+                DSVDesc.Texture2DArray.ArraySize = TextureDesc.DepthOrArraySize;
+				DSVDesc.Texture2DArray.MipSlice = 0;
+			}
+			else
+			{
+				DSVDesc.ViewDimension = DepthStencilViewDimension;
+				DSVDesc.Texture2D.MipSlice = 0;
+			}
+
+			const bool HasStencil = HasStencilBits(DSVDesc.Format);
+			for (uint32 AccessType = 0; AccessType < FExclusiveDepthStencil::MaxIndex; ++AccessType)
+			{
+				// Create a read-only access views for the texture.
+				DSVDesc.Flags = (AccessType & FExclusiveDepthStencil::DepthRead_StencilWrite) ? D3D12_DSV_FLAG_READ_ONLY_DEPTH : D3D12_DSV_FLAG_NONE;
+				if (HasStencil)
+				{
+					DSVDesc.Flags |= (AccessType & FExclusiveDepthStencil::DepthWrite_StencilRead) ? D3D12_DSV_FLAG_READ_ONLY_STENCIL : D3D12_DSV_FLAG_NONE;
+				}
+				DepthStencilViews[AccessType] = new FD3D12DepthStencilView(GetRHIDevice(), &DSVDesc, TextureResource, HasStencil);
+			}
+		}
+	}
+	check(IsValidRef(TextureResource));
+
+	// Create a shader resource view for the texture.
+    D3D12_SHADER_RESOURCE_VIEW_DESC SRVDesc = {};
+    SRVDesc.Shader4ComponentMapping = D3D12_DEFAULT_SHADER_4_COMPONENT_MAPPING;
+
+	if (bCreateShaderResource)
+	{
+		{
+			SRVDesc.Format = PlatformShaderResourceFormat;
+
+			if (bCubeTexture && bTextureArray)
+			{
+				SRVDesc.ViewDimension = D3D12_SRV_DIMENSION_TEXTURECUBEARRAY;
+				SRVDesc.TextureCubeArray.MostDetailedMip = 0;
+				SRVDesc.TextureCubeArray.MipLevels = NumMips;
+				SRVDesc.TextureCubeArray.First2DArrayFace = 0;
+				SRVDesc.TextureCubeArray.NumCubes = SizeZ / 6;
+			}
+			else if (bCubeTexture)
+			{
+				SRVDesc.ViewDimension = D3D12_SRV_DIMENSION_TEXTURECUBE;
+				SRVDesc.TextureCube.MostDetailedMip = 0;
+				SRVDesc.TextureCube.MipLevels = NumMips;
+			}
+			else if (bTextureArray)
+			{
+				SRVDesc.ViewDimension = D3D12_SRV_DIMENSION_TEXTURE2DARRAY;
+				SRVDesc.Texture2DArray.MostDetailedMip = 0;
+				SRVDesc.Texture2DArray.MipLevels = NumMips;
+				SRVDesc.Texture2DArray.FirstArraySlice = 0;
+                SRVDesc.Texture2DArray.ArraySize = TextureDesc.DepthOrArraySize;
+				SRVDesc.Texture2DArray.PlaneSlice = GetPlaneSliceFromViewFormat(PlatformResourceFormat, SRVDesc.Format);
+			}
+			else
+			{
+				SRVDesc.ViewDimension = ShaderResourceViewDimension;
+				SRVDesc.Texture2D.MostDetailedMip = 0;
+				SRVDesc.Texture2D.MipLevels = NumMips;
+				SRVDesc.Texture2D.PlaneSlice = GetPlaneSliceFromViewFormat(PlatformResourceFormat, SRVDesc.Format);
+			}
+		}
+	}
+
+	TD3D12Texture2D<BaseResourceType>* Texture2D = new TD3D12Texture2D<BaseResourceType>(
+		GetRHIDevice(),
+		TextureResource,
+		bCreatedRTVPerSlice,
+        TextureDesc.DepthOrArraySize,
+		RenderTargetViews,
+		DepthStencilViews,
+		SizeX,
+		SizeY,
+		SizeZ,
+		NumMips,
+		ActualMSAACount,
+		(EPixelFormat)Format,
+		bCubeTexture,
+		Flags,
+		bPooledTexture,
+		CreateInfo.ClearValueBinding
+#if PLATFORM_SUPPORTS_VIRTUAL_TEXTURES
+		, RawTextureMemory
+#endif
+		);
+
+	if (Flags & TexCreate_CPUReadback)
+	{
+		const uint32 BlockBytes = GPixelFormats[PlatformResourceFormat].BlockBytes;
+		const uint32 XBytesAligned = Align((uint32)SizeX * BlockBytes, D3D12_TEXTURE_DATA_PITCH_ALIGNMENT);
+		D3D12_SUBRESOURCE_FOOTPRINT destSubresource;
+		destSubresource.Depth = SizeZ;
+		destSubresource.Height = SizeY;
+		destSubresource.Width = SizeX;
+		destSubresource.Format = PlatformResourceFormat;
+		destSubresource.RowPitch = XBytesAligned;
+
+		check(destSubresource.RowPitch % D3D12_TEXTURE_DATA_PITCH_ALIGNMENT == 0);	// Make sure we align correctly.
+
+		D3D12_PLACED_SUBRESOURCE_FOOTPRINT placedTexture2D = { 0 };
+		placedTexture2D.Offset = 0;
+		placedTexture2D.Footprint = destSubresource;
+		Texture2D->SetReadBackHeapDesc(placedTexture2D);
+	}
+
+	Texture2D->ResourceInfo.VRamAllocation = VRamAllocation;
+
+	FD3D12ShaderResourceView* pWrappedShaderResourceView =
+		bCreateShaderResource ? new FD3D12ShaderResourceView(GetRHIDevice(), &SRVDesc, Texture2D->ResourceLocation) : nullptr;
+	Texture2D->SetShaderResourceView(pWrappedShaderResourceView);
+
+	D3D11TextureAllocated(*Texture2D);
+
+	return Texture2D;
+}
+
+FD3D12Texture3D* FD3D12DynamicRHI::CreateD3D11Texture3D(uint32 SizeX, uint32 SizeY, uint32 SizeZ, uint8 Format, uint32 NumMips, uint32 Flags, FRHIResourceCreateInfo& CreateInfo)
+{
+	SCOPE_CYCLE_COUNTER(STAT_D3D12CreateTextureTime);
+
+	const bool bSRGB = (Flags & TexCreate_SRGB) != 0;
+
+	const DXGI_FORMAT PlatformResourceFormat = (DXGI_FORMAT)GPixelFormats[Format].PlatformFormat;
+	const DXGI_FORMAT PlatformShaderResourceFormat = FindShaderResourceDXGIFormat(PlatformResourceFormat, bSRGB);
+	const DXGI_FORMAT PlatformRenderTargetFormat = FindShaderResourceDXGIFormat(PlatformResourceFormat, bSRGB);
+
+	// Describe the texture.
+    D3D12_RESOURCE_DESC TextureDesc = CD3DX12_RESOURCE_DESC::Tex3D(
+        PlatformResourceFormat,
+        SizeX,
+        SizeY,
+        SizeZ,
+        NumMips);
+
+	if (Flags & TexCreate_UAV)
+	{
+        TextureDesc.Flags |= D3D12_RESOURCE_FLAG_ALLOW_UNORDERED_ACCESS;
+	}
+
+	bool bCreateRTV = false;
+
+	if (Flags & TexCreate_RenderTargetable)
+	{
+        TextureDesc.Flags |= D3D12_RESOURCE_FLAG_ALLOW_RENDER_TARGET;
+		bCreateRTV = true;
+	}
+
+	// Set up the texture bind flags.
+	check(!(Flags & TexCreate_DepthStencilTargetable));
+	check(!(Flags & TexCreate_ResolveTargetable));
+	check(Flags & TexCreate_ShaderResource);
+
+
+	TArray<D3D12_SUBRESOURCE_DATA> SubResourceData;
+
+	if (CreateInfo.BulkData)
+	{
+		uint8* Data = (uint8*)CreateInfo.BulkData->GetResourceBulkData();
+		SubResourceData.AddZeroed(NumMips);
+		uint32 MipOffset = 0;
+		for (uint32 MipIndex = 0; MipIndex < NumMips; ++MipIndex)
+		{
+			SubResourceData[MipIndex].pData = &Data[MipOffset];
+			SubResourceData[MipIndex].RowPitch = FMath::Max<uint32>(1, SizeX >> MipIndex) * GPixelFormats[Format].BlockBytes;
+			SubResourceData[MipIndex].SlicePitch = FMath::Max<uint32>(1, SizeY >> MipIndex) * SubResourceData[MipIndex].RowPitch;
+			MipOffset += FMath::Max<uint32>(1, SizeZ >> MipIndex) * SubResourceData[MipIndex].SlicePitch;
+		}
+	}
+
+	FVRamAllocation VRamAllocation;
+
+	if (FPlatformProperties::SupportsFastVRAMMemory())
+	{
+		if (Flags & TexCreate_FastVRAM)
+		{
+			VRamAllocation = FFastVRAMAllocator::GetFastVRAMAllocator()->AllocTexture3D(TextureDesc);
+		}
+	}
+
+    FD3D12CommandListHandle& hCommandList = GetRHIDevice()->GetDefaultCommandContext().CommandListHandle;
+
+	TRefCountPtr<FD3D12Resource> TextureResource;
+	VERIFYD3D11CREATETEXTURERESULT(
+		GetRHIDevice()->GetResourceHelper().CreateDefaultResource(TextureDesc, TextureResource.GetInitReference()),
+		SizeX,
+		SizeY,
+		SizeZ,
+		PlatformShaderResourceFormat,
+		NumMips,
+		TextureDesc.Flags
+		);
+	D3D12_RESOURCE_DESC const& Desc12 = TextureResource->GetDesc();
+	if (CreateInfo.BulkData != NULL)
+	{
+		uint64 Size = GetRequiredIntermediateSize(TextureResource->GetResource(), 0, NumMips);
+
+		FD3D12ResourceLocation TempResourceLocation;
+        void *pData = GetRHIDevice()->GetDefaultUploadHeapAllocator().Alloc(Size, D3D12_TEXTURE_DATA_PLACEMENT_ALIGNMENT, &TempResourceLocation);
+
+		{
+			FScopeResourceBarrier ScopeResourceBarrierDest(hCommandList, TextureResource, TextureResource->GetDefaultResourceState(), D3D12_RESOURCE_STATE_COPY_DEST, D3D12_RESOURCE_BARRIER_ALL_SUBRESOURCES);
+			
+			hCommandList.GetCurrentOwningContext()->numCopies++;
+			UpdateSubresources(
+				hCommandList.CommandList(),
+				TextureResource->GetResource(),
+				TempResourceLocation.GetResource()->GetResource(),
+				TempResourceLocation.GetOffset(),
+				0, NumMips,
+				SubResourceData.GetData());
+		}
+	}
+
+	// Create a shader resource view for the texture.
+    D3D12_SHADER_RESOURCE_VIEW_DESC SRVDesc = {};
+    SRVDesc.Shader4ComponentMapping = D3D12_DEFAULT_SHADER_4_COMPONENT_MAPPING;
+	SRVDesc.Format = PlatformShaderResourceFormat;
+	SRVDesc.ViewDimension = D3D12_SRV_DIMENSION_TEXTURE3D;
+	SRVDesc.Texture3D.MipLevels = NumMips;
+	SRVDesc.Texture3D.MostDetailedMip = 0;
+
+	TRefCountPtr<FD3D12RenderTargetView> RenderTargetView;
+	if (bCreateRTV)
+	{
+		// Create a render-target-view for the texture.
+		D3D12_RENDER_TARGET_VIEW_DESC RTVDesc;
+		FMemory::Memzero(&RTVDesc, sizeof(RTVDesc));
+		RTVDesc.Format = PlatformRenderTargetFormat;
+		RTVDesc.ViewDimension = D3D12_RTV_DIMENSION_TEXTURE3D;
+		RTVDesc.Texture3D.MipSlice = 0;
+		RTVDesc.Texture3D.FirstWSlice = 0;
+		RTVDesc.Texture3D.WSize = SizeZ;
+
+        RenderTargetView = new FD3D12RenderTargetView(GetRHIDevice(), &RTVDesc, TextureResource);
+	}
+
+	TArray<TRefCountPtr<FD3D12RenderTargetView> > RenderTargetViews;
+	RenderTargetViews.Add(RenderTargetView);
+	FD3D12Texture3D* Texture3D = new FD3D12Texture3D(GetRHIDevice(), TextureResource, RenderTargetViews, SizeX, SizeY, SizeZ, NumMips, (EPixelFormat)Format, Flags, CreateInfo.ClearValueBinding);
+
+	Texture3D->ResourceInfo.VRamAllocation = VRamAllocation;
+
+	// Create a wrapper for the SRV and set it on the texture
+	FD3D12ShaderResourceView* pShaderResourceView = new FD3D12ShaderResourceView(GetRHIDevice(), &SRVDesc, Texture3D->ResourceLocation);
+	Texture3D->SetShaderResourceView(pShaderResourceView);
+
+	D3D11TextureAllocated(*Texture3D);
+
+	return Texture3D;
+}
+
+/*-----------------------------------------------------------------------------
+	2D texture support.
+	-----------------------------------------------------------------------------*/
+
+FTexture2DRHIRef FD3D12DynamicRHI::RHICreateTexture2D(uint32 SizeX, uint32 SizeY, uint8 Format, uint32 NumMips, uint32 NumSamples, uint32 Flags, FRHIResourceCreateInfo& CreateInfo)
+{
+	return CreateD3D11Texture2D<FD3D12BaseTexture2D>(SizeX, SizeY, 1, false, false, Format, NumMips, NumSamples, Flags, CreateInfo);
+}
+
+FTexture2DRHIRef FD3D12DynamicRHI::RHIAsyncCreateTexture2D(uint32 SizeX, uint32 SizeY, uint8 Format, uint32 NumMips, uint32 Flags, void** InitialMipData, uint32 NumInitialMips)
+{
+	FD3D12Texture2D* NewTexture = NULL;
+	TRefCountPtr<FD3D12Resource> TextureResource;
+	D3D12_RESOURCE_DESC TextureDesc;
+	FMemory::Memset(&TextureDesc, 0, sizeof(TextureDesc));
+    D3D12_SHADER_RESOURCE_VIEW_DESC SRVDesc = {};
+
+	D3D12_SUBRESOURCE_DATA SubResourceData[MAX_TEXTURE_MIP_COUNT];
+	FPlatformMemory::Memzero(SubResourceData, sizeof(D3D12_SUBRESOURCE_DATA) * MAX_TEXTURE_MIP_COUNT);
+
+	uint32 InvalidFlags = TexCreate_RenderTargetable | TexCreate_ResolveTargetable | TexCreate_DepthStencilTargetable | TexCreate_GenerateMipCapable | TexCreate_UAV | TexCreate_Presentable | TexCreate_CPUReadback;
+	TArray<TRefCountPtr<FD3D12RenderTargetView> > RenderTargetViews;
+
+	check(GRHISupportsAsyncTextureCreation);
+	check((Flags & InvalidFlags) == 0);
+
+	if (GMaxRHIFeatureLevel <= ERHIFeatureLevel::ES2)
+	{
+		// Remove sRGB read flag when not supported
+		Flags &= ~TexCreate_SRGB;
+	}
+
+	const DXGI_FORMAT PlatformResourceFormat = (DXGI_FORMAT)GPixelFormats[Format].PlatformFormat;
+	const DXGI_FORMAT PlatformShaderResourceFormat = FindShaderResourceDXGIFormat(PlatformResourceFormat, ((Flags & TexCreate_SRGB) != 0));
+
+    TextureDesc = CD3DX12_RESOURCE_DESC::Tex2D(
+        PlatformResourceFormat,
+        SizeX,
+        SizeY,
+        1,
+        NumMips,
+        1,  // Sample count
+        0);  // Sample quality
+
+	for (uint32 MipIndex = 0; MipIndex < NumInitialMips; ++MipIndex)
+	{
+		uint32 NumBlocksX = FMath::Max<uint32>(1, (SizeX >> MipIndex) / GPixelFormats[Format].BlockSizeX);
+		uint32 NumBlocksY = FMath::Max<uint32>(1, (SizeY >> MipIndex) / GPixelFormats[Format].BlockSizeY);
+
+		SubResourceData[MipIndex].pData = InitialMipData[MipIndex];
+		SubResourceData[MipIndex].RowPitch = NumBlocksX * GPixelFormats[Format].BlockBytes;
+		SubResourceData[MipIndex].SlicePitch = NumBlocksX * NumBlocksY * GPixelFormats[Format].BlockBytes;
+	}
+
+	void* TempBuffer = ZeroBuffer;
+	uint32 TempBufferSize = ZeroBufferSize;
+	for (uint32 MipIndex = NumInitialMips; MipIndex < NumMips; ++MipIndex)
+	{
+		uint32 NumBlocksX = FMath::Max<uint32>(1, (SizeX >> MipIndex) / GPixelFormats[Format].BlockSizeX);
+		uint32 NumBlocksY = FMath::Max<uint32>(1, (SizeY >> MipIndex) / GPixelFormats[Format].BlockSizeY);
+		uint32 MipSize = NumBlocksX * NumBlocksY * GPixelFormats[Format].BlockBytes;
+
+		if (MipSize > TempBufferSize)
+		{
+			UE_LOG(LogD3D12RHI, Warning, TEXT("Temp texture streaming buffer not large enough, needed %d bytes"), MipSize);
+			check(TempBufferSize == ZeroBufferSize);
+			TempBufferSize = MipSize;
+			TempBuffer = FMemory::Malloc(TempBufferSize);
+			FMemory::Memzero(TempBuffer, TempBufferSize);
+		}
+
+		SubResourceData[MipIndex].pData = TempBuffer;
+		SubResourceData[MipIndex].RowPitch = NumBlocksX * GPixelFormats[Format].BlockBytes;
+		SubResourceData[MipIndex].SlicePitch = MipSize;
+	}
+
+    if (bForceSingleQueueGPU || !SubResourceData)
+    {
+		FD3D12DynamicHeapAllocator& DynamicHeapAllocator = GetHelperThreadDynamicUploadHeapAllocator();
+		FD3D12CommandListHandle hCommandList = GetRHIDevice()->GetCommandListManager().BeginCommandList();
+		hCommandList.SetCurrentOwningContext(&GetRHIDevice()->GetDefaultCommandContext());
+
+	    DynamicHeapAllocator.SetCurrentCommandListHandle(hCommandList);
+		SafeCreateTexture2D(GetRHIDevice(), TextureDesc, SubResourceData, hCommandList, DynamicHeapAllocator, TextureResource.GetInitReference(), Format, Flags);
+
+		// Wait for the copy context to finish before continuing as this function is only expected to return once all the texture streaming has finished.
+		// Execute command list which contains CopySubresourceRegion call to upload initial texture data.
+		hCommandList.Close();
+		GetRHIDevice()->GetCommandListManager().ExecuteCommandList(hCommandList, true);
+    }
+	else
+	{
+		FD3D12DynamicHeapAllocator& DynamicHeapAllocator = GetHelperThreadDynamicUploadHeapAllocator();
+		FD3D12CommandListHandle hCopyCommandList = GetRHIDevice()->GetCopyCommandListManager().BeginCommandList();
+		hCopyCommandList.SetCurrentOwningContext(&GetRHIDevice()->GetDefaultCommandContext());
+
+        DynamicHeapAllocator.SetCurrentCommandListHandle(hCopyCommandList);
+		SafeCreateTexture2D(GetRHIDevice(), TextureDesc, SubResourceData, hCopyCommandList, DynamicHeapAllocator, TextureResource.GetInitReference(), Format, Flags);
+		
+		// Wait for the copy context to finish before continuing as this function is only expected to return once all the texture streaming has finished.
+		hCopyCommandList.Close();
+		GetRHIDevice()->GetCopyCommandListManager().ExecuteCommandList(hCopyCommandList, true);
+    }
+
+	if (TempBufferSize != ZeroBufferSize)
+	{
+		FMemory::Free(TempBuffer);
+	}
+
+    SRVDesc.Shader4ComponentMapping = D3D12_DEFAULT_SHADER_4_COMPONENT_MAPPING;
+	SRVDesc.Format = PlatformShaderResourceFormat;
+	SRVDesc.ViewDimension = D3D12_SRV_DIMENSION_TEXTURE2D;
+	SRVDesc.Texture2D.MostDetailedMip = 0;
+	SRVDesc.Texture2D.MipLevels = NumMips;
+	SRVDesc.Texture2D.PlaneSlice = GetPlaneSliceFromViewFormat(PlatformResourceFormat, SRVDesc.Format);
+
+	NewTexture = new FD3D12Texture2D(
+		GetRHIDevice(),
+		TextureResource,
+		false,
+		1,
+		RenderTargetViews,
+		/*DepthStencilViews=*/ NULL,
+		SizeX,
+		SizeY,
+		0,
+		NumMips,
+		/*ActualMSAACount=*/ 1,
+		(EPixelFormat)Format,
+		/*bInCubemap=*/ false,
+		Flags,
+		/*bPooledTexture=*/ false,
+		FClearValueBinding()
+		);
+
+	// Create a wrapper for the SRV and set it on the texture
+	FD3D12ShaderResourceView* pShaderResourceView = new FD3D12ShaderResourceView(GetRHIDevice(), &SRVDesc, NewTexture->ResourceLocation);
+	NewTexture->SetShaderResourceView(pShaderResourceView);
+
+	D3D11TextureAllocated(*NewTexture);
+
+	return NewTexture;
+}
+
+void FD3D12DynamicRHI::RHICopySharedMips(FTexture2DRHIParamRef DestTexture2DRHI, FTexture2DRHIParamRef SrcTexture2DRHI)
+{
+	FD3D12Texture2D*  DestTexture2D = FD3D12DynamicRHI::ResourceCast(DestTexture2DRHI);
+	FD3D12Texture2D*  SrcTexture2D = FD3D12DynamicRHI::ResourceCast(SrcTexture2DRHI);
+
+	// Use the GPU to asynchronously copy the old mip-maps into the new texture.
+	const uint32 NumSharedMips = FMath::Min(DestTexture2D->GetNumMips(), SrcTexture2D->GetNumMips());
+	const uint32 SourceMipOffset = SrcTexture2D->GetNumMips() - NumSharedMips;
+	const uint32 DestMipOffset = DestTexture2D->GetNumMips() - NumSharedMips;
+
+    uint32 srcSubresource = 0;
+    uint32 destSubresource = 0;
+
+	FD3D12CommandListHandle& hCommandList = GetRHIDevice()->GetDefaultCommandContext().CommandListHandle;
+
+#if SUPPORTS_MEMORY_RESIDENCY
+	DestTexture2D->GetResource()->UpdateResidency();
+	SrcTexture2D->GetResource()->UpdateResidency();
+#endif
+	{
+		FScopeResourceBarrier ScopeResourceBarrierDest(hCommandList, DestTexture2D->GetResource(), DestTexture2D->GetResource()->GetDefaultResourceState(), D3D12_RESOURCE_STATE_COPY_DEST, D3D12_RESOURCE_BARRIER_ALL_SUBRESOURCES);
+		FScopeResourceBarrier ScopeResourceBarrierSrc(hCommandList, SrcTexture2D->GetResource(), SrcTexture2D->GetResource()->GetDefaultResourceState(), D3D12_RESOURCE_STATE_COPY_SOURCE, D3D12_RESOURCE_BARRIER_ALL_SUBRESOURCES);
+
+		for (uint32 MipIndex = 0; MipIndex < NumSharedMips; ++MipIndex)
+		{
+			// Use the GPU to copy between mip-maps.
+			srcSubresource = CalcSubresource(MipIndex + SourceMipOffset, 0, SrcTexture2D->GetNumMips());
+			destSubresource = CalcSubresource(MipIndex + DestMipOffset, 0, DestTexture2D->GetNumMips());
+
+			CD3DX12_TEXTURE_COPY_LOCATION DestCopyLocation(DestTexture2D->GetResource()->GetResource(), destSubresource);
+			CD3DX12_TEXTURE_COPY_LOCATION SourceCopyLocation(SrcTexture2D->GetResource()->GetResource(), srcSubresource);
+
+			GetRHIDevice()->GetDefaultCommandContext().numCopies++;
+			hCommandList->CopyTextureRegion(
+				&DestCopyLocation,
+				0, 0, 0,
+				&SourceCopyLocation,
+				nullptr);
+		}
+	}
+
+	DEBUG_RHI_EXECUTE_COMMAND_LIST(this);
+}
+
+FTexture2DArrayRHIRef FD3D12DynamicRHI::RHICreateTexture2DArray(uint32 SizeX, uint32 SizeY, uint32 SizeZ, uint8 Format, uint32 NumMips, uint32 Flags, FRHIResourceCreateInfo& CreateInfo)
+{
+	check(SizeZ >= 1);
+	return CreateD3D11Texture2D<FD3D12BaseTexture2DArray>(SizeX, SizeY, SizeZ, true, false, Format, NumMips, 1, Flags, CreateInfo);
+}
+
+FTexture3DRHIRef FD3D12DynamicRHI::RHICreateTexture3D(uint32 SizeX, uint32 SizeY, uint32 SizeZ, uint8 Format, uint32 NumMips, uint32 Flags, FRHIResourceCreateInfo& CreateInfo)
+{
+	check(SizeZ >= 1);
+	return CreateD3D11Texture3D(SizeX, SizeY, SizeZ, Format, NumMips, Flags, CreateInfo);
+}
+
+void FD3D12DynamicRHI::RHIGetResourceInfo(FTextureRHIParamRef Ref, FRHIResourceInfo& OutInfo)
+{
+	if (Ref)
+	{
+		OutInfo = Ref->ResourceInfo;
+	}
+}
+
+FShaderResourceViewRHIRef FD3D12DynamicRHI::RHICreateShaderResourceView(FTexture2DRHIParamRef Texture2DRHI, uint8 MipLevel)
+{
+	FD3D12Texture2D*  Texture2D = FD3D12DynamicRHI::ResourceCast(Texture2DRHI);
+
+	uint64 Offset = 0;
+    const D3D12_RESOURCE_DESC& TextureDesc = Texture2D->GetResource()->GetDesc();
+
+	// Create a Shader Resource View
+    D3D12_SHADER_RESOURCE_VIEW_DESC SRVDesc = {};
+    SRVDesc.Shader4ComponentMapping = D3D12_DEFAULT_SHADER_4_COMPONENT_MAPPING;
+	SRVDesc.ViewDimension = D3D12_SRV_DIMENSION_TEXTURE2D;
+	SRVDesc.Texture2D.MostDetailedMip = MipLevel;
+	SRVDesc.Texture2D.MipLevels = 1;
+
+	SRVDesc.Format = TextureDesc.Format;
+
+	SRVDesc.Texture2D.PlaneSlice = GetPlaneSliceFromViewFormat(TextureDesc.Format, SRVDesc.Format);
+
+	return new FD3D12ShaderResourceView(GetRHIDevice(), &SRVDesc, Texture2D->ResourceLocation);
+}
+
+FShaderResourceViewRHIRef FD3D12DynamicRHI::RHICreateShaderResourceView(FTexture2DArrayRHIParamRef Texture2DRHI, uint8 MipLevel)
+{
+	check(0);
+#if 0
+	FD3D12Texture*  Texture2D = FD3D12DynamicRHI::ResourceCast(Texture2DRHI);
+
+	uint64 Offset = 0;
+	const D3D12_RESOURCE_DESC& TextureDesc = Texture2D->GetResource()->GetDesc();
+
+	// Create a Shader Resource View
+	D3D12_SHADER_RESOURCE_VIEW_DESC SRVDesc = {};
+	SRVDesc.Shader4ComponentMapping = D3D12_DEFAULT_SHADER_4_COMPONENT_MAPPING;
+	SRVDesc.ViewDimension = D3D12_SRV_DIMENSION_TEXTURE2D;
+	SRVDesc.Texture2D.MostDetailedMip = MipLevel;
+	SRVDesc.Texture2D.MipLevels = 1;
+
+	SRVDesc.Format = TextureDesc.Format;
+
+	SRVDesc.Texture2D.PlaneSlice = GetPlaneSliceFromViewFormat(TextureDesc.Format, SRVDesc.Format);
+
+	return new FD3D12ShaderResourceView(GetRHIDevice(), &SRVDesc, Texture2D->ResourceLocation);
+#else
+	return nullptr;
+#endif
+}
+
+FShaderResourceViewRHIRef FD3D12DynamicRHI::RHICreateShaderResourceView(FTextureCubeRHIParamRef TextureCubeRHI, uint8 MipLevel)
+{
+	check(0);
+#if 0
+	FD3D12Texture*  Texture2D = FD3D12DynamicRHI::ResourceCast(Texture2DRHI);
+
+	uint64 Offset = 0;
+	const D3D12_RESOURCE_DESC& TextureDesc = Texture2D->GetResource()->GetDesc();
+
+	// Create a Shader Resource View
+	D3D12_SHADER_RESOURCE_VIEW_DESC SRVDesc = {};
+	SRVDesc.Shader4ComponentMapping = D3D12_DEFAULT_SHADER_4_COMPONENT_MAPPING;
+	SRVDesc.ViewDimension = D3D12_SRV_DIMENSION_TEXTURE2D;
+	SRVDesc.Texture2D.MostDetailedMip = MipLevel;
+	SRVDesc.Texture2D.MipLevels = 1;
+
+	SRVDesc.Format = TextureDesc.Format;
+
+	SRVDesc.Texture2D.PlaneSlice = GetPlaneSliceFromViewFormat(TextureDesc.Format, SRVDesc.Format);
+
+	return new FD3D12ShaderResourceView(GetRHIDevice(), &SRVDesc, Texture2D->ResourceLocation);
+#else
+	return nullptr;
+#endif
+}
+
+FShaderResourceViewRHIRef FD3D12DynamicRHI::RHICreateShaderResourceView(FTexture2DRHIParamRef Texture2DRHI, uint8 MipLevel, uint8 NumMipLevels, uint8 Format)
+{
+	FD3D12Texture2D*  Texture2D = FD3D12DynamicRHI::ResourceCast(Texture2DRHI);
+	uint64 Offset = 0;
+    const D3D12_RESOURCE_DESC& TextureDesc = Texture2D->GetResource()->GetDesc();
+
+	const DXGI_FORMAT PlatformResourceFormat = FD3D12DynamicRHI::GetPlatformTextureResourceFormat((DXGI_FORMAT)GPixelFormats[Format].PlatformFormat, Texture2D->GetFlags());
+
+	bool bSRGB = (Texture2D->GetFlags() & TexCreate_SRGB) != 0;
+	const DXGI_FORMAT PlatformShaderResourceFormat = FindShaderResourceDXGIFormat(PlatformResourceFormat, bSRGB);
+
+	// Create a Shader Resource View
+    D3D12_SHADER_RESOURCE_VIEW_DESC SRVDesc = {};
+    SRVDesc.Shader4ComponentMapping = D3D12_DEFAULT_SHADER_4_COMPONENT_MAPPING;
+
+	if (TextureDesc.SampleDesc.Count > 1)
+	{
+		// MS textures can't have mips apparently, so nothing else to set.
+		SRVDesc.ViewDimension = D3D12_SRV_DIMENSION_TEXTURE2DMS;
+	}
+	else
+	{
+		SRVDesc.ViewDimension = D3D12_SRV_DIMENSION_TEXTURE2D;
+		SRVDesc.Texture2D.MostDetailedMip = MipLevel;
+		SRVDesc.Texture2D.MipLevels = NumMipLevels;
+		SRVDesc.Texture2D.PlaneSlice = GetPlaneSliceFromViewFormat(PlatformResourceFormat, PlatformShaderResourceFormat);
+	}
+
+	SRVDesc.Format = PlatformShaderResourceFormat;
+	
+	return new FD3D12ShaderResourceView(GetRHIDevice(), &SRVDesc, Texture2D->ResourceLocation);
+}
+
+FShaderResourceViewRHIRef FD3D12DynamicRHI::RHICreateShaderResourceView(FTexture3DRHIParamRef Texture3DRHI, uint8 MipLevel)
+{
+	FD3D12Texture3D*  Texture3D = FD3D12DynamicRHI::ResourceCast(Texture3DRHI);
+check(0);
+#if 0
+	uint64 Offset = 0;
+	const D3D12_RESOURCE_DESC& TextureDesc = Texture2D->GetResource()->GetDesc();
+
+	// Create a Shader Resource View
+	D3D12_SHADER_RESOURCE_VIEW_DESC SRVDesc = {};
+	SRVDesc.Shader4ComponentMapping = D3D12_DEFAULT_SHADER_4_COMPONENT_MAPPING;
+	SRVDesc.ViewDimension = D3D12_SRV_DIMENSION_TEXTURE3D;
+	SRVDesc.Texture3D.MostDetailedMip = MipLevel;
+	SRVDesc.Texture3D.MipLevels = 1;
+
+	SRVDesc.Format = TextureDesc.Format;
+
+	//SRVDesc.Texture3D.PlaneSlice = GetPlaneSliceFromViewFormat(TextureDesc.Format, SRVDesc.Format);
+
+	return new FD3D12ShaderResourceView(GetRHIDevice(), &SRVDesc, Texture3D->ResourceLocation);
+#else
+	return nullptr;
+#endif
+}
+
+/** Generates mip maps for the surface. */
+void FD3D12DynamicRHI::RHIGenerateMips(FTextureRHIParamRef TextureRHI)
+{
+	// MS: GenerateMips has been removed in D3D12.  However, this code path isn't executed in 
+	// available UE content, so there is no need to re-implement GenerateMips for now.
+	FD3D12TextureBase* Texture = GetD3D11TextureFromRHITexture(TextureRHI);
+	// Surface must have been created with D3D11_BIND_RENDER_TARGET for GenerateMips to work
+	check(Texture->GetShaderResourceView() && Texture->GetRenderTargetView(0, -1));
+	GPUProfilingData.RegisterGPUWork(0);
+}
+
+/**
+ * Computes the size in memory required by a given texture.
+ *
+ * @param	TextureRHI		- Texture we want to know the size of
+ * @return					- Size in Bytes
+ */
+uint32 FD3D12DynamicRHI::RHIComputeMemorySize(FTextureRHIParamRef TextureRHI)
+{
+	if (!TextureRHI)
+	{
+		return 0;
+	}
+
+	FD3D12TextureBase* Texture = GetD3D11TextureFromRHITexture(TextureRHI);
+	return Texture->GetMemorySize();
+}
+
+/**
+ * Starts an asynchronous texture reallocation. It may complete immediately if the reallocation
+ * could be performed without any reshuffling of texture memory, or if there isn't enough memory.
+ * The specified status counter will be decremented by 1 when the reallocation is complete (success or failure).
+ *
+ * Returns a new reference to the texture, which will represent the new mip count when the reallocation is complete.
+ * RHIGetAsyncReallocateTexture2DStatus() can be used to check the status of an ongoing or completed reallocation.
+ *
+ * @param Texture2D		- Texture to reallocate
+ * @param NewMipCount	- New number of mip-levels
+ * @param NewSizeX		- New width, in pixels
+ * @param NewSizeY		- New height, in pixels
+ * @param RequestStatus	- Will be decremented by 1 when the reallocation is complete (success or failure).
+ * @return				- New reference to the texture, or an invalid reference upon failure
+ */
+FTexture2DRHIRef FD3D12DynamicRHI::RHIAsyncReallocateTexture2D(FTexture2DRHIParamRef Texture2DRHI, int32 NewMipCount, int32 NewSizeX, int32 NewSizeY, FThreadSafeCounter* RequestStatus)
+{
+	FD3D12Texture2D*  Texture2D = FD3D12DynamicRHI::ResourceCast(Texture2DRHI);
+    	
+	// Allocate a new texture.
+	FRHIResourceCreateInfo CreateInfo;
+	FD3D12Texture2D* NewTexture2D = CreateD3D11Texture2D<FD3D12BaseTexture2D>(NewSizeX, NewSizeY, 1, false, false, Texture2D->GetFormat(), NewMipCount, 1, Texture2D->GetFlags(), CreateInfo);
+
+	// Use the GPU to asynchronously copy the old mip-maps into the new texture.
+	const uint32 NumSharedMips = FMath::Min(Texture2D->GetNumMips(), NewTexture2D->GetNumMips());
+	const uint32 SourceMipOffset = Texture2D->GetNumMips() - NumSharedMips;
+	const uint32 DestMipOffset = NewTexture2D->GetNumMips() - NumSharedMips;
+
+    uint32 destSubresource = 0;
+    uint32 srcSubresource = 0;
+
+	FD3D12CommandListHandle& hCommandList = GetRHIDevice()->GetDefaultCommandContext().CommandListHandle;
+
+	{
+		FScopeResourceBarrier ScopeResourceBarrierDest(hCommandList, NewTexture2D->GetResource(), NewTexture2D->GetResource()->GetDefaultResourceState(), D3D12_RESOURCE_STATE_COPY_DEST, D3D12_RESOURCE_BARRIER_ALL_SUBRESOURCES);
+		FScopeResourceBarrier ScopeResourceBarrierSource(hCommandList, Texture2D->GetResource(), Texture2D->GetResource()->GetDefaultResourceState(), D3D12_RESOURCE_STATE_COPY_SOURCE, D3D12_RESOURCE_BARRIER_ALL_SUBRESOURCES);
+
+		for (uint32 MipIndex = 0; MipIndex < NumSharedMips; ++MipIndex)
+		{
+			// Use the GPU to copy between mip-maps.
+			// This is serialized with other D3D commands, so it isn't necessary to increment Counter to signal a pending asynchronous copy.
+
+			srcSubresource = CalcSubresource(MipIndex + SourceMipOffset, 0, Texture2D->GetNumMips());
+			destSubresource = CalcSubresource(MipIndex + DestMipOffset, 0, NewTexture2D->GetNumMips());
+
+			CD3DX12_TEXTURE_COPY_LOCATION DestCopyLocation(NewTexture2D->GetResource()->GetResource(), destSubresource);
+			CD3DX12_TEXTURE_COPY_LOCATION SourceCopyLocation(Texture2D->GetResource()->GetResource(), srcSubresource);
+
+			GetRHIDevice()->GetDefaultCommandContext().numCopies++;
+			hCommandList->CopyTextureRegion(
+				&DestCopyLocation,
+				0, 0, 0,
+				&SourceCopyLocation,
+				nullptr);
+		}
+	}
+
+	DEBUG_RHI_EXECUTE_COMMAND_LIST(this);
+
+	// Decrement the thread-safe counter used to track the completion of the reallocation, since D3D handles sequencing the
+	// async mip copies with other D3D calls.
+	RequestStatus->Decrement();
+
+	return NewTexture2D;
+}
+
+/**
+ * Returns the status of an ongoing or completed texture reallocation:
+ *	TexRealloc_Succeeded	- The texture is ok, reallocation is not in progress.
+ *	TexRealloc_Failed		- The texture is bad, reallocation is not in progress.
+ *	TexRealloc_InProgress	- The texture is currently being reallocated async.
+ *
+ * @param Texture2D		- Texture to check the reallocation status for
+ * @return				- Current reallocation status
+ */
+ETextureReallocationStatus FD3D12DynamicRHI::RHIFinalizeAsyncReallocateTexture2D(FTexture2DRHIParamRef Texture2D, bool bBlockUntilCompleted)
+{
+	return TexRealloc_Succeeded;
+}
+
+/**
+ * Cancels an async reallocation for the specified texture.
+ * This should be called for the new texture, not the original.
+ *
+ * @param Texture				Texture to cancel
+ * @param bBlockUntilCompleted	If true, blocks until the cancellation is fully completed
+ * @return						Reallocation status
+ */
+ETextureReallocationStatus FD3D12DynamicRHI::RHICancelAsyncReallocateTexture2D(FTexture2DRHIParamRef Texture2D, bool bBlockUntilCompleted)
+{
+	return TexRealloc_Succeeded;
+}
+
+template<typename RHIResourceType>
+void* TD3D12Texture2D<RHIResourceType>::Lock(uint32 MipIndex, uint32 ArrayIndex, EResourceLockMode LockMode, uint32& DestStride)
+{
+	SCOPE_CYCLE_COUNTER(STAT_D3D12LockTextureTime);
+
+	// Calculate the subresource index corresponding to the specified mip-map.
+	const uint32 Subresource = CalcSubresource(MipIndex, ArrayIndex, GetNumMips());
+
+	// Calculate the dimensions of the mip-map.
+	const uint32 BlockSizeX = GPixelFormats[GetFormat()].BlockSizeX;
+	const uint32 BlockSizeY = GPixelFormats[GetFormat()].BlockSizeY;
+	const uint32 BlockBytes = GPixelFormats[GetFormat()].BlockBytes;
+	const uint32 MipSizeX = FMath::Max(GetSizeX() >> MipIndex, BlockSizeX);
+	const uint32 MipSizeY = FMath::Max(GetSizeY() >> MipIndex, BlockSizeY);
+	const uint32 NumBlocksX = (MipSizeX + BlockSizeX - 1) / BlockSizeX;
+	const uint32 NumBlocksY = (MipSizeY + BlockSizeY - 1) / BlockSizeY;
+
+	const uint32 XBytesAligned = Align(NumBlocksX * BlockBytes, D3D12_TEXTURE_DATA_PITCH_ALIGNMENT);
+	const uint32 MipBytesAligned = XBytesAligned * NumBlocksY;
+
+	FD3D12LockedData LockedData;
+#if	PLATFORM_SUPPORTS_VIRTUAL_TEXTURES
+	if (D3DRHI->HandleSpecialLock(LockedData, MipIndex, ArrayIndex, GetFlags(), LockMode, GetResource(), RawTextureMemory, GetNumMips(), DestStride))
+	{
+		// nothing left to do...
+	}
+	else
+#endif
+		if (LockMode == RLM_WriteOnly)
+		{
+		// If we're writing to the texture, allocate a system memory buffer to receive the new contents.
+		// Use an upload heap to copy data to a default resource.
+		const uint32 bufferSize = Align(MipBytesAligned, D3D12_TEXTURE_DATA_PLACEMENT_ALIGNMENT);
+
+		TRefCountPtr<FD3D12ResourceLocation> pUploadHeapResourceLocation = new FD3D12ResourceLocation();
+        void *pData = GetParentDevice()->GetDefaultUploadHeapAllocator().Alloc(bufferSize, D3D12_TEXTURE_DATA_PLACEMENT_ALIGNMENT, pUploadHeapResourceLocation);
+		if (nullptr == pData)
+		{
+			check(false);
+			return nullptr;
+		}
+
+		// Add the lock to the lock map.
+		LockedData.SetData(pData);
+		LockedData.Pitch = DestStride = XBytesAligned;
+		check(LockedData.Pitch % D3D12_TEXTURE_DATA_PITCH_ALIGNMENT == 0);
+
+		// Keep track of the underlying resource for the upload heap so it can be referenced during Unmap.
+		LockedData.UploadHeapLocation = pUploadHeapResourceLocation;
+		}
+		else
+		{
+			// If we're reading from the texture, we create a staging resource, copy the texture contents to it, and map it.
+
+			// Create the staging texture.
+            const D3D12_RESOURCE_DESC& StagingTextureDesc = GetResource()->GetDesc();
+			TRefCountPtr<FD3D12Resource> StagingTexture;
+            VERIFYD3D11RESULT(GetParentDevice()->GetResourceHelper().CreateBuffer(D3D12_HEAP_TYPE_READBACK, MipBytesAligned, StagingTexture.GetInitReference()));
+			LockedData.StagingResource = StagingTexture;
+
+			// Copy the mip-map data from the real resource into the staging resource
+			D3D12_SUBRESOURCE_FOOTPRINT destSubresource;
+			destSubresource.Depth = 1;
+			destSubresource.Height = MipSizeY;
+			destSubresource.Width = MipSizeX;
+            destSubresource.Format = StagingTextureDesc.Format;
+			destSubresource.RowPitch = XBytesAligned;
+			check(destSubresource.RowPitch % D3D12_TEXTURE_DATA_PITCH_ALIGNMENT == 0);
+
+			D3D12_PLACED_SUBRESOURCE_FOOTPRINT placedTexture2D = { 0 };
+			placedTexture2D.Offset = 0;
+			placedTexture2D.Footprint = destSubresource;
+
+            CD3DX12_TEXTURE_COPY_LOCATION DestCopyLocation(StagingTexture->GetResource(), placedTexture2D);
+            CD3DX12_TEXTURE_COPY_LOCATION SourceCopyLocation(GetResource()->GetResource(), Subresource);
+
+			FD3D12CommandListHandle& hCommandList = GetParentDevice()->GetDefaultCommandContext().CommandListHandle;
+
+			{
+				FScopeResourceBarrier ScopeResourceBarrierSource(hCommandList, GetResource(), GetResource()->GetDefaultResourceState(), D3D12_RESOURCE_STATE_COPY_SOURCE, SourceCopyLocation.SubresourceIndex);
+
+				GetParentDevice()->GetDefaultCommandContext().numCopies++;
+				hCommandList->CopyTextureRegion(
+					&DestCopyLocation,
+					0, 0, 0,
+					&SourceCopyLocation,
+					nullptr);
+			}
+
+			// We need to execute the command list so we can read the data from the map below
+           GetParentDevice()->GetDefaultCommandContext().FlushCommands(true);
+
+			// Map the staging resource, and return the mapped address.
+			void* pData;
+			VERIFYD3D11RESULT(StagingTexture->GetResource()->Map(0, nullptr, &pData));
+			LockedData.SetData(pData);
+			LockedData.Pitch = DestStride = XBytesAligned;
+			check(LockedData.Pitch % D3D12_TEXTURE_DATA_PITCH_ALIGNMENT == 0);
+		}
+
+	check(DestStride == XBytesAligned);
+
+	// Add the lock to the outstanding lock list.
+	GetParentDevice()->GetOwningRHI()->OutstandingLocks.Add(FD3D12LockedKey(GetResource(), Subresource), LockedData);
+
+	return (void*)LockedData.GetData();
+}
+
+template<typename RHIResourceType>
+void TD3D12Texture2D<RHIResourceType>::Unlock(uint32 MipIndex, uint32 ArrayIndex)
+{
+	SCOPE_CYCLE_COUNTER(STAT_D3D12UnlockTextureTime);
+
+	// Calculate the subresource index corresponding to the specified mip-map.
+	const uint32 Subresource = CalcSubresource(MipIndex, ArrayIndex, GetNumMips());
+
+	// Calculate the dimensions of the mip-map.
+	const uint32 BlockSizeX = GPixelFormats[GetFormat()].BlockSizeX;
+	const uint32 BlockSizeY = GPixelFormats[GetFormat()].BlockSizeY;
+	const uint32 BlockBytes = GPixelFormats[GetFormat()].BlockBytes;
+	const uint32 MipSizeX = FMath::Max(GetSizeX() >> MipIndex, BlockSizeX);
+	const uint32 MipSizeY = FMath::Max(GetSizeY() >> MipIndex, BlockSizeY);
+
+	// Find the object that is tracking this lock
+	const FD3D12LockedKey LockedKey(GetResource(), Subresource);
+    FD3D12LockedData* LockedData = GetParentDevice()->GetOwningRHI()->OutstandingLocks.Find(LockedKey);
+	check(LockedData);
+
+#if PLATFORM_SUPPORTS_VIRTUAL_TEXTURES
+	if (D3DRHI->HandleSpecialUnlock(MipIndex, GetFlags(), GetResource(), RawTextureMemory))
+	{
+		// nothing left to do...
+	}
+    else
+#endif
+    {
+        if (!LockedData->StagingResource)
+        {
+            check(LockedData->UploadHeapLocation != nullptr);
+
+            // Copy the contents of the temporary memory buffer allocated for writing into the VB.
+            FD3D12ResourceLocation* pUploadHeapLocation = LockedData->UploadHeapLocation.GetReference();
+            FD3D12Resource* UploadHeap = pUploadHeapLocation->GetResource();
+
+            // Copy the mip-map data from the real resource into the staging resource
+            const D3D12_RESOURCE_DESC& resourceDesc = GetResource()->GetDesc();
+            D3D12_SUBRESOURCE_FOOTPRINT bufferPitchDesc;
+            bufferPitchDesc.Depth = 1;
+            bufferPitchDesc.Height = MipSizeY;
+            bufferPitchDesc.Width = MipSizeX;
+            bufferPitchDesc.Format = resourceDesc.Format;
+            bufferPitchDesc.RowPitch = LockedData->Pitch;
+            check(bufferPitchDesc.RowPitch % D3D12_TEXTURE_DATA_PITCH_ALIGNMENT == 0);
+
+            D3D12_PLACED_SUBRESOURCE_FOOTPRINT placedTexture2D = { 0 };
+            placedTexture2D.Offset = pUploadHeapLocation->GetOffset();
+            placedTexture2D.Footprint = bufferPitchDesc;
+
+            CD3DX12_TEXTURE_COPY_LOCATION DestCopyLocation(GetResource()->GetResource(), Subresource);
+            CD3DX12_TEXTURE_COPY_LOCATION SourceCopyLocation(UploadHeap->GetResource(), placedTexture2D);
+
+#if SUPPORTS_MEMORY_RESIDENCY
+            GetResource()->UpdateResidency();
+#endif
+
+			FD3D12CommandListHandle& hCommandList = GetParentDevice()->GetDefaultCommandContext().CommandListHandle;
+
+			{
+				FScopeResourceBarrier ScopeResourceBarrierDest(hCommandList, GetResource(), GetResource()->GetDefaultResourceState(), D3D12_RESOURCE_STATE_COPY_DEST, DestCopyLocation.SubresourceIndex);
+				// Don't need to transition upload heaps
+
+				GetParentDevice()->GetDefaultCommandContext().numCopies++;
+				hCommandList->CopyTextureRegion(
+					&DestCopyLocation,
+					0, 0, 0,
+					&SourceCopyLocation,
+					nullptr);
+			}
+
+            DEBUG_EXECUTE_COMMAND_CONTEXT(GetParentDevice()->GetDefaultCommandContext());
+        }
+    }
+
+	// Remove the lock from the outstanding lock list.
+    GetParentDevice()->GetOwningRHI()->OutstandingLocks.Remove(LockedKey);
+}
+
+void* FD3D12DynamicRHI::RHILockTexture2D(FTexture2DRHIParamRef TextureRHI, uint32 MipIndex, EResourceLockMode LockMode, uint32& DestStride, bool bLockWithinMiptail)
+{
+	check(TextureRHI);
+	FD3D12Texture2D*  Texture = FD3D12DynamicRHI::ResourceCast(TextureRHI);
+    GetRHIDevice()->GetDefaultCommandContext().ConditionalClearShaderResource(Texture->ResourceLocation);
+	return Texture->Lock(MipIndex, 0, LockMode, DestStride);
+}
+
+void FD3D12DynamicRHI::RHIUnlockTexture2D(FTexture2DRHIParamRef TextureRHI, uint32 MipIndex, bool bLockWithinMiptail)
+{
+	check(TextureRHI);
+	FD3D12Texture2D*  Texture = FD3D12DynamicRHI::ResourceCast(TextureRHI);
+	Texture->Unlock(MipIndex, 0);
+}
+
+void* FD3D12DynamicRHI::RHILockTexture2DArray(FTexture2DArrayRHIParamRef TextureRHI, uint32 TextureIndex, uint32 MipIndex, EResourceLockMode LockMode, uint32& DestStride, bool bLockWithinMiptail)
+{
+	FD3D12Texture2DArray*  Texture = FD3D12DynamicRHI::ResourceCast(TextureRHI);
+    GetRHIDevice()->GetDefaultCommandContext().ConditionalClearShaderResource(Texture->ResourceLocation);
+	return Texture->Lock(MipIndex, TextureIndex, LockMode, DestStride);
+}
+
+void FD3D12DynamicRHI::RHIUnlockTexture2DArray(FTexture2DArrayRHIParamRef TextureRHI, uint32 TextureIndex, uint32 MipIndex, bool bLockWithinMiptail)
+{
+	FD3D12Texture2DArray*  Texture = FD3D12DynamicRHI::ResourceCast(TextureRHI);
+	Texture->Unlock(MipIndex, TextureIndex);
+}
+
+void FD3D12DynamicRHI::RHIUpdateTexture2D(FTexture2DRHIParamRef TextureRHI, uint32 MipIndex, const FUpdateTextureRegion2D& UpdateRegion, uint32 SourcePitch, const uint8* SourceData)
+{
+	FD3D12Texture2D*  Texture = FD3D12DynamicRHI::ResourceCast(TextureRHI);
+
+	D3D12_BOX DestBox =
+	{
+		UpdateRegion.DestX, UpdateRegion.DestY, 0,
+		UpdateRegion.DestX + UpdateRegion.Width, UpdateRegion.DestY + UpdateRegion.Height, 1
+	};
+
+	check(GPixelFormats[Texture->GetFormat()].BlockSizeX == 1);
+	check(GPixelFormats[Texture->GetFormat()].BlockSizeY == 1);
+
+	const uint32 AlignedSourcePitch = Align(SourcePitch, D3D12_TEXTURE_DATA_PLACEMENT_ALIGNMENT);
+	const uint32 bufferSize = Align(UpdateRegion.Height*AlignedSourcePitch, D3D12_TEXTURE_DATA_PLACEMENT_ALIGNMENT);
+
+	FD3D12ResourceLocation UploadHeapResourceLocation;
+    void *pData = GetRHIDevice()->GetDefaultUploadHeapAllocator().FastAlloc(bufferSize, D3D12_TEXTURE_DATA_PLACEMENT_ALIGNMENT, &UploadHeapResourceLocation);
+	check(nullptr != pData);
+
+	byte* pRowData = (byte*)pData;
+	byte* pSourceRowData = (byte*)SourceData;
+	for (uint32 i = 0; i < UpdateRegion.Height; i++)
+	{
+		FMemory::Memcpy(pRowData, pSourceRowData, SourcePitch);
+		pSourceRowData += SourcePitch;
+		pRowData += AlignedSourcePitch;
+	}
+
+    D3D12_SUBRESOURCE_FOOTPRINT sourceSubresource;
+    sourceSubresource.Depth = 1;
+    sourceSubresource.Height = UpdateRegion.Height;
+    sourceSubresource.Width = UpdateRegion.Width;
+    sourceSubresource.Format = (DXGI_FORMAT)GPixelFormats[Texture->GetFormat()].PlatformFormat;
+    sourceSubresource.RowPitch = AlignedSourcePitch;
+    check(sourceSubresource.RowPitch % D3D12_TEXTURE_DATA_PITCH_ALIGNMENT == 0);
+
+
+	D3D12_PLACED_SUBRESOURCE_FOOTPRINT placedTexture2D = { 0 };
+	placedTexture2D.Offset = UploadHeapResourceLocation.GetOffset();
+    placedTexture2D.Footprint = sourceSubresource;
+
+    FD3D12Resource* UploadBuffer = UploadHeapResourceLocation.GetResource();
+    CD3DX12_TEXTURE_COPY_LOCATION DestCopyLocation(Texture->GetResource()->GetResource(), MipIndex);
+    CD3DX12_TEXTURE_COPY_LOCATION SourceCopyLocation(UploadBuffer->GetResource(), placedTexture2D);
+
+	FD3D12CommandListHandle& hCommandList = GetRHIDevice()->GetDefaultCommandContext().CommandListHandle;
+
+	{
+		FScopeResourceBarrier ScopeResourceBarrierDest(hCommandList, Texture->GetResource(), Texture->GetResource()->GetDefaultResourceState(), D3D12_RESOURCE_STATE_COPY_DEST, DestCopyLocation.SubresourceIndex);
+
+		GetRHIDevice()->GetDefaultCommandContext().numCopies++;
+		hCommandList->CopyTextureRegion(
+			&DestCopyLocation,
+			UpdateRegion.DestX, UpdateRegion.DestY, 0,
+			&SourceCopyLocation,
+			nullptr);
+	}
+
+	DEBUG_RHI_EXECUTE_COMMAND_LIST(this);
+
+}
+
+void FD3D12DynamicRHI::RHIUpdateTexture3D(FTexture3DRHIParamRef TextureRHI, uint32 MipIndex, const FUpdateTextureRegion3D& UpdateRegion, uint32 SourceRowPitch, uint32 SourceDepthPitch, const uint8* SourceData)
+{
+	FD3D12Texture3D*  Texture = FD3D12DynamicRHI::ResourceCast(TextureRHI);
+
+	D3D12_BOX DestBox =
+	{
+		UpdateRegion.DestX, UpdateRegion.DestY, UpdateRegion.DestZ,
+		UpdateRegion.DestX + UpdateRegion.Width, UpdateRegion.DestY + UpdateRegion.Height, UpdateRegion.DestZ + UpdateRegion.Depth
+	};
+
+	check(GPixelFormats[Texture->GetFormat()].BlockSizeX == 1);
+	check(GPixelFormats[Texture->GetFormat()].BlockSizeY == 1);
+
+	const uint32 AlignedSourcePitch = Align(SourceRowPitch, D3D12_TEXTURE_DATA_PITCH_ALIGNMENT);
+	const uint32 bufferSize = Align(UpdateRegion.Height*UpdateRegion.Depth*AlignedSourcePitch, D3D12_TEXTURE_DATA_PLACEMENT_ALIGNMENT);
+
+	FD3D12ResourceLocation UploadHeapResourceLocation;
+    void *pData = GetRHIDevice()->GetDefaultUploadHeapAllocator().FastAlloc(bufferSize, D3D12_TEXTURE_DATA_PLACEMENT_ALIGNMENT, &UploadHeapResourceLocation);
+	check(nullptr != pData);
+
+	byte* pRowData = (byte*)pData;
+	byte* pSourceRowData = (byte*)SourceData;
+	byte* pSourceDepthSlice = (byte*)SourceData;
+	for (uint32 i = 0; i < UpdateRegion.Depth; i++)
+	{
+		for (uint32 j = 0; j < UpdateRegion.Height; j++)
+		{
+			FMemory::Memcpy(pRowData, pSourceRowData, SourceRowPitch);
+			pSourceRowData += SourceRowPitch;
+			pRowData += AlignedSourcePitch;
+		}
+		pSourceDepthSlice += SourceDepthPitch;
+		pSourceRowData = pSourceDepthSlice;
+	}
+    D3D12_SUBRESOURCE_FOOTPRINT sourceSubresource;
+    sourceSubresource.Depth = UpdateRegion.Depth;
+    sourceSubresource.Height = UpdateRegion.Height;
+    sourceSubresource.Width = UpdateRegion.Width;
+    sourceSubresource.Format = (DXGI_FORMAT)GPixelFormats[Texture->GetFormat()].PlatformFormat;
+    sourceSubresource.RowPitch = AlignedSourcePitch;
+    check(sourceSubresource.RowPitch % D3D12_TEXTURE_DATA_PITCH_ALIGNMENT == 0);
+
+
+	D3D12_PLACED_SUBRESOURCE_FOOTPRINT placedTexture3D = { 0 };
+    placedTexture3D.Offset = UploadHeapResourceLocation.GetOffset();
+	placedTexture3D.Footprint = sourceSubresource;
+
+    FD3D12Resource* UploadBuffer = UploadHeapResourceLocation.GetResource();
+    CD3DX12_TEXTURE_COPY_LOCATION DestCopyLocation(Texture->GetResource()->GetResource(), MipIndex);
+    CD3DX12_TEXTURE_COPY_LOCATION SourceCopyLocation(UploadBuffer->GetResource(), placedTexture3D);
+
+	FD3D12CommandListHandle& hCommandList = GetRHIDevice()->GetDefaultCommandContext().CommandListHandle;
+
+	{
+		FScopeResourceBarrier ScopeResourceBarrierDest(hCommandList, Texture->GetResource(), Texture->GetResource()->GetDefaultResourceState(), D3D12_RESOURCE_STATE_COPY_DEST, DestCopyLocation.SubresourceIndex);
+
+		GetRHIDevice()->GetDefaultCommandContext().numCopies++;
+		hCommandList->CopyTextureRegion(
+			&DestCopyLocation,
+			UpdateRegion.DestX, UpdateRegion.DestY, UpdateRegion.DestZ,
+			&SourceCopyLocation,
+			nullptr);
+	}
+
+	DEBUG_RHI_EXECUTE_COMMAND_LIST(this);
+
+}
+
+/*-----------------------------------------------------------------------------
+	Cubemap texture support.
+	-----------------------------------------------------------------------------*/
+FTextureCubeRHIRef FD3D12DynamicRHI::RHICreateTextureCube(uint32 Size, uint8 Format, uint32 NumMips, uint32 Flags, FRHIResourceCreateInfo& CreateInfo)
+{
+	return CreateD3D11Texture2D<FD3D12BaseTextureCube>(Size, Size, 6, false, true, Format, NumMips, 1, Flags, CreateInfo);
+}
+
+FTextureCubeRHIRef FD3D12DynamicRHI::RHICreateTextureCubeArray(uint32 Size, uint32 ArraySize, uint8 Format, uint32 NumMips, uint32 Flags, FRHIResourceCreateInfo& CreateInfo)
+{
+	return CreateD3D11Texture2D<FD3D12BaseTextureCube>(Size, Size, 6 * ArraySize, true, true, Format, NumMips, 1, Flags, CreateInfo);
+}
+
+void* FD3D12DynamicRHI::RHILockTextureCubeFace(FTextureCubeRHIParamRef TextureCubeRHI, uint32 FaceIndex, uint32 ArrayIndex, uint32 MipIndex, EResourceLockMode LockMode, uint32& DestStride, bool bLockWithinMiptail)
+{
+	FD3D12TextureCube*  TextureCube = FD3D12DynamicRHI::ResourceCast(TextureCubeRHI);
+    GetRHIDevice()->GetDefaultCommandContext().ConditionalClearShaderResource(TextureCube->ResourceLocation);
+	uint32 D3DFace = GetD3D11CubeFace((ECubeFace)FaceIndex);
+	return TextureCube->Lock(MipIndex, D3DFace + ArrayIndex * 6, LockMode, DestStride);
+}
+void FD3D12DynamicRHI::RHIUnlockTextureCubeFace(FTextureCubeRHIParamRef TextureCubeRHI, uint32 FaceIndex, uint32 ArrayIndex, uint32 MipIndex, bool bLockWithinMiptail)
+{
+	FD3D12TextureCube*  TextureCube = FD3D12DynamicRHI::ResourceCast(TextureCubeRHI);
+	uint32 D3DFace = GetD3D11CubeFace((ECubeFace)FaceIndex);
+	TextureCube->Unlock(MipIndex, D3DFace + ArrayIndex * 6);
+}
+
+void FD3D12DynamicRHI::RHIBindDebugLabelName(FTextureRHIParamRef TextureRHI, const TCHAR* Name)
+{
+	FName DebugName(Name);
+	TextureRHI->SetName(DebugName);
+}
+
+void FD3D12DynamicRHI::RHIVirtualTextureSetFirstMipInMemory(FTexture2DRHIParamRef TextureRHI, uint32 FirstMip)
+{
+}
+
+void FD3D12DynamicRHI::RHIVirtualTextureSetFirstMipVisible(FTexture2DRHIParamRef TextureRHI, uint32 FirstMip)
+{
+}
+
+FTextureReferenceRHIRef FD3D12DynamicRHI::RHICreateTextureReference(FLastRenderTimeContainer* LastRenderTime)
+{
+	return new FD3D12TextureReference(GetRHIDevice(), LastRenderTime);
+}
+
+void FD3D12CommandContext::RHIUpdateTextureReference(FTextureReferenceRHIParamRef TextureRefRHI, FTextureRHIParamRef NewTextureRHI)
+{
+#if 0
+	// Updating texture references is disallowed while the RHI could be caching them in referenced resource tables.
+	check(ResourceTableFrameCounter == INDEX_NONE);
+#endif
+
+	FD3D12TextureReference* TextureRef = (FD3D12TextureReference*)TextureRefRHI;
+	if (TextureRef)
+	{
+		FD3D12TextureBase* NewTexture = NULL;
+		FD3D12ShaderResourceView* NewSRV = NULL;
+		if (NewTextureRHI)
+		{
+			NewTexture = GetD3D11TextureFromRHITexture(NewTextureRHI);
+			if (NewTexture)
+			{
+				NewSRV = NewTexture->GetShaderResourceView();
+			}
+		}
+		TextureRef->SetReferencedTexture(NewTextureRHI, NewTexture, NewSRV);
+	}
+}
diff --git a/Engine/Source/Runtime/Windows/D3D12RHI/Private/D3D12UAV.cpp b/Engine/Source/Runtime/Windows/D3D12RHI/Private/D3D12UAV.cpp
new file mode 100644
--- /dev/null
+++ b/Engine/Source/Runtime/Windows/D3D12RHI/Private/D3D12UAV.cpp
@@ -0,0 +1,217 @@
+// Copyright 1998-2014 Epic Games, Inc. All Rights Reserved.
+
+#include "D3D12RHIPrivate.h"
+
+FUnorderedAccessViewRHIRef FD3D12DynamicRHI::RHICreateUnorderedAccessView(FStructuredBufferRHIParamRef StructuredBufferRHI, bool bUseUAVCounter, bool bAppendBuffer)
+{
+	FD3D12StructuredBuffer*  StructuredBuffer = FD3D12DynamicRHI::ResourceCast(StructuredBufferRHI);
+
+    const D3D12_RESOURCE_DESC& BufferDesc = StructuredBuffer->Resource->GetDesc();
+
+    const uint32 BufferUsage = StructuredBuffer->GetUsage();
+    const bool bByteAccessBuffer = (BufferUsage & BUF_ByteAddressBuffer) != 0;
+    const bool bStructuredBuffer = !bByteAccessBuffer;
+    check(bByteAccessBuffer != bStructuredBuffer); // You can't have a structured buffer that allows raw views
+    
+    D3D12_UNORDERED_ACCESS_VIEW_DESC UAVDesc = {};
+	UAVDesc.ViewDimension = D3D12_UAV_DIMENSION_BUFFER;
+	UAVDesc.Format = DXGI_FORMAT_UNKNOWN;
+
+	uint32 EffectiveStride = StructuredBuffer->GetStride();
+
+	if (bByteAccessBuffer)
+	{
+		UAVDesc.Format  = DXGI_FORMAT_R32_TYPELESS;
+		EffectiveStride = 4;
+	}
+
+	else if (BufferUsage & BUF_DrawIndirect)
+	{
+		UAVDesc.Format  = DXGI_FORMAT_R32_UINT;
+		EffectiveStride = 4;
+	}
+	UAVDesc.Buffer.FirstElement = StructuredBuffer->ResourceLocation->GetOffset() / EffectiveStride;
+	UAVDesc.Buffer.NumElements  = StructuredBuffer->ResourceLocation->GetEffectiveBufferSize() / EffectiveStride;
+	UAVDesc.Buffer.StructureByteStride = bStructuredBuffer ? EffectiveStride : 0;
+    UAVDesc.Buffer.Flags = D3D12_BUFFER_UAV_FLAG_NONE;
+	
+	UAVDesc.Buffer.CounterOffsetInBytes = 0;
+
+	const bool bNeedsCounterResource = bAppendBuffer | bUseUAVCounter;
+
+	TRefCountPtr<FD3D12Resource> CounterResource; 
+	
+	if (bNeedsCounterResource)
+	{
+		GetRHIDevice()->GetResourceHelper().CreateBuffer(D3D12_HEAP_TYPE_READBACK, 4, CounterResource.GetInitReference());
+	}
+
+	if (bByteAccessBuffer)
+	{
+		UAVDesc.Buffer.Flags |= D3D12_BUFFER_UAV_FLAG_RAW;
+	}
+
+    return new FD3D12UnorderedAccessView(GetRHIDevice(), &UAVDesc, StructuredBuffer->ResourceLocation, CounterResource);
+}
+
+FUnorderedAccessViewRHIRef FD3D12DynamicRHI::RHICreateUnorderedAccessView(FTextureRHIParamRef TextureRHI, uint32 MipLevel)
+{
+	FD3D12TextureBase* Texture = GetD3D11TextureFromRHITexture(TextureRHI);
+
+    D3D12_UNORDERED_ACCESS_VIEW_DESC UAVDesc = {};
+
+	const DXGI_FORMAT PlatformResourceFormat = (DXGI_FORMAT)GPixelFormats[TextureRHI->GetFormat()].PlatformFormat;
+	UAVDesc.Format = FindShaderResourceDXGIFormat(PlatformResourceFormat, false);
+
+	if (TextureRHI->GetTexture3D() != NULL)
+	{
+		FD3D12Texture3D* Texture3D = (FD3D12Texture3D*)Texture;
+		UAVDesc.ViewDimension = D3D12_UAV_DIMENSION_TEXTURE3D;
+		UAVDesc.Texture3D.MipSlice = MipLevel;
+		UAVDesc.Texture3D.FirstWSlice = 0;
+		UAVDesc.Texture3D.WSize = Texture3D->GetSizeZ();
+	}
+	else
+	{
+		UAVDesc.ViewDimension = D3D12_UAV_DIMENSION_TEXTURE2D;
+		UAVDesc.Texture2D.MipSlice = MipLevel;
+		UAVDesc.Texture2D.PlaneSlice = GetPlaneSliceFromViewFormat(PlatformResourceFormat, UAVDesc.Format);
+	}
+
+    return new FD3D12UnorderedAccessView(GetRHIDevice(), &UAVDesc, Texture->ResourceLocation);
+}
+
+FUnorderedAccessViewRHIRef FD3D12DynamicRHI::RHICreateUnorderedAccessView(FVertexBufferRHIParamRef VertexBufferRHI, uint8 Format)
+{
+	FD3D12VertexBuffer*  VertexBuffer = FD3D12DynamicRHI::ResourceCast(VertexBufferRHI);
+
+	FD3D12ResourceLocation *pResourceLocation = VertexBuffer->ResourceLocation.GetReference();
+    const D3D12_RESOURCE_DESC& BufferDesc = pResourceLocation->GetResource()->GetDesc();
+	const uint64 effectiveBufferSize = pResourceLocation->GetEffectiveBufferSize();
+
+    const uint32 BufferUsage = VertexBuffer->GetUsage();
+    const bool bByteAccessBuffer = (BufferUsage & BUF_ByteAddressBuffer) != 0;
+
+    D3D12_UNORDERED_ACCESS_VIEW_DESC UAVDesc = {};
+	UAVDesc.ViewDimension = D3D12_UAV_DIMENSION_BUFFER;
+	UAVDesc.Format = FindUnorderedAccessDXGIFormat((DXGI_FORMAT)GPixelFormats[Format].PlatformFormat);
+	UAVDesc.Buffer.FirstElement = pResourceLocation->GetOffset();
+
+	UAVDesc.Buffer.NumElements = effectiveBufferSize / GPixelFormats[Format].BlockBytes;
+    UAVDesc.Buffer.Flags = D3D12_BUFFER_UAV_FLAG_NONE;
+	UAVDesc.Buffer.CounterOffsetInBytes = 0;
+	UAVDesc.Buffer.StructureByteStride = 0;
+
+	if (bByteAccessBuffer)
+	{
+		UAVDesc.Buffer.Flags |= D3D12_BUFFER_UAV_FLAG_RAW;
+		UAVDesc.Format = DXGI_FORMAT_R32_TYPELESS;
+		UAVDesc.Buffer.NumElements = effectiveBufferSize / 4;
+		UAVDesc.Buffer.FirstElement /= 4;
+	}
+
+	else
+	{
+		UAVDesc.Buffer.NumElements = effectiveBufferSize / GPixelFormats[Format].BlockBytes;
+		UAVDesc.Buffer.FirstElement /= GPixelFormats[Format].BlockBytes;
+	}
+
+	return new FD3D12UnorderedAccessView(GetRHIDevice(), &UAVDesc, VertexBuffer->ResourceLocation);
+}
+
+FShaderResourceViewRHIRef FD3D12DynamicRHI::RHICreateShaderResourceView(FStructuredBufferRHIParamRef StructuredBufferRHI)
+{
+	FD3D12StructuredBuffer*  StructuredBuffer = FD3D12DynamicRHI::ResourceCast(StructuredBufferRHI);
+
+	const uint64 Offset = StructuredBuffer->ResourceLocation->GetOffset();
+    const D3D12_RESOURCE_DESC& BufferDesc = StructuredBuffer->Resource->GetDesc();
+
+    const uint32 BufferUsage = StructuredBuffer->GetUsage();
+    const bool bByteAccessBuffer = (BufferUsage & BUF_ByteAddressBuffer) != 0;
+
+	// Create a Shader Resource View
+    D3D12_SHADER_RESOURCE_VIEW_DESC SRVDesc = {};
+    SRVDesc.Shader4ComponentMapping = D3D12_DEFAULT_SHADER_4_COMPONENT_MAPPING;
+
+	// BufferDesc.StructureByteStride  is not getting patched through the D3D resource DESC structs, so use the RHI version as a hack
+	uint32 Stride = StructuredBuffer->GetStride();
+
+	if (bByteAccessBuffer)
+	{
+		SRVDesc.Buffer.Flags = D3D12_BUFFER_SRV_FLAG_RAW;
+		SRVDesc.Format = DXGI_FORMAT_R32_TYPELESS;
+		Stride = 4;
+	}
+	else
+	{
+		SRVDesc.Buffer.StructureByteStride = Stride;
+		SRVDesc.Format = DXGI_FORMAT_UNKNOWN;
+	}
+
+	SRVDesc.ViewDimension = D3D12_SRV_DIMENSION_BUFFER;
+	SRVDesc.Buffer.NumElements = StructuredBuffer->ResourceLocation->GetEffectiveBufferSize() / Stride;
+	SRVDesc.Buffer.FirstElement = Offset / Stride;
+
+    return new FD3D12ShaderResourceView(GetRHIDevice(), &SRVDesc, StructuredBuffer->ResourceLocation, Stride);
+}
+
+FShaderResourceViewRHIRef FD3D12DynamicRHI::RHICreateShaderResourceView(FVertexBufferRHIParamRef VertexBufferRHI, uint32 Stride, uint8 Format)
+{
+	FD3D12VertexBuffer*  VertexBuffer = FD3D12DynamicRHI::ResourceCast(VertexBufferRHI);
+	check(VertexBuffer);
+
+	uint32 Width = VertexBuffer->GetSize();
+
+	FD3D12Resource *pResource = VertexBuffer->ResourceLocation->GetResource();
+
+    D3D12_SHADER_RESOURCE_VIEW_DESC SRVDesc = {};
+    SRVDesc.Shader4ComponentMapping = D3D12_DEFAULT_SHADER_4_COMPONENT_MAPPING;
+	SRVDesc.ViewDimension = D3D12_SRV_DIMENSION_BUFFER;
+	if (VertexBuffer->GetUsage() & BUF_ByteAddressBuffer)
+	{
+		SRVDesc.Format = DXGI_FORMAT_R32_TYPELESS;
+		SRVDesc.Buffer.NumElements = Width / 4;
+		SRVDesc.Buffer.Flags = D3D12_BUFFER_SRV_FLAG_RAW;
+		Stride = 4;
+	}
+	else
+	{
+		SRVDesc.Format = FindShaderResourceDXGIFormat((DXGI_FORMAT)GPixelFormats[Format].PlatformFormat, false);
+		SRVDesc.Buffer.NumElements = Width / Stride;
+	}
+	SRVDesc.Buffer.StructureByteStride = 0;
+
+	if (pResource)
+	{
+		// Create a Shader Resource View
+		SRVDesc.Buffer.FirstElement = VertexBuffer->ResourceLocation->GetOffset() / Stride;
+	}
+	else
+	{
+		// Null underlying D3D12 resource should only be the case for dynamic resources
+		check(VertexBuffer->GetUsage() & BUF_AnyDynamic);
+	}
+
+    FD3D12ShaderResourceView *SRV = new FD3D12ShaderResourceView(GetRHIDevice(), &SRVDesc, VertexBuffer->ResourceLocation, Stride);
+	VertexBuffer->SetDynamicSRV(SRV);
+	return SRV;
+}
+
+void FD3D12CommandContext::RHIClearUAV(FUnorderedAccessViewRHIParamRef UnorderedAccessViewRHI, const uint32* Values)
+{
+	FD3D12UnorderedAccessView*  UnorderedAccessView = FD3D12DynamicRHI::ResourceCast(UnorderedAccessViewRHI);
+
+	// Check if the view heap is full and needs to rollover.
+	if (!StateCache.GetDescriptorCache()->ViewHeap.CanReserveSlots(1))
+		StateCache.GetDescriptorCache()->ViewHeap.RollOver();
+	uint32 ReservedSlot = StateCache.GetDescriptorCache()->ViewHeap.ReserveSlots(1);
+	D3D12_CPU_DESCRIPTOR_HANDLE CPUHandle = UnorderedAccessView->GetView();
+	D3D12_CPU_DESCRIPTOR_HANDLE DestSlot = StateCache.GetDescriptorCache()->ViewHeap.GetCPUSlotHandle(ReservedSlot);
+	D3D12_GPU_DESCRIPTOR_HANDLE GPUHandle = StateCache.GetDescriptorCache()->ViewHeap.GetGPUSlotHandle(ReservedSlot);
+	GetParentDevice()->GetDevice()->CopyDescriptorsSimple(1, DestSlot, CPUHandle, D3D12_DESCRIPTOR_HEAP_TYPE_CBV_SRV_UAV);
+	FD3D12DynamicRHI::TransitionResource (CommandListHandle, UnorderedAccessView, D3D12_RESOURCE_STATE_UNORDERED_ACCESS);
+	numClears++;
+	CommandListHandle->ClearUnorderedAccessViewUint(GPUHandle, CPUHandle, UnorderedAccessView->GetResource()->GetResource(), Values, 0, nullptr);
+
+	OwningRHI.RegisterGPUWork(1);
+}
diff --git a/Engine/Source/Runtime/Windows/D3D12RHI/Private/D3D12UniformBuffer.cpp b/Engine/Source/Runtime/Windows/D3D12RHI/Private/D3D12UniformBuffer.cpp
new file mode 100644
--- /dev/null
+++ b/Engine/Source/Runtime/Windows/D3D12RHI/Private/D3D12UniformBuffer.cpp
@@ -0,0 +1,195 @@
+// Copyright 1998-2014 Epic Games, Inc. All Rights Reserved.
+
+/*=============================================================================
+	D3D12UniformBuffer.cpp: D3D uniform buffer RHI implementation.
+	=============================================================================*/
+
+#include "D3D12RHIPrivate.h"
+#include "UniformBuffer.h"
+
+namespace D3D12RHI
+{
+	/** Describes a uniform buffer in the free pool. */
+	struct FPooledUniformBuffer
+	{
+		TRefCountPtr<FD3D12Resource> Buffer;
+		uint32 CreatedSize;
+		uint32 FrameFreed;
+	};
+
+	/**
+	 * Number of size buckets to use for the uniform buffer free pool.
+	 * This needs to be enough to cover the valid uniform buffer size range combined with the heuristic used to map sizes to buckets.
+	 */
+	const int32 NumPoolBuckets = 17;
+
+	/**
+	 * Number of frames that a uniform buffer will not be re-used for, after being freed.
+	 * This is done as a workaround for what appears to be an AMD driver bug with 11.10 drivers and a 6970 HD,
+	 * Where reusing a constant buffer with D3D11_MAP_WRITE_DISCARD still in use by the GPU will result in incorrect contents randomly.
+	 */
+	const int32 NumSafeFrames = 3;
+
+	/** Returns the size in bytes of the bucket that the given size fits into. */
+	uint32 GetPoolBucketSize(uint32 NumBytes)
+	{
+		return FMath::RoundUpToPowerOfTwo(NumBytes);
+	}
+
+	/** Returns the index of the bucket that the given size fits into. */
+	uint32 GetPoolBucketIndex(uint32 NumBytes)
+	{
+		return FMath::CeilLogTwo(NumBytes);
+	}
+
+	/** Does per-frame global updating for the uniform buffer pool. */
+	void UniformBufferBeginFrame()
+	{
+	}
+}
+
+// 0 is reserved to mean "not a uniform buffer"
+LONGLONG D3D12UniformBufferSequenceNumber = 1;
+
+FUniformBufferRHIRef FD3D12DynamicRHI::RHICreateUniformBuffer(const void* Contents, const FRHIUniformBufferLayout& Layout, EUniformBufferUsage Usage)
+{
+	uint64 SequenceNumber = InterlockedIncrement64(&D3D12UniformBufferSequenceNumber);
+
+	FD3D12UniformBuffer* NewUniformBuffer = nullptr;
+	const uint32 NumBytesActualData = Layout.ConstantBufferSize;
+    const uint32 NumBytes = Align(NumBytesActualData, D3D12_CONSTANT_BUFFER_DATA_PLACEMENT_ALIGNMENT);	// Allocate a size that is big enough for a multiple of 256
+	if (NumBytes > 0)
+	{
+		// Constant buffers must also be 16-byte aligned.
+		check(Align(NumBytes, 16) == NumBytes);
+		check(Align(Contents, 16) == Contents);
+		check(NumBytes <= D3D12_REQ_CONSTANT_BUFFER_ELEMENT_COUNT * 16);
+		check(NumBytes < (1 << NumPoolBuckets));
+
+		SCOPE_CYCLE_COUNTER(STAT_D3D12UpdateUniformBufferTime);
+
+		// Use an upload heap
+		TRefCountPtr<FD3D12ResourceLocation> ResourceLocation = new FD3D12ResourceLocation();
+        void *pData = GetRHIDevice()->GetDefaultUploadHeapAllocator().Alloc(NumBytes, D3D12_CONSTANT_BUFFER_DATA_PLACEMENT_ALIGNMENT, ResourceLocation);
+
+		check(pData != nullptr);
+
+		// Copy the data to the upload heap
+		FMemory::Memcpy(pData, Contents, NumBytesActualData);
+
+		check(ResourceLocation->GetOffset() % 16 == 0);
+		check(ResourceLocation->GetEffectiveBufferSize() == NumBytes);
+		NewUniformBuffer = new FD3D12UniformBuffer(GetRHIDevice(), Layout, ResourceLocation, FRingAllocation(), SequenceNumber);
+
+		// Create an offline CBV descriptor
+        NewUniformBuffer->OfflineDescriptorHandle = 
+            GetRHIDevice()->GetViewDescriptorAllocator<D3D12_CONSTANT_BUFFER_VIEW_DESC>().AllocateHeapSlot(NewUniformBuffer->OfflineHeapIndex);
+
+		D3D12_CONSTANT_BUFFER_VIEW_DESC CBView;
+        CBView.BufferLocation = ResourceLocation->GetResource()->GetResource()->GetGPUVirtualAddress() + ResourceLocation->GetOffset();
+        CBView.SizeInBytes = NumBytes;
+        check(ResourceLocation->GetOffset() % 256 == 0);
+        check(CBView.SizeInBytes <= 4096 * 16);
+        check(CBView.SizeInBytes % 256 == 0);
+
+        GetRHIDevice()->GetDevice()->CreateConstantBufferView(&CBView, NewUniformBuffer->OfflineDescriptorHandle);
+	}
+	else
+	{
+		// This uniform buffer contains no constants, only a resource table.
+        NewUniformBuffer = new FD3D12UniformBuffer(GetRHIDevice(), Layout, nullptr, FRingAllocation(), SequenceNumber);
+
+		check(0 == NewUniformBuffer->OfflineDescriptorHandle.ptr);
+	}
+
+	if (Layout.Resources.Num())
+	{
+		int32 NumResources = Layout.Resources.Num();
+		FRHIResource** InResources = (FRHIResource**)((uint8*)Contents + Layout.ResourceOffset);
+		NewUniformBuffer->ResourceTable.Empty(NumResources);
+		NewUniformBuffer->ResourceTable.AddZeroed(NumResources);
+		for (int32 i = 0; i < NumResources; ++i)
+		{
+			check(InResources[i]);
+			NewUniformBuffer->ResourceTable[i] = InResources[i];
+		}
+		NewUniformBuffer->RawResourceTable.Empty(NumResources);
+		NewUniformBuffer->RawResourceTable.AddZeroed(NumResources);
+	}
+
+	return NewUniformBuffer;
+}
+
+FD3D12UniformBuffer::~FD3D12UniformBuffer()
+{
+	check(!GRHISupportsRHIThread || IsInRenderingThread());
+
+	if (OfflineDescriptorHandle.ptr != 0)
+	{
+        FDescriptorHeapManager& HeapManager = GetParentDevice()->GetViewDescriptorAllocator<D3D12_CONSTANT_BUFFER_VIEW_DESC>();
+
+		HeapManager.FreeHeapSlot(OfflineDescriptorHandle, OfflineHeapIndex);
+	}
+}
+
+void FD3D12UniformBuffer::CacheResourcesInternal()
+{
+	uint32 Start = FPlatformTime::Cycles();
+
+	const FRHIUniformBufferLayout& Layout = GetLayout();
+	int32 NumResources = Layout.Resources.Num();
+	const uint8* RESTRICT ResourceTypes = Layout.Resources.GetData();
+	const TRefCountPtr<FRHIResource>* RESTRICT Resources = ResourceTable.GetData();
+	FResourcePair* RESTRICT RawResources = RawResourceTable.GetData();
+	float CurrentTime = FApp::GetCurrentTime();
+
+	// todo: Immutable resources, i.e. not textures, can be safely cached across frames.
+	// Texture streaming makes textures complicated :)
+	for (int32 i = 0; i < NumResources; ++i)
+	{
+		switch (ResourceTypes[i])
+		{
+		case UBMT_SRV:
+		{
+			FD3D12ShaderResourceView* ShaderResourceViewRHI = (FD3D12ShaderResourceView*)Resources[i].GetReference();
+			RawResources[i].ShaderResourceLocation = ShaderResourceViewRHI->GetResourceLocation();
+			RawResources[i].D3D11Resource = (IUnknown*)ShaderResourceViewRHI;
+		}
+			break;
+
+		case UBMT_TEXTURE:
+		{
+			// todo: this does multiple virtual function calls to find the right type to cast to
+			// this is due to multiple inheritance nastiness, NEEDS CLEANUP
+			FRHITexture* TextureRHI = (FRHITexture*)Resources[i].GetReference();
+			TextureRHI->SetLastRenderTime(CurrentTime);
+			FD3D12TextureBase* TextureD3D11 = GetD3D11TextureFromRHITexture(TextureRHI);
+			RawResources[i].ShaderResourceLocation = TextureD3D11->GetBaseShaderResource()->ResourceLocation;
+			RawResources[i].D3D11Resource = (IUnknown*)TextureD3D11->GetShaderResourceView();
+		}
+			break;
+
+		case UBMT_UAV:
+			RawResources[i].ShaderResourceLocation = NULL;
+			RawResources[i].D3D11Resource = NULL;
+			check(0);
+			break;
+
+		case UBMT_SAMPLER:
+			RawResources[i].ShaderResourceLocation = NULL;
+			RawResources[i].D3D11Resource = (IUnknown*)((FD3D12SamplerState*)Resources[i].GetReference());
+			break;
+
+		default:
+			check(0);
+			break;
+		}
+	}
+
+    GetParentDevice()->GetOwningRHI()->IncrementCacheResourceTableCycles(FPlatformTime::Cycles() - Start);
+	GetParentDevice()->GetOwningRHI()->IncrementCacheResourceTableCalls();
+}
+
+void FD3D12Device::ReleasePooledUniformBuffers()
+{
+}
\ No newline at end of file
