UE4 - D3D12RHI (not working)

diff --git a/Engine/Source/Runtime/Windows/D3D12RHI/Private/D3D12Util.cpp b/Engine/Source/Runtime/Windows/D3D12RHI/Private/D3D12Util.cpp
new file mode 100644
--- /dev/null
+++ b/Engine/Source/Runtime/Windows/D3D12RHI/Private/D3D12Util.cpp
@@ -0,0 +1,461 @@
+// Copyright 1998-2014 Epic Games, Inc. All Rights Reserved.
+
+/*=============================================================================
+D3D12Util.h: D3D RHI utility implementation.
+=============================================================================*/
+
+#include "D3D12RHIPrivate.h"
+#include "EngineModule.h"
+
+#define D3DERR(x) case x: ErrorCodeText = TEXT(#x); break;
+#define LOCTEXT_NAMESPACE "Developer.MessageLog"
+
+#ifndef _FACD3D 
+#define _FACD3D  0x876
+#endif	//_FACD3D 
+#ifndef MAKE_D3DHRESULT
+#define _FACD3D  0x876
+#define MAKE_D3DHRESULT( code )  MAKE_HRESULT( 1, _FACD3D, code )
+#endif	//MAKE_D3DHRESULT
+
+#if WITH_D3DX_LIBS
+#ifndef D3DERR_INVALIDCALL
+#define D3DERR_INVALIDCALL MAKE_D3DHRESULT(2156)
+#endif//D3DERR_INVALIDCALL
+#ifndef D3DERR_WASSTILLDRAWING
+#define D3DERR_WASSTILLDRAWING MAKE_D3DHRESULT(540)
+#endif//D3DERR_WASSTILLDRAWING
+#endif
+
+static FString GetD3D11DeviceHungErrorString(HRESULT ErrorCode)
+{
+	FString ErrorCodeText;
+
+	switch (ErrorCode)
+	{
+		D3DERR(DXGI_ERROR_DEVICE_HUNG)
+			D3DERR(DXGI_ERROR_DEVICE_REMOVED)
+			D3DERR(DXGI_ERROR_DEVICE_RESET)
+			D3DERR(DXGI_ERROR_DRIVER_INTERNAL_ERROR)
+			D3DERR(DXGI_ERROR_INVALID_CALL)
+	default: ErrorCodeText = FString::Printf(TEXT("%08X"), (int32)ErrorCode);
+	}
+
+	return ErrorCodeText;
+}
+
+static FString GetD3D11ErrorString(HRESULT ErrorCode, ID3D12Device* Device)
+{
+	FString ErrorCodeText;
+
+	switch (ErrorCode)
+	{
+		D3DERR(S_OK);
+		D3DERR(D3D11_ERROR_FILE_NOT_FOUND)
+			D3DERR(D3D11_ERROR_TOO_MANY_UNIQUE_STATE_OBJECTS)
+#if WITH_D3DX_LIBS
+			D3DERR(D3DERR_INVALIDCALL)
+			D3DERR(D3DERR_WASSTILLDRAWING)
+#endif	//WITH_D3DX_LIBS
+			D3DERR(E_FAIL)
+			D3DERR(E_INVALIDARG)
+			D3DERR(E_OUTOFMEMORY)
+			D3DERR(DXGI_ERROR_INVALID_CALL)
+			D3DERR(E_NOINTERFACE)
+			D3DERR(DXGI_ERROR_DEVICE_REMOVED)
+	default: ErrorCodeText = FString::Printf(TEXT("%08X"), (int32)ErrorCode);
+	}
+
+	if (ErrorCode == DXGI_ERROR_DEVICE_REMOVED && Device)
+	{
+		HRESULT hResDeviceRemoved = Device->GetDeviceRemovedReason();
+		ErrorCodeText += FString(TEXT(" ")) + GetD3D11DeviceHungErrorString(hResDeviceRemoved);
+	}
+
+	return ErrorCodeText;
+}
+
+#undef D3DERR
+
+namespace D3D12RHI
+{
+	const TCHAR* GetD3D11TextureFormatString(DXGI_FORMAT TextureFormat)
+	{
+		static const TCHAR* EmptyString = TEXT("");
+		const TCHAR* TextureFormatText = EmptyString;
+#define D3DFORMATCASE(x) case x: TextureFormatText = TEXT(#x); break;
+		switch (TextureFormat)
+		{
+			D3DFORMATCASE(DXGI_FORMAT_R8G8B8A8_UNORM)
+				D3DFORMATCASE(DXGI_FORMAT_B8G8R8A8_UNORM)
+				D3DFORMATCASE(DXGI_FORMAT_B8G8R8X8_UNORM)
+				D3DFORMATCASE(DXGI_FORMAT_BC1_UNORM)
+				D3DFORMATCASE(DXGI_FORMAT_BC2_UNORM)
+				D3DFORMATCASE(DXGI_FORMAT_BC3_UNORM)
+				D3DFORMATCASE(DXGI_FORMAT_BC4_UNORM)
+				D3DFORMATCASE(DXGI_FORMAT_R16G16B16A16_FLOAT)
+				D3DFORMATCASE(DXGI_FORMAT_R32G32B32A32_FLOAT)
+				D3DFORMATCASE(DXGI_FORMAT_UNKNOWN)
+				D3DFORMATCASE(DXGI_FORMAT_R8_UNORM)
+#if DEPTH_32_BIT_CONVERSION
+				D3DFORMATCASE(DXGI_FORMAT_D32_FLOAT_S8X24_UINT)
+				D3DFORMATCASE(DXGI_FORMAT_R32_FLOAT_X8X24_TYPELESS)
+#endif
+                D3DFORMATCASE(DXGI_FORMAT_R32G8X24_TYPELESS)
+				D3DFORMATCASE(DXGI_FORMAT_D24_UNORM_S8_UINT)
+				D3DFORMATCASE(DXGI_FORMAT_R24_UNORM_X8_TYPELESS)
+				D3DFORMATCASE(DXGI_FORMAT_R32_FLOAT)
+				D3DFORMATCASE(DXGI_FORMAT_R16G16_UINT)
+                D3DFORMATCASE(DXGI_FORMAT_R16G16_UNORM)
+                D3DFORMATCASE(DXGI_FORMAT_R16G16_SNORM)
+				D3DFORMATCASE(DXGI_FORMAT_R16G16_FLOAT)
+				D3DFORMATCASE(DXGI_FORMAT_R32G32_FLOAT)
+				D3DFORMATCASE(DXGI_FORMAT_R10G10B10A2_UNORM)
+				D3DFORMATCASE(DXGI_FORMAT_R16G16B16A16_UINT)
+				D3DFORMATCASE(DXGI_FORMAT_R8G8_SNORM)
+				D3DFORMATCASE(DXGI_FORMAT_BC5_UNORM)
+				D3DFORMATCASE(DXGI_FORMAT_R1_UNORM)
+				D3DFORMATCASE(DXGI_FORMAT_R8G8B8A8_TYPELESS)
+				D3DFORMATCASE(DXGI_FORMAT_B8G8R8A8_TYPELESS)
+				D3DFORMATCASE(DXGI_FORMAT_BC7_UNORM)
+				D3DFORMATCASE(DXGI_FORMAT_BC6H_UF16)
+		default: TextureFormatText = EmptyString;
+		}
+#undef D3DFORMATCASE
+		return TextureFormatText;
+	}
+}
+using namespace D3D12RHI;
+
+static FString GetD3D12TextureFlagString(uint32 TextureFlags)
+{
+	FString TextureFormatText = TEXT("");
+
+	if (TextureFlags & D3D12_RESOURCE_FLAG_ALLOW_RENDER_TARGET)
+	{
+		TextureFormatText += TEXT("D3D12_RESOURCE_FLAG_ALLOW_RENDER_TARGET ");
+	}
+
+    if (TextureFlags & D3D12_RESOURCE_FLAG_ALLOW_DEPTH_STENCIL)
+	{
+		TextureFormatText += TEXT("D3D12_RESOURCE_FLAG_ALLOW_DEPTH_STENCIL ");
+	}
+
+    if (TextureFlags & D3D12_RESOURCE_FLAG_DENY_SHADER_RESOURCE)
+	{
+		TextureFormatText += TEXT("D3D12_RESOURCE_FLAG_DENY_SHADER_RESOURCE ");
+	}
+
+    if (TextureFlags & D3D12_RESOURCE_FLAG_ALLOW_UNORDERED_ACCESS)
+	{
+		TextureFormatText += TEXT("D3D12_RESOURCE_FLAG_ALLOW_UNORDERED_ACCESS ");
+	}
+	return TextureFormatText;
+}
+
+
+static void TerminateOnDeviceRemoved(HRESULT D3DResult)
+{
+	if (D3DResult == DXGI_ERROR_DEVICE_REMOVED)
+	{
+		FPlatformMisc::MessageBoxExt(EAppMsgType::Ok, *LOCTEXT("DeviceRemoved", "Video driver crashed and was reset!  Make sure your video drivers are up to date.  Exiting...").ToString(), TEXT("Error"));
+		FPlatformMisc::RequestExit(true);
+	}
+}
+
+#ifndef MAKE_D3DHRESULT
+#define _FACD3D						0x876
+#define MAKE_D3DHRESULT( code)		MAKE_HRESULT( 1, _FACD3D, code )
+#endif	//MAKE_D3DHRESULT
+
+namespace D3D12RHI
+{
+	void VerifyD3D11Result(HRESULT D3DResult, const ANSICHAR* Code, const ANSICHAR* Filename, uint32 Line, ID3D12Device* Device)
+	{
+		check(FAILED(D3DResult));
+
+		const FString& ErrorString = GetD3D11ErrorString(D3DResult, Device);
+
+		UE_LOG(LogD3D12RHI, Error, TEXT("%s failed \n at %s:%u \n with error %s"), ANSI_TO_TCHAR(Code), ANSI_TO_TCHAR(Filename), Line, *ErrorString);
+
+		TerminateOnDeviceRemoved(D3DResult);
+
+		// this is to track down a rarely happening crash
+		if (D3DResult == E_OUTOFMEMORY)
+		{
+			FPlatformMisc::MessageBoxExt(EAppMsgType::Ok, *NSLOCTEXT("D3D12RHI", "OutOfMemory", "Out of video memory trying to allocate a rendering resource.").ToString(), TEXT("Error"));
+#if STATS
+			GetRendererModule().DebugLogOnCrash();
+#endif
+			FPlatformMisc::RequestExit(/*bForce=*/ true);
+		}
+
+		UE_LOG(LogD3D12RHI, Fatal, TEXT("%s failed \n at %s:%u \n with error %s"), ANSI_TO_TCHAR(Code), ANSI_TO_TCHAR(Filename), Line, *ErrorString);
+	}
+
+	void VerifyD3D12CreateTextureResult(HRESULT D3DResult, const ANSICHAR* Code, const ANSICHAR* Filename, uint32 Line, uint32 SizeX, uint32 SizeY, uint32 SizeZ, uint8 Format, uint32 NumMips, uint32 Flags)
+	{
+		check(FAILED(D3DResult));
+
+		const FString ErrorString = GetD3D11ErrorString(D3DResult, 0);
+		const TCHAR* D3DFormatString = GetD3D11TextureFormatString((DXGI_FORMAT)Format);
+
+		UE_LOG(LogD3D12RHI, Error, TEXT("%s failed \n at %s:%u \n with error %s"), ANSI_TO_TCHAR(Code), ANSI_TO_TCHAR(Filename), Line, *ErrorString);
+
+		TerminateOnDeviceRemoved(D3DResult);
+
+		// this is to track down a rarely happening crash
+		if (D3DResult == E_OUTOFMEMORY)
+		{
+#if STATS
+			GetRendererModule().DebugLogOnCrash();
+#endif // STATS
+		}
+
+		UE_LOG(LogD3D12RHI, Fatal,
+			TEXT("%s failed \n at %s:%u \n with error %s, \n Size=%ix%ix%i Format=%s(0x%08X), NumMips=%i, Flags=%s"),
+			ANSI_TO_TCHAR(Code),
+			ANSI_TO_TCHAR(Filename),
+			Line,
+			*ErrorString,
+			SizeX,
+			SizeY,
+			SizeZ,
+			D3DFormatString,
+			Format,
+			NumMips,
+			*GetD3D12TextureFlagString(Flags));
+	}
+
+	void VerifyComRefCount(IUnknown* Object, int32 ExpectedRefs, const TCHAR* Code, const TCHAR* Filename, int32 Line)
+	{
+		int32 NumRefs;
+
+		if (Object)
+		{
+			Object->AddRef();
+			NumRefs = Object->Release();
+			if (NumRefs != ExpectedRefs)
+			{
+				UE_LOG(LogD3D12RHI, Fatal,
+					TEXT("%s:(%d): %s has %d refs, expected %d"),
+					Filename,
+					Line,
+					Code,
+					NumRefs,
+					ExpectedRefs
+					);
+			}
+		}
+	}
+}
+
+FD3D12BoundRenderTargets::FD3D12BoundRenderTargets(FD3D12RenderTargetView **RTArray, uint32 NumActiveRTs, FD3D12DepthStencilView* DSView)
+{
+	FMemory::Memcpy(RenderTargetViews, RTArray, sizeof(RenderTargetViews));
+	DepthStencilView = DSView;
+	NumActiveTargets = NumActiveRTs;
+}
+
+FD3D12BoundRenderTargets::~FD3D12BoundRenderTargets()
+{
+}
+
+
+FD3D12DynamicBuffer::FD3D12DynamicBuffer(FD3D12Device* InParent, FD3D12DynamicHeapAllocator& UploadHeap)
+: UploadHeapAllocator(UploadHeap),
+ FD3D12DeviceChild(InParent)
+{
+	InitResource();
+}
+
+FD3D12DynamicBuffer::~FD3D12DynamicBuffer()
+{
+	ReleaseResource();
+}
+
+void FD3D12DynamicBuffer::InitRHI()
+{
+}
+
+void FD3D12DynamicBuffer::ReleaseRHI()
+{
+}
+
+void* FD3D12DynamicBuffer::Lock(uint32 Size)
+{
+	FD3D12ResourceLocation *NewResourceLocation = new FD3D12ResourceLocation();
+	ResourceLocation = NewResourceLocation;
+    return UploadHeapAllocator.FastAlloc(Size, D3D12_TEXTURE_DATA_PLACEMENT_ALIGNMENT, NewResourceLocation);
+}
+
+FD3D12ResourceLocation* FD3D12DynamicBuffer::Unlock()
+{
+    return ResourceLocation;
+}
+
+//==================================================================================================================================
+// CResourceState
+// Tracking of per-resource or per-subresource state
+//==================================================================================================================================
+//----------------------------------------------------------------------------------------------------------------------------------
+void CResourceState::Initialize(uint32 SubresourceCount)
+{
+	check(0 == m_SubresourceState.Num());
+
+	// Allocate space for per-subresource tracking structures
+	check(SubresourceCount > 0);
+	m_SubresourceState.SetNumUninitialized(SubresourceCount);
+	check(m_SubresourceState.Num() == SubresourceCount);
+
+	// All subresources start out in an unknown state
+	SetResourceState(D3D12_RESOURCE_STATE_TBD);
+}
+
+//----------------------------------------------------------------------------------------------------------------------------------
+bool CResourceState::AreAllSubresourcesSame() const
+{
+	return m_AllSubresourcesSame && (m_ResourceState != D3D12_RESOURCE_STATE_TBD);
+}
+
+//----------------------------------------------------------------------------------------------------------------------------------
+bool CResourceState::CheckResourceState(D3D12_RESOURCE_STATES State) const
+{
+	if (m_AllSubresourcesSame)
+	{
+		return State == m_ResourceState;
+	}
+	else
+	{
+		// All subresources must be individually checked
+		const uint32 numSubresourceStates = m_SubresourceState.Num();
+		for (uint32 i = 0; i < numSubresourceStates; i++)
+		{
+			if (m_SubresourceState[i] != State)
+			{
+				return false;
+			}
+		}
+
+		return true;
+	}
+}
+
+//----------------------------------------------------------------------------------------------------------------------------------
+bool CResourceState::CheckResourceStateInitalized() const
+{
+	return m_SubresourceState.Num() > 0;
+}
+
+//----------------------------------------------------------------------------------------------------------------------------------
+D3D12_RESOURCE_STATES CResourceState::GetSubresourceState(uint32 SubresourceIndex) const
+{
+	if (m_AllSubresourcesSame)
+	{
+		return m_ResourceState;
+	}
+	else
+	{
+		check(SubresourceIndex < static_cast<uint32>(m_SubresourceState.Num()));
+		return m_SubresourceState[SubresourceIndex];
+	}
+}
+
+//----------------------------------------------------------------------------------------------------------------------------------
+void CResourceState::SetResourceState(D3D12_RESOURCE_STATES State)
+{
+	m_AllSubresourcesSame = 1;
+
+	m_ResourceState = State;
+
+	// State is now tracked per-resource, so m_SubresourceState should not be read.
+#if UE_BUILD_DEBUG
+	const uint32 numSubresourceStates = m_SubresourceState.Num();
+	for (uint32 i = 0; i < numSubresourceStates; i++)
+	{
+		m_SubresourceState[i] = D3D12_RESOURCE_STATE_CORRUPT;
+	}
+#endif
+}
+
+//----------------------------------------------------------------------------------------------------------------------------------
+void CResourceState::SetSubresourceState(uint32 SubresourceIndex, D3D12_RESOURCE_STATES State)
+{
+	// If setting all subresources, or the resource only has a single subresource, set the per-resource state
+	if (SubresourceIndex == D3D12_RESOURCE_BARRIER_ALL_SUBRESOURCES ||
+		m_SubresourceState.Num() == 1)
+	{
+		SetResourceState(State);
+	}
+	else
+	{
+		check(SubresourceIndex < static_cast<uint32>(m_SubresourceState.Num()));
+
+		// If state was previously tracked on a per-resource level, then transition to per-subresource tracking
+		if (m_AllSubresourcesSame)
+		{
+			const uint32 numSubresourceStates = m_SubresourceState.Num();
+			for (uint32 i = 0; i < numSubresourceStates; i++)
+			{
+				m_SubresourceState[i] = m_ResourceState;
+			}
+
+			m_AllSubresourcesSame = 0;
+
+			// State is now tracked per-subresource, so m_ResourceState should not be read.
+#if UE_BUILD_DEBUG
+			m_ResourceState = D3D12_RESOURCE_STATE_CORRUPT;
+#endif
+		}
+
+		m_SubresourceState[SubresourceIndex] = State;
+	}
+}
+
+
+//
+// Stat declarations.
+//
+
+DEFINE_STAT(STAT_D3D12PresentTime);
+DEFINE_STAT(STAT_D3D12TexturesAllocated);
+DEFINE_STAT(STAT_D3D12TexturesReleased);
+DEFINE_STAT(STAT_D3D12ClearShaderResourceTime);
+DEFINE_STAT(STAT_D3D12CreateTextureTime);
+DEFINE_STAT(STAT_D3D12LockTextureTime);
+DEFINE_STAT(STAT_D3D12UnlockTextureTime);
+DEFINE_STAT(STAT_D3D12CopyTextureTime);
+DEFINE_STAT(STAT_D3D12NewBoundShaderStateTime);
+DEFINE_STAT(STAT_D3D12CreateBoundShaderStateTime);
+DEFINE_STAT(STAT_D3D12CleanUniformBufferTime);
+DEFINE_STAT(STAT_D3D12UpdateUniformBufferTime);
+DEFINE_STAT(STAT_D3D12TexturePoolMemory);
+DEFINE_STAT(STAT_D3D12FreeUniformBufferMemory);
+DEFINE_STAT(STAT_D3D12NumFreeUniformBuffers);
+DEFINE_STAT(STAT_D3D12NumBoundShaderState);
+
+DEFINE_STAT(STAT_D3D12CommitResourceTables);
+DEFINE_STAT(STAT_D3D12CacheResourceTables);
+DEFINE_STAT(STAT_D3D12CacheResourceTableCalls);
+DEFINE_STAT(STAT_D3D12SetShaderTextureTime);
+DEFINE_STAT(STAT_D3D12SetShaderTextureCalls);
+DEFINE_STAT(STAT_D3D12SetTextureInTableCalls);
+
+DEFINE_STAT(STAT_D3D12ApplyStateTime);
+DEFINE_STAT(STAT_D3D12ApplyStateRebuildPSOTime);
+DEFINE_STAT(STAT_D3D12ApplyStateSetSRVTime);
+DEFINE_STAT(STAT_D3D12ApplyStateSetUAVTime);
+DEFINE_STAT(STAT_D3D12ApplyStateSetVertexBufferTime);
+DEFINE_STAT(STAT_D3D12ApplyStateSetConstantBufferTime);
+DEFINE_STAT(STAT_D3D12ApplyStatePSOCreateTime);
+DEFINE_STAT(STAT_D3D12ClearMRT);
+DEFINE_STAT(STAT_D3D12CommitGraphicsConstants);
+DEFINE_STAT(STAT_D3D12SetShaderUniformBuffer);
+DEFINE_STAT(STAT_D3D12SetBoundShaderState);
+DEFINE_STAT(STAT_D3D12ExecuteCommandListTime);
+DEFINE_STAT(STAT_D3D12WaitForFenceTime);
+DEFINE_STAT(STAT_D3D12ResourceCleanUpTime);
+DEFINE_STAT(STAT_D3D12UsedVideoMemory);
+DEFINE_STAT(STAT_D3D12AvailableVideoMemory);
+DEFINE_STAT(STAT_D3D12TotalVideoMemory);
+#undef LOCTEXT_NAMESPACE
diff --git a/Engine/Source/Runtime/Windows/D3D12RHI/Private/D3D12VertexBuffer.cpp b/Engine/Source/Runtime/Windows/D3D12RHI/Private/D3D12VertexBuffer.cpp
new file mode 100644
--- /dev/null
+++ b/Engine/Source/Runtime/Windows/D3D12RHI/Private/D3D12VertexBuffer.cpp
@@ -0,0 +1,339 @@
+// Copyright 1998-2014 Epic Games, Inc. All Rights Reserved.
+
+/*=============================================================================
+	D3D12VertexBuffer.cpp: D3D vertex buffer RHI implementation.
+	=============================================================================*/
+
+#include "D3D12RHIPrivate.h"
+
+FVertexBufferRHIRef FD3D12DynamicRHI::RHICreateVertexBuffer(uint32 Size, uint32 InUsage, FRHIResourceCreateInfo& CreateInfo)
+{
+	// Explicitly check that the size is nonzero before allowing CreateVertexBuffer to opaquely fail.
+	check(Size > 0);
+
+    // Describe the vertex buffer.
+    const bool bIsDynamic = (InUsage & BUF_AnyDynamic) ? true : false;
+    D3D12_RESOURCE_DESC Desc = CD3DX12_RESOURCE_DESC::Buffer(Size);
+
+	if (InUsage & BUF_UnorderedAccess)
+	{
+        Desc.Flags |= D3D12_RESOURCE_FLAG_ALLOW_UNORDERED_ACCESS;
+
+        static bool bRequiresRawView = (GMaxRHIFeatureLevel < ERHIFeatureLevel::SM5);
+        if (bRequiresRawView)
+        {
+            // Force the buffer to be a raw, byte address buffer
+            InUsage |= BUF_ByteAddressBuffer;
+        }
+	}
+
+	if ((InUsage & BUF_ShaderResource) == 0)
+	{
+        Desc.Flags |= D3D12_RESOURCE_FLAG_DENY_SHADER_RESOURCE;
+	}
+
+	if (FPlatformProperties::SupportsFastVRAMMemory())
+	{
+		if (InUsage & BUF_FastVRAM)
+		{
+			FFastVRAMAllocator::GetFastVRAMAllocator()->AllocUAVBuffer(Desc);
+		}
+	}
+
+	// If a resource array was provided for the resource, create the resource pre-populated
+	D3D12_SUBRESOURCE_DATA InitData = { 0 };
+	D3D12_SUBRESOURCE_DATA* pInitData = nullptr;
+	if (CreateInfo.ResourceArray)
+	{
+		check(Size == CreateInfo.ResourceArray->GetResourceDataSize());
+		InitData.pData = CreateInfo.ResourceArray->GetResourceData();
+		InitData.RowPitch = Size;
+		InitData.SlicePitch = 0;
+		pInitData = &InitData;
+	}
+
+	TRefCountPtr<FD3D12ResourceLocation> ResourceLocation;
+    if (bIsDynamic)
+	{
+		void *pData = GetRHIDevice()->GetDefaultUploadHeapAllocator().Alloc(Size, D3D12_TEXTURE_DATA_PLACEMENT_ALIGNMENT, ResourceLocation.GetInitReference());
+
+		if (pInitData)
+		{
+			// Handle initial data
+			FMemory::Memcpy(pData, InitData.pData, Size);
+		}
+		check(ResourceLocation->GetEffectiveBufferSize() == Size);
+	}
+	else
+	{
+		// Create an empty resource location
+        VERIFYD3D11RESULT(GetRHIDevice()->GetDefaultBufferAllocator().AllocDefaultResource(Desc, pInitData, ResourceLocation.GetInitReference()));
+		check(ResourceLocation->GetEffectiveBufferSize() == Size);
+	}
+
+	if (CreateInfo.ResourceArray)
+	{
+		// Discard the resource array's contents.
+		CreateInfo.ResourceArray->Discard();
+	}
+
+	return new FD3D12VertexBuffer(GetRHIDevice(), ResourceLocation, Size, InUsage);
+}
+
+FD3D12VertexBuffer::~FD3D12VertexBuffer()
+{
+	// Delete blocks from the SRV pool by allocating each object from the pool and 
+	// manually deleting them
+	ResourceViewHandleDesc ViewHandleDesc;
+	while (SRVPool.TryAllocate(ViewHandleDesc))
+	{
+		  FD3D12ResourceAllocator *Allocator = ViewHandleDesc.BlockInfo->Allocator;
+		  check(!!Allocator);
+
+		  Allocator->ExpireBlock(ViewHandleDesc.BlockInfo);
+	}
+
+	if (ResourceLocation->GetResource() != nullptr)
+	{
+		UpdateBufferStats(ResourceLocation, false, D3D12_BUFFER_TYPE_VERTEX);
+	}
+
+	// Dynamic resources have a permanente entry in the OutstandingLocks that needs to be cleaned up. 
+	if (GetUsage() & BUF_AnyDynamic)
+	{
+		FD3D12LockedKey LockedKey(ResourceLocation.GetReference());
+		FD3D12DynamicRHI::GetD3DRHI()->OutstandingLocks.Remove (LockedKey);
+	}
+}
+
+void *FD3D12VertexBuffer::DynamicLock()
+{
+	void *pData = nullptr;
+	check((GetUsage() & BUF_AnyDynamic) ? true : false);
+
+	if (DynamicSRV != nullptr)
+	{
+		// Place current resource location and SRV handle into the pool
+		FD3D12ResourceBlockInfo *pBlockInfo = ResourceLocation->GetBlockInfo();
+		if (pBlockInfo)
+		{
+			ResourceViewHandleDesc ViewHandleDesc(pBlockInfo, DynamicSRV->GetView(), DynamicSRV->GetDescriptorHeapIndex());
+			SRVPool.AddFencedObject(ViewHandleDesc, GetParentDevice()->GetDefaultCommandContext().CommandListHandle);
+		}
+
+		// Prune old allocations.
+		ResourceViewHandleDesc PruneViewHandleDesc;
+		while (SRVPool.TryAllocate(PruneViewHandleDesc, PruneFenceCount))
+		{
+			FD3D12ResourceAllocator *Allocator = PruneViewHandleDesc.BlockInfo->Allocator;
+			check(!!Allocator);
+
+			Allocator->ExpireBlock(PruneViewHandleDesc.BlockInfo);
+		}
+
+		// Try grabbing a free SRV handle from the pool
+		ResourceViewHandleDesc ViewHandleDesc;
+		if (SRVPool.TryAllocate(ViewHandleDesc))
+		{
+			pData = ViewHandleDesc.BlockInfo->Address;
+			ResourceLocation->SetBlockInfoNoUpdate(ViewHandleDesc.BlockInfo);
+			DynamicSRV->Rename(ResourceLocation, ViewHandleDesc.SRVHandle, ViewHandleDesc.DescriptorHeapIndex);
+		}
+		else
+		{
+			FD3D12DynamicHeapAllocator &Allocator = GetParentDevice()->GetDefaultUploadHeapAllocator();
+			TRefCountPtr<FD3D12ResourceBlockInfo> pBlock;
+			Allocator.Alloc(ResourceLocation->GetEffectiveBufferSize(), D3D12_TEXTURE_DATA_PLACEMENT_ALIGNMENT, pBlock.GetInitReference());
+			pData = pBlock->Address;
+			ResourceLocation->SetBlockInfoNoUpdate(pBlock);
+			DynamicSRV->Rename(ResourceLocation, ViewHandleDesc.SRVHandle, ViewHandleDesc.DescriptorHeapIndex);
+		}
+	}
+	else
+	{
+        FD3D12DynamicHeapAllocator &Allocator = GetParentDevice()->GetDefaultUploadHeapAllocator();
+		pData = Allocator.Alloc(ResourceLocation->GetEffectiveBufferSize(), D3D12_TEXTURE_DATA_PLACEMENT_ALIGNMENT, ResourceLocation);
+	}
+
+	check(pData);
+
+	return pData;
+}
+
+void* FD3D12DynamicRHI::RHILockVertexBuffer(FVertexBufferRHIParamRef VertexBufferRHI, uint32 Offset, uint32 Size, EResourceLockMode LockMode)
+{
+	check(Size > 0);
+
+	FD3D12VertexBuffer*  VertexBuffer = FD3D12DynamicRHI::ResourceCast(VertexBufferRHI);
+
+    FD3D12CommandContext& DefaultContext = GetRHIDevice()->GetDefaultCommandContext();
+	FD3D12CommandListHandle& hCommandList = DefaultContext.CommandListHandle;
+
+	// If this resource is bound to the device, unbind it
+	DefaultContext.ConditionalClearShaderResource(VertexBuffer->ResourceLocation);
+
+	// Determine whether the vertex buffer is dynamic or not.
+	const bool bIsDynamic = (VertexBuffer->GetUsage() & BUF_AnyDynamic) ? true : false;
+
+	FD3D12LockedKey LockedKey(VertexBuffer->ResourceLocation.GetReference());
+	FD3D12LockedData LockedData;
+
+	if (bIsDynamic)
+	{
+		check(LockMode == RLM_WriteOnly);
+
+		void *pData = VertexBuffer->DynamicLock();
+
+		// Add the lock to the lock map.
+		LockedData.SetData(pData);
+		LockedData.Pitch = Size;
+		OutstandingLocks.Add(LockedKey, LockedData);
+
+		return (void *)((uint8*)pData + Offset);
+	}
+	else
+	{
+		FD3D12Resource *pResource = VertexBuffer->ResourceLocation->GetResource();
+		uint32 ResourceSize = VertexBuffer->ResourceLocation->GetEffectiveBufferSize();
+
+		if (LockMode == RLM_ReadOnly)
+		{
+			// If the static buffer is being locked for reading, create a staging buffer.
+			TRefCountPtr<FD3D12Resource> StagingVertexBuffer;
+            VERIFYD3D11RESULT(GetRHIDevice()->GetResourceHelper().CreateBuffer(D3D12_HEAP_TYPE_READBACK, Offset + Size, StagingVertexBuffer.GetInitReference()));
+			LockedData.StagingResource = StagingVertexBuffer;
+
+			// Copy the contents of the vertex buffer to the staging buffer.
+			{
+				FScopeResourceBarrier ScopeResourceBarrierSource(hCommandList, pResource, pResource->GetDefaultResourceState(), D3D12_RESOURCE_STATE_COPY_SOURCE, 0);
+				// Don't need to transition upload heaps
+
+				DefaultContext.numCopies++;
+				hCommandList->CopyBufferRegion(
+					StagingVertexBuffer->GetResource(),
+					0,
+					pResource->GetResource(),
+					Offset, Size);
+			}
+
+            GetRHIDevice()->GetDefaultCommandContext().FlushCommands(true);
+
+			// Map the staging buffer's memory for reading.
+			void* pData;
+			VERIFYD3D11RESULT(StagingVertexBuffer->GetResource()->Map(0, nullptr, &pData));
+			LockedData.SetData(pData);
+			LockedData.Pitch = Offset + Size;
+		}
+		else
+		{
+			// If the static buffer is being locked for writing, allocate memory for the contents to be written to.
+
+			// Use an upload heap to copy data to a default resource.
+
+			// Use an upload heap for dynamic resources and map its memory for writing.
+			TRefCountPtr<FD3D12ResourceLocation> pUploadBufferLocation = new FD3D12ResourceLocation();
+            void *pData = GetRHIDevice()->GetDefaultUploadHeapAllocator().FastAlloc(Offset + Size, D3D12_TEXTURE_DATA_PLACEMENT_ALIGNMENT, pUploadBufferLocation);
+
+			if (nullptr == pData)
+			{
+				return nullptr;
+			}
+
+
+			// Add the lock to the lock map.
+			FD3D12LockedData LockedData;
+			LockedData.SetData(pData);
+			LockedData.Pitch = Offset + Size;
+
+			// Keep track of the underlying resource for the upload heap so it can be referenced during Unmap.
+			LockedData.UploadHeapLocation = pUploadBufferLocation;
+			OutstandingLocks.Add(LockedKey, LockedData);
+
+			return (void *)((uint8*)pData + Offset);
+		}
+	}
+
+	// Add the lock to the lock map.
+	OutstandingLocks.Add(LockedKey, LockedData);
+
+	// Return the offset pointer
+	return (void*)((uint8*)LockedData.GetData() + Offset);
+}
+
+void FD3D12DynamicRHI::RHIUnlockVertexBuffer(FVertexBufferRHIParamRef VertexBufferRHI)
+{
+	FD3D12VertexBuffer*  VertexBuffer = FD3D12DynamicRHI::ResourceCast(VertexBufferRHI);
+
+	// Determine whether the vertex buffer is dynamic or not.
+	const bool bIsDynamic = (VertexBuffer->GetUsage() & BUF_AnyDynamic) ? true : false;
+
+	// Find the outstanding lock for this VB.
+	FD3D12LockedKey LockedKey(VertexBuffer->ResourceLocation);
+	FD3D12LockedData* LockedData = OutstandingLocks.Find(LockedKey);
+	check(LockedData);
+
+	if (bIsDynamic)
+	{
+		// If the VB is dynamic, its upload heap memory can always stay mapped. Don't do anything.
+	}
+	else
+	{
+		// If the static VB lock involved a staging resource, it was locked for reading.
+		if (LockedData->StagingResource)
+		{
+			// Unmap the staging buffer's memory.
+			ID3D12Resource* StagingBuffer = LockedData->StagingResource.GetReference()->GetResource();
+			StagingBuffer->Unmap(0, nullptr);
+		}
+		else
+		{
+			FD3D12Resource *pResource = VertexBuffer->ResourceLocation->GetResource();
+
+			// Copy the contents of the temporary memory buffer allocated for writing into the IB.
+			FD3D12ResourceLocation* pUploadHeapLocation = LockedData->UploadHeapLocation.GetReference();
+			FD3D12Resource* UploadHeap = pUploadHeapLocation->GetResource();
+
+			FD3D12CommandContext& defaultContext = GetRHIDevice()->GetDefaultCommandContext();
+			FD3D12CommandListHandle& hCommandList = defaultContext.CommandListHandle;
+
+			{
+				FScopeResourceBarrier ScopeResourceBarrierDest(hCommandList, pResource, pResource->GetDefaultResourceState(), D3D12_RESOURCE_STATE_COPY_DEST, 0);
+				// Don't need to transition upload heaps
+
+				defaultContext.numCopies++;
+				hCommandList->CopyBufferRegion(
+					pResource->GetResource(),
+					VertexBuffer->ResourceLocation->GetOffset(),
+					UploadHeap->GetResource(),
+					pUploadHeapLocation->GetOffset(), LockedData->Pitch);
+			}
+
+			DEBUG_RHI_EXECUTE_COMMAND_LIST(this);
+		}
+	}
+
+	// Remove the FD3D12LockedData from the lock map.
+	// If the lock involved a staging resource, this releases it.
+	OutstandingLocks.Remove(LockedKey);
+}
+
+void FD3D12DynamicRHI::RHICopyVertexBuffer(FVertexBufferRHIParamRef SourceBufferRHI, FVertexBufferRHIParamRef DestBufferRHI)
+{
+	FD3D12VertexBuffer*  SourceBuffer = FD3D12DynamicRHI::ResourceCast(SourceBufferRHI);
+	FD3D12VertexBuffer*  DestBuffer = FD3D12DynamicRHI::ResourceCast(DestBufferRHI);
+
+	FD3D12Resource *pSourceResource = SourceBuffer->ResourceLocation->GetResource();
+	D3D12_RESOURCE_DESC const& SourceBufferDesc = pSourceResource->GetDesc();
+
+	FD3D12Resource *pDestResource = DestBuffer->ResourceLocation->GetResource();
+	D3D12_RESOURCE_DESC const& DestBufferDesc = pDestResource->GetDesc();
+
+	check(SourceBufferDesc.Width == DestBufferDesc.Width);
+	check(SourceBuffer->GetSize() == DestBuffer->GetSize());
+
+	GetRHIDevice()->GetDefaultCommandContext().numCopies++;
+    GetRHIDevice()->GetDefaultCommandContext().CommandListHandle->CopyResource(pDestResource->GetResource(), pSourceResource->GetResource());
+	DEBUG_RHI_EXECUTE_COMMAND_LIST(this);
+
+	GPUProfilingData.RegisterGPUWork(1);
+}
diff --git a/Engine/Source/Runtime/Windows/D3D12RHI/Private/D3D12Viewport.cpp b/Engine/Source/Runtime/Windows/D3D12RHI/Private/D3D12Viewport.cpp
new file mode 100644
--- /dev/null
+++ b/Engine/Source/Runtime/Windows/D3D12RHI/Private/D3D12Viewport.cpp
@@ -0,0 +1,698 @@
+// Copyright 1998-2014 Epic Games, Inc. All Rights Reserved.
+
+/*=============================================================================
+	D3D12Viewport.cpp: D3D viewport RHI implementation.
+	=============================================================================*/
+
+#include "D3D12RHIPrivate.h"
+#include "RenderCore.h"
+
+#ifndef D3D11_WITH_DWMAPI
+#if WINVER > 0x502		// Windows XP doesn't support DWM
+#define D3D11_WITH_DWMAPI	0
+#else
+#define D3D11_WITH_DWMAPI	0
+#endif
+#endif
+
+#if D3D11_WITH_DWMAPI
+#include "AllowWindowsPlatformTypes.h"
+#include <dwmapi.h>
+#endif	//D3D11_WITH_DWMAPI
+
+namespace D3D12RHI
+{
+	/**
+	 * RHI console variables used by viewports.
+	 */
+	namespace RHIConsoleVariables
+	{
+		int32 bSyncWithDWM = 0;
+		static FAutoConsoleVariableRef CVarSyncWithDWM(
+		TEXT("D3D12RHI.SyncWithDWM"),
+			bSyncWithDWM,
+			TEXT("If true, synchronize with the desktop window manager for vblank."),
+			ECVF_RenderThreadSafe
+			);
+
+		float RefreshPercentageBeforePresent = 1.0f;
+		static FAutoConsoleVariableRef CVarRefreshPercentageBeforePresent(
+		TEXT("D3D12RHI.RefreshPercentageBeforePresent"),
+			RefreshPercentageBeforePresent,
+			TEXT("The percentage of the refresh period to wait before presenting."),
+			ECVF_RenderThreadSafe
+			);
+
+		int32 bForceThirtyHz = 1;
+		static FAutoConsoleVariableRef CVarForceThirtyHz(
+		TEXT("D3D12RHI.ForceThirtyHz"),
+			bForceThirtyHz,
+			TEXT("If true, the display will never update more often than 30Hz."),
+			ECVF_RenderThreadSafe
+			);
+
+		int32 SyncInterval = 1;
+		static FAutoConsoleVariableRef CVarSyncInterval(
+		TEXT("D3D12RHI.SyncInterval"),
+			SyncInterval,
+			TEXT("When synchronizing with D3D, specifies the interval at which to refresh."),
+			ECVF_RenderThreadSafe
+			);
+
+		float SyncRefreshThreshold = 1.05f;
+		static FAutoConsoleVariableRef CVarSyncRefreshThreshold(
+		TEXT("D3D12RHI.SyncRefreshThreshold"),
+			SyncRefreshThreshold,
+			TEXT("Threshold for time above which vsync will be disabled as a percentage of the refresh rate."),
+			ECVF_RenderThreadSafe
+			);
+
+		int32 MaxSyncCounter = 8;
+		static FAutoConsoleVariableRef CVarMaxSyncCounter(
+			TEXT("D3D12RHI.MaxSyncCounter"),
+			MaxSyncCounter,
+			TEXT("Maximum sync counter to smooth out vsync transitions."),
+			ECVF_RenderThreadSafe
+			);
+
+		int32 SyncThreshold = 7;
+		static FAutoConsoleVariableRef CVarSyncThreshold(
+			TEXT("D3D12RHI.SyncThreshold"),
+			SyncThreshold,
+			TEXT("Number of consecutive 'fast' frames before vsync is enabled."),
+			ECVF_RenderThreadSafe
+			);
+
+		int32 MaximumFrameLatency = 3;
+		static FAutoConsoleVariableRef CVarMaximumFrameLatency(
+			TEXT("D3D12RHI.MaximumFrameLatency"),
+			MaximumFrameLatency,
+			TEXT("Number of frames that can be queued for render."),
+			ECVF_RenderThreadSafe
+			);
+
+	};
+
+	extern void D3D11TextureAllocated2D(FD3D12Texture2D& Texture);
+}
+using namespace D3D12RHI;
+
+/**
+ * Creates a FD3D12Surface to represent a swap chain's back buffer.
+ */
+FD3D12Texture2D* GetSwapChainSurface(FD3D12Device* Parent, EPixelFormat PixelFormat, IDXGISwapChain3* SwapChain, const uint32 &backBufferIndex)
+{
+	// Grab the back buffer
+	TRefCountPtr<ID3D12Resource> BackBufferResource;
+	VERIFYD3D11RESULT_EX(SwapChain->GetBuffer(backBufferIndex, __uuidof(ID3D12Resource), (void**)BackBufferResource.GetInitReference()), Parent->GetDevice());
+
+	D3D12_RESOURCE_DESC BackBufferDesc = BackBufferResource->GetDesc();
+	TRefCountPtr<FD3D12Resource> BackBufferWrappedResource = new FD3D12Resource(BackBufferResource, BackBufferDesc);
+
+	// create the render target view
+	TRefCountPtr<FD3D12RenderTargetView> BackBufferRenderTargetView;
+	D3D12_RENDER_TARGET_VIEW_DESC RTVDesc = {};
+	RTVDesc.Format = DXGI_FORMAT_UNKNOWN;
+	RTVDesc.ViewDimension = D3D12_RTV_DIMENSION_TEXTURE2D;
+	RTVDesc.Texture2D.MipSlice = 0;
+	BackBufferRenderTargetView = new FD3D12RenderTargetView(Parent, &RTVDesc, BackBufferWrappedResource);
+
+	TArray<TRefCountPtr<FD3D12RenderTargetView> > RenderTargetViews;
+	RenderTargetViews.Add(BackBufferRenderTargetView);
+
+	// create a shader resource view to allow using the backbuffer as a texture
+	D3D12_SHADER_RESOURCE_VIEW_DESC SRVDesc = {};
+	SRVDesc.Shader4ComponentMapping = D3D12_DEFAULT_SHADER_4_COMPONENT_MAPPING;
+	SRVDesc.Format = DXGI_FORMAT_UNKNOWN;
+	SRVDesc.ViewDimension = D3D12_SRV_DIMENSION_TEXTURE2D;
+	SRVDesc.Texture2D.MostDetailedMip = 0;
+	SRVDesc.Texture2D.MipLevels = 1;
+
+	FD3D12Texture2D* NewTexture = new FD3D12Texture2D(
+		Parent,
+		BackBufferWrappedResource,
+		false,
+		1,
+		RenderTargetViews,
+		NULL,
+		(uint32)BackBufferDesc.Width,
+		BackBufferDesc.Height,
+		1,
+		1,
+		1,
+		PixelFormat,
+		false,
+		false,
+		false,
+		FClearValueBinding()
+		);
+
+    FD3D12ShaderResourceView* pWrappedShaderResourceView = new FD3D12ShaderResourceView(Parent, &SRVDesc, NewTexture->ResourceLocation);
+	NewTexture->SetShaderResourceView(pWrappedShaderResourceView);
+
+	D3D11TextureAllocated2D(*NewTexture);
+
+	NewTexture->DoNoDeferDelete();
+
+	return NewTexture;
+}
+
+FD3D12Viewport::~FD3D12Viewport()
+{
+	check(IsInRenderingThread());
+
+	GetParentDevice()->GetViewports().Remove(this);
+}
+
+DXGI_MODE_DESC FD3D12Viewport::SetupDXGI_MODE_DESC() const
+{
+	DXGI_MODE_DESC Ret;
+
+	Ret.Width = SizeX;
+	Ret.Height = SizeY;
+	Ret.RefreshRate.Numerator = 0;	// illamas: use 0 to avoid a potential mismatch with hw
+	Ret.RefreshRate.Denominator = 0;	// illamas: ditto
+	Ret.Format = GetRenderTargetFormat(PixelFormat);
+	Ret.ScanlineOrdering = DXGI_MODE_SCANLINE_ORDER_UNSPECIFIED;
+	Ret.Scaling = DXGI_MODE_SCALING_UNSPECIFIED;
+
+	return Ret;
+}
+
+void FD3D12Viewport::Resize(uint32 InSizeX, uint32 InSizeY, bool bInIsFullscreen)
+{
+	// Unbind any dangling references to resources
+	GetParentDevice()->GetDefaultCommandContext().ClearState();
+
+	if (IsValidRef(CustomPresent))
+	{
+		CustomPresent->OnBackBufferResize();
+	}
+
+	// Release our backbuffer reference, as required by DXGI before calling ResizeBuffers.
+	for (uint32 i = 0; i < NumBackBuffers; ++i)
+	{
+		if (IsValidRef(BackBuffers[i]))
+		{
+			check(BackBuffers[i]->GetRefCount() == 1);
+			check(BackBuffers[i]->ResourceLocation->GetRefCount() == 2);
+		}
+		BackBuffers[i].SafeRelease();
+		check(BackBuffers[i] == nullptr);
+	}
+	
+	// Flush all pending deletes. The resize may not happen on the rendering thread so we need to
+	// enqueue a command on the rendering thread to ensure it happens.
+	ENQUEUE_UNIQUE_RENDER_COMMAND(FlushCommand,
+	{
+		FRHIResource::FlushPendingDeletes();
+	}
+	);
+	FlushRenderingCommands();
+
+	// Flush the outstanding GPU work
+	GetParentDevice()->GetDefaultCommandContext().FlushCommands();
+
+	// Indicate that we've flushed the outstanding work for the current frame and wait for the command queue to flush
+	// This is needed so the deferred deletion queue empties out completely because the resources in the
+	// queue use the frame fence values.
+	GetParentDevice()->GetCommandListManager().SignalFrameComplete(true);
+
+	// Flush the deferred deletion queue to bring backbuffer refcount to 0.
+	// Note: Clear ensures that the queue is completely empty.
+	GetParentDevice()->GetDeferredDeletionQueue().Clear();
+
+	bool bResizeBuffers = false;
+	if (SizeX != InSizeX || SizeY != InSizeY)
+	{
+		SizeX = InSizeX;
+		SizeY = InSizeY;
+
+		check(SizeX > 0);
+		check(SizeY > 0);
+
+		if (bInIsFullscreen)
+		{
+			DXGI_MODE_DESC BufferDesc = SetupDXGI_MODE_DESC();
+
+			if (FAILED(SwapChain3->ResizeTarget(&BufferDesc)))
+			{
+				ConditionalResetSwapChain(true);
+			}
+		}
+
+		bResizeBuffers = true;
+	}
+
+	if (bIsFullscreen != bInIsFullscreen)
+	{
+		bIsFullscreen = bInIsFullscreen;
+		bIsValid = false;
+
+		// Use ConditionalResetSwapChain to call SetFullscreenState, to handle the failure case.
+		// Ignore the viewport's focus state; since Resize is called as the result of a user action we assume authority without waiting for Focus.
+		ConditionalResetSwapChain(true);
+
+		bResizeBuffers = true;
+	}
+
+	if (bResizeBuffers)
+	{
+		DXGI_SWAP_CHAIN_FLAG swapChainFlags = DXGI_SWAP_CHAIN_FLAG_ALLOW_MODE_SWITCH;
+		if (bIsBenchmarkMode)
+		{
+			swapChainFlags = static_cast<DXGI_SWAP_CHAIN_FLAG>(static_cast<uint32>(swapChainFlags) | DXGI_SWAP_CHAIN_FLAG_FRAME_LATENCY_WAITABLE_OBJECT);
+		}
+
+		// Resize the swap chain.
+		VERIFYD3D11RESULT_EX(SwapChain3->ResizeBuffers(NumBackBuffers, SizeX, SizeY, GetRenderTargetFormat(PixelFormat), swapChainFlags), GetParentDevice()->GetDevice());
+	}
+
+	// Create a RHI surface to represent the viewport's back buffer.
+	for (uint32 i = 0; i < NumBackBuffers; ++i)
+	{
+		check(BackBuffers[i].GetReference() == nullptr);
+		BackBuffers[i] = GetSwapChainSurface(GetParentDevice(), PixelFormat, SwapChain3, i);
+	}
+
+	// Reset the viewport frame counter to get the correct back buffer after the resize.
+	GetParentDevice()->GetOwningRHI()->ResetViewportFrameCounter();
+}
+
+/** Returns true if desktop composition is enabled. */
+static bool IsCompositionEnabled()
+{
+	BOOL bDwmEnabled = false;
+#if D3D11_WITH_DWMAPI
+	DwmIsCompositionEnabled(&bDwmEnabled);
+#endif	//D3D11_WITH_DWMAPI
+	return !!bDwmEnabled;
+}
+
+/** Presents the swap chain checking the return result. */
+bool FD3D12Viewport::PresentChecked(int32 SyncInterval)
+{
+	HRESULT Result = S_OK;
+	bool bNeedNativePresent = true;
+	if (IsValidRef(CustomPresent))
+	{
+		bNeedNativePresent = CustomPresent->Present(SyncInterval);
+	}
+	if (bNeedNativePresent)
+	{
+		// Present the back buffer to the viewport window.
+		Result = SwapChain3->Present(SyncInterval, 0);
+
+		// Signal the frame is complete.
+		GetParentDevice()->GetCommandListManager().SignalFrameComplete();
+
+#if UE_BUILD_DEBUG
+		UE_LOG(LogD3D12RHI, Log, TEXT("*** PRESENT ***"), this);
+		GetParentDevice()->GetOwningRHI()->DrawCount = 0;
+		GetParentDevice()->GetOwningRHI()->PresentCount++;
+
+		// Dump stats every 240 frames
+		if (GetParentDevice()->GetOwningRHI()->PresentCount % 240 == 0)
+		{
+			FOutputDeviceRedirector* pOutputDevice = FOutputDeviceRedirector::Get();
+			GetParentDevice()->GetDefaultUploadHeapAllocator().DumpAllocatorStats(*pOutputDevice);
+		}
+#endif
+	}
+
+	// Detect a lost device.
+	if (Result == DXGI_ERROR_DEVICE_REMOVED || Result == DXGI_ERROR_DEVICE_RESET || Result == DXGI_ERROR_DRIVER_INTERNAL_ERROR)
+	{
+		// This variable is checked periodically by the main thread.
+		GetParentDevice()->SetDeviceRemoved(true);
+	}
+	else
+	{
+		VERIFYD3D11RESULT(Result);
+	}
+
+	return bNeedNativePresent;
+}
+
+/** Blocks the CPU to synchronize with vblank by communicating with DWM. */
+void FD3D12Viewport::PresentWithVsyncDWM()
+{
+#if D3D11_WITH_DWMAPI
+	LARGE_INTEGER Cycles;
+	DWM_TIMING_INFO TimingInfo;
+
+	// Find out how long since we last flipped and query DWM for timing information.
+	QueryPerformanceCounter(&Cycles);
+	FMemory::Memzero(TimingInfo);
+	TimingInfo.cbSize = sizeof(DWM_TIMING_INFO);
+	DwmGetCompositionTimingInfo(WindowHandle, &TimingInfo);
+
+	uint64 QpcAtFlip = Cycles.QuadPart;
+	uint64 CyclesSinceLastFlip = Cycles.QuadPart - LastFlipTime;
+	float CPUTime = FPlatformTime::ToMilliseconds(CyclesSinceLastFlip);
+	float GPUTime = FPlatformTime::ToMilliseconds(TimingInfo.qpcFrameComplete - LastCompleteTime);
+	float DisplayRefreshPeriod = FPlatformTime::ToMilliseconds(TimingInfo.qpcRefreshPeriod);
+
+	// Find the smallest multiple of the refresh rate that is >= 33ms, our target frame rate.
+	float RefreshPeriod = DisplayRefreshPeriod;
+	if (RHIConsoleVariables::bForceThirtyHz && RefreshPeriod > 1.0f)
+	{
+		while (RefreshPeriod - (1000.0f / 30.0f) < -1.0f)
+		{
+			RefreshPeriod *= 2.0f;
+		}
+	}
+
+	// If the last frame hasn't completed yet, we don't know how long the GPU took.
+	bool bValidGPUTime = (TimingInfo.cFrameComplete > LastFrameComplete);
+	if (bValidGPUTime)
+	{
+		GPUTime /= (float)(TimingInfo.cFrameComplete - LastFrameComplete);
+	}
+
+	// Update the sync counter depending on how much time it took to complete the previous frame.
+	float FrameTime = FMath::Max<float>(CPUTime, GPUTime);
+	if (FrameTime >= RHIConsoleVariables::SyncRefreshThreshold * RefreshPeriod)
+	{
+		SyncCounter--;
+	}
+	else if (bValidGPUTime)
+	{
+		SyncCounter++;
+	}
+	SyncCounter = FMath::Clamp<int32>(SyncCounter, 0, RHIConsoleVariables::MaxSyncCounter);
+
+	// If frames are being completed quickly enough, block for vsync.
+	bool bSync = (SyncCounter >= RHIConsoleVariables::SyncThreshold);
+	if (bSync)
+	{
+		// This flushes the previous present call and blocks until it is made available to DWM.
+		GetParentDevice()->GetDefaultCommandContext().FlushCommands();
+		// MS: Might need to wait for the previous command list to finish
+
+		DwmFlush();
+
+		// We sleep a percentage of the remaining time. The trick is to get the
+		// present call in after the vblank we just synced for but with time to
+		// spare for the next vblank.
+		float MinFrameTime = RefreshPeriod * RHIConsoleVariables::RefreshPercentageBeforePresent;
+		float TimeToSleep;
+		do
+		{
+			QueryPerformanceCounter(&Cycles);
+			float TimeSinceFlip = FPlatformTime::ToMilliseconds(Cycles.QuadPart - LastFlipTime);
+			TimeToSleep = (MinFrameTime - TimeSinceFlip);
+			if (TimeToSleep > 0.0f)
+			{
+				FPlatformProcess::Sleep(TimeToSleep * 0.001f);
+			}
+		} while (TimeToSleep > 0.0f);
+	}
+
+	// Present.
+	PresentChecked(/*SyncInterval=*/ 0);
+
+	// If we are forcing <= 30Hz, block the CPU an additional amount of time if needed.
+	// This second block is only needed when RefreshPercentageBeforePresent < 1.0.
+	if (bSync)
+	{
+		LARGE_INTEGER LocalCycles;
+		float TimeToSleep;
+		bool bSaveCycles = false;
+		do
+		{
+			QueryPerformanceCounter(&LocalCycles);
+			float TimeSinceFlip = FPlatformTime::ToMilliseconds(LocalCycles.QuadPart - LastFlipTime);
+			TimeToSleep = (RefreshPeriod - TimeSinceFlip);
+			if (TimeToSleep > 0.0f)
+			{
+				bSaveCycles = true;
+				FPlatformProcess::Sleep(TimeToSleep * 0.001f);
+			}
+		} while (TimeToSleep > 0.0f);
+
+		if (bSaveCycles)
+		{
+			Cycles = LocalCycles;
+		}
+	}
+
+	// If we are dropping vsync reset the counter. This provides a debounce time
+	// before which we try to vsync again.
+	if (!bSync && bSyncedLastFrame)
+	{
+		SyncCounter = 0;
+	}
+
+	if (bSync != bSyncedLastFrame || UE_LOG_ACTIVE(LogRHI, VeryVerbose))
+	{
+		UE_LOG(LogRHI, Verbose, TEXT("BlockForVsync[%d]: CPUTime:%.2fms GPUTime[%d]:%.2fms Blocked:%.2fms Pending/Complete:%d/%d"),
+			bSync,
+			CPUTime,
+			bValidGPUTime,
+			GPUTime,
+			FPlatformTime::ToMilliseconds(Cycles.QuadPart - QpcAtFlip),
+			TimingInfo.cFramePending,
+			TimingInfo.cFrameComplete);
+	}
+
+	// Remember if we synced, when the frame completed, etc.
+	bSyncedLastFrame = bSync;
+	LastFlipTime = Cycles.QuadPart;
+	LastFrameComplete = TimingInfo.cFrameComplete;
+	LastCompleteTime = TimingInfo.qpcFrameComplete;
+#endif	//D3D11_WITH_DWMAPI
+}
+
+bool FD3D12Viewport::Present(bool bLockToVsync)
+{
+	FD3D12CommandContext& DefaultContext = GetParentDevice()->GetDefaultCommandContext();
+
+	bool bNativelyPresented = true;
+	FD3D12DynamicRHI::TransitionResource(DefaultContext.CommandListHandle, GetBackBuffer()->GetShaderResourceView(), D3D12_RESOURCE_STATE_PRESENT);
+
+	DefaultContext.FlushCommands();
+
+	// Reset the default context state
+	DefaultContext.ClearState();
+
+	int32 syncInterval = bLockToVsync ? RHIConsoleVariables::SyncInterval : 0;
+
+	bNativelyPresented = PresentChecked(syncInterval);
+
+	return bNativelyPresented;
+}
+
+/*=============================================================================
+ *	The following RHI functions must be called from the main thread.
+ *=============================================================================*/
+FViewportRHIRef FD3D12DynamicRHI::RHICreateViewport(void* WindowHandle, uint32 SizeX, uint32 SizeY, bool bIsFullscreen, EPixelFormat PreferredPixelFormat)
+{
+	check(IsInGameThread());
+
+	// Use a default pixel format if none was specified	
+	if (PreferredPixelFormat == EPixelFormat::PF_Unknown)
+	{
+		PreferredPixelFormat = EPixelFormat::PF_A2B10G10R10;
+	}
+
+	// MSFT: Seb: Remove the Init? Thats what XB1 does.
+	FD3D12Viewport* RenderingViewport = new FD3D12Viewport(GetRHIDevice(), (HWND)WindowHandle, SizeX, SizeY, bIsFullscreen, PreferredPixelFormat);
+	RenderingViewport->Init(GetFactory(), true);
+	return RenderingViewport;
+}
+
+void FD3D12DynamicRHI::RHIResizeViewport(FViewportRHIParamRef ViewportRHI, uint32 SizeX, uint32 SizeY, bool bIsFullscreen)
+{
+	FD3D12Viewport*  Viewport = FD3D12DynamicRHI::ResourceCast(ViewportRHI);
+
+	check(IsInGameThread());
+	Viewport->Resize(SizeX, SizeY, bIsFullscreen);
+}
+
+void FD3D12DynamicRHI::RHITick(float DeltaTime)
+{
+	check(IsInGameThread());
+
+	// Check to see if the device has been removed.
+	if (GetRHIDevice()->IsDeviceRemoved())
+	{
+		InitD3DDevices();
+	}
+
+	// Check if any swap chains have been invalidated.
+	for (int32 ViewportIndex = 0; ViewportIndex < GetRHIDevice()->GetViewports().Num(); ViewportIndex++)
+	{
+		GetRHIDevice()->GetViewports()[ViewportIndex]->ConditionalResetSwapChain(false);
+	}
+}
+
+/*=============================================================================
+ *	Viewport functions.
+ *=============================================================================*/
+
+//#if PLATFORM_SUPPORTS_RHI_THREAD
+void FD3D12CommandContext::RHIBeginDrawingViewport(FViewportRHIParamRef ViewportRHI, FTextureRHIParamRef RenderTarget)
+{
+	FD3D12CommandContext& CmdContext = *this;
+	FD3D12Device* Device = GetParentDevice();
+//#else
+//void FD3D12DynamicRHI::RHIBeginDrawingViewport(FViewportRHIParamRef ViewportRHI, FTextureRHIParamRef RenderTarget)
+//{
+//	FD3D12CommandContext& CmdContext = GetRHIDevice()->GetDefaultCommandContext();
+//	FD3D12Device* Device = GetRHIDevice();
+//#endif
+
+	FD3D12Viewport*  Viewport = FD3D12DynamicRHI::ResourceCast(ViewportRHI);
+
+	SCOPE_CYCLE_COUNTER(STAT_D3D12PresentTime);
+
+	check(!Device->GetDrawingViewport());
+	Device->SetDrawingViewport( Viewport );
+
+	// Set the render target and viewport.
+	// MSFT: Seb
+	//check (RenderTarget);
+	if (RenderTarget == NULL)
+	{
+		RenderTarget = Viewport->GetBackBuffer();
+	}
+	FRHIRenderTargetView View(RenderTarget);
+	CmdContext.RHISetRenderTargets(1, &View, nullptr, 0, NULL);
+
+	// Set an initially disabled scissor rect.
+	CmdContext.RHISetScissorRect(false, 0, 0, 0, 0);
+
+	// MSFT: Seb
+	//CmdContext.FlushCommands();
+}
+
+//#if PLATFORM_SUPPORTS_RHI_THREAD
+void FD3D12CommandContext::RHIEndDrawingViewport(FViewportRHIParamRef ViewportRHI, bool bPresent, bool bLockToVsync)
+{
+	//check(IsDefaultContext());
+	FD3D12Device* Device = GetParentDevice();
+	FD3D12CommandContext& CmdContext = *this;
+	FD3D12DynamicRHI& RHI = *Device->GetOwningRHI();
+//#else
+//void FD3D12DynamicRHI::RHIEndDrawingViewport(FViewportRHIParamRef ViewportRHI, bool bPresent, bool bLockToVsync)
+//{
+//	FD3D12Device* Device = GetRHIDevice();
+//	FD3D12CommandContext& CmdContext = Device->GetDefaultCommandContext();
+//	FD3D12DynamicRHI& RHI = *this;
+//#endif
+
+	FD3D12Viewport*  Viewport = FD3D12DynamicRHI::ResourceCast(ViewportRHI);
+
+	SCOPE_CYCLE_COUNTER(STAT_D3D12PresentTime);
+
+	check(Device->GetDrawingViewport() == Viewport);
+	Device->SetDrawingViewport(nullptr);
+
+	// Clear references the device might have to resources.
+	CmdContext.CurrentDepthTexture = NULL;
+	CmdContext.CurrentDepthStencilTarget = NULL;
+	CmdContext.CurrentDSVAccessType = FExclusiveDepthStencil::DepthWrite_StencilWrite;
+	CmdContext.CurrentRenderTargets[0] = NULL;
+	for (uint32 RenderTargetIndex = 1; RenderTargetIndex < D3D12_SIMULTANEOUS_RENDER_TARGET_COUNT; ++RenderTargetIndex)
+	{
+		CmdContext.CurrentRenderTargets[RenderTargetIndex] = NULL;
+	}
+
+	CmdContext.ClearAllShaderResources();
+
+	CmdContext.CommitRenderTargetsAndUAVs();
+
+	CmdContext.StateCache.SetBoundShaderState(nullptr);
+
+	for (uint32 StreamIndex = 0; StreamIndex < D3D12_IA_VERTEX_INPUT_RESOURCE_SLOT_COUNT; StreamIndex++)
+	{
+		CmdContext.StateCache.SetStreamSource(nullptr, StreamIndex, 0, 0);
+	}
+
+	CmdContext.StateCache.SetIndexBuffer((FD3D12ResourceLocation*)nullptr, DXGI_FORMAT_R16_UINT, 0);
+	// Compute Shader is set to NULL after each Dispatch call, so no need to clear it here
+#if STATS
+	SET_CYCLE_COUNTER(STAT_D3D12CommitResourceTables, RHI.CommitResourceTableCycles);
+	SET_CYCLE_COUNTER(STAT_D3D12CacheResourceTables, RHI.CacheResourceTableCycles);
+	SET_CYCLE_COUNTER(STAT_D3D12SetShaderTextureTime, RHI.SetShaderTextureCycles);
+	INC_DWORD_STAT_BY(STAT_D3D12SetShaderTextureCalls, RHI.SetShaderTextureCalls);
+	INC_DWORD_STAT_BY(STAT_D3D12CacheResourceTableCalls, RHI.CacheResourceTableCalls);
+	INC_DWORD_STAT_BY(STAT_D3D12SetTextureInTableCalls, RHI.SetTextureInTableCalls);
+#endif
+
+	RHI.CommitResourceTableCycles = 0;
+	RHI.CacheResourceTableCalls = 0;
+	RHI.CacheResourceTableCycles = 0;
+	RHI.SetShaderTextureCycles = 0;
+	RHI.SetShaderTextureCalls = 0;
+	RHI.SetTextureInTableCalls = 0;
+
+	bool bNativelyPresented = Viewport->Present(bLockToVsync);
+
+	// Don't wait on the GPU when using SLI, let the driver determine how many frames behind the GPU should be allowed to get
+	if (GNumActiveGPUsForRendering == 1)
+	{
+		if (bNativelyPresented)
+		{
+			static const auto CFinishFrameVar = IConsoleManager::Get().FindTConsoleVariableDataInt(TEXT("r.FinishCurrentFrame"));
+			if (CFinishFrameVar->GetValueOnRenderThread())
+			{
+				// Wait for the GPU to finish rendering the previous frame before finishing this frame.
+				CmdContext.FlushCommands(true);
+			}
+		}
+
+		// If the input latency timer has been triggered, block until the GPU is completely
+		// finished displaying this frame and calculate the delta time.
+		if (GInputLatencyTimer.RenderThreadTrigger)
+		{
+			CmdContext.FlushCommands(true);
+
+			uint32 EndTime = FPlatformTime::Cycles();
+			GInputLatencyTimer.DeltaTime = EndTime - GInputLatencyTimer.StartTime;
+			GInputLatencyTimer.RenderThreadTrigger = false;
+		}
+	}
+
+#if CHECK_SRV_TRANSITIONS
+	UnresolvedTargets.Reset();
+#endif
+}
+
+/**
+ * Determine if currently drawing the viewport
+ *
+ * @return true if currently within a BeginDrawingViewport/EndDrawingViewport block
+ */
+bool FD3D12DynamicRHI::RHIIsDrawingViewport()
+{
+	return GetRHIDevice()->GetDrawingViewport() != NULL;
+}
+
+void FD3D12DynamicRHI::RHIAdvanceFrameForGetViewportBackBuffer()
+{
+	ViewportFrameCounter++;
+}
+
+FTexture2DRHIRef FD3D12DynamicRHI::RHIGetViewportBackBuffer(FViewportRHIParamRef ViewportRHI)
+{
+	FD3D12Viewport*  Viewport = FD3D12DynamicRHI::ResourceCast(ViewportRHI);
+
+	if (GRHISupportsRHIThread)
+		return Viewport->GetBackBuffer(ViewportFrameCounter);
+	else
+		return Viewport->GetBackBuffer();
+}
+
+// MSFT: Seb
+//FTexture2DRHIRef FD3D12DynamicRHI::RHIGetViewportForegroundBuffer(FViewportRHIParamRef ViewportRHI)
+//{
+//	FD3D12Viewport* Viewport = FD3D12DynamicRHI::ResourceCast(ViewportRHI);
+//
+//	return Viewport->GetBackBuffer(ViewportFrameCounter);
+//}
+
+#if D3D11_WITH_DWMAPI
+#include "HideWindowsPlatformTypes.h"
+#endif	//D3D11_WITH_DWMAPI
diff --git a/Engine/Source/Runtime/Windows/D3D12RHI/Private/WinRT/D3D12RHIBasePrivate.h b/Engine/Source/Runtime/Windows/D3D12RHI/Private/WinRT/D3D12RHIBasePrivate.h
new file mode 100644
--- /dev/null
+++ b/Engine/Source/Runtime/Windows/D3D12RHI/Private/WinRT/D3D12RHIBasePrivate.h
@@ -0,0 +1,35 @@
+// Copyright 1998-2014 Epic Games, Inc. All Rights Reserved.
+
+/*=============================================================================
+	D3D12RHIBasePrivate.h: Private D3D RHI definitions for WinRT.
+=============================================================================*/
+
+#pragma once
+
+// Assume D3DX is available
+#ifndef WITH_D3DX_LIBS
+#define WITH_D3DX_LIBS	1
+#endif
+
+// Disable macro redefinition warning for compatibility with Windows SDK 8+
+#pragma warning(push)
+#if _MSC_VER >= 1700
+#pragma warning(disable : 4005)	// macro redefinition
+#endif
+
+// D3D headers.
+#if PLATFORM_64BITS
+#pragma pack(push,16)
+#else
+#pragma pack(push,8)
+#endif
+#define D3D_OVERLOADS 1
+#include "AllowWindowsPlatformTypes.h"
+#include <d3d12.h>
+#include "d3dx12.h"
+#include "HideWindowsPlatformTypes.h"
+
+#undef DrawText
+
+#pragma pack(pop)
+#pragma warning(pop)
diff --git a/Engine/Source/Runtime/Windows/D3D12RHI/Private/WinRT/D3D12RHIPrivateUtil.h b/Engine/Source/Runtime/Windows/D3D12RHI/Private/WinRT/D3D12RHIPrivateUtil.h
new file mode 100644
--- /dev/null
+++ b/Engine/Source/Runtime/Windows/D3D12RHI/Private/WinRT/D3D12RHIPrivateUtil.h
@@ -0,0 +1,22 @@
+// Copyright 1998-2014 Epic Games, Inc. All Rights Reserved.
+
+/*=============================================================================
+	D3D12RHIPrivateUtil.h: Private D3D RHI Utility definitions for WinRT.
+=============================================================================*/
+
+#pragma once
+
+struct FD3DRHIUtil
+{
+	template <EShaderFrequency ShaderFrequencyT>
+	static FORCEINLINE void CommitConstants(FD3D12ConstantBuffer* InConstantBuffer, FD3D12StateCache& StateCache, uint32 Index, bool bDiscardSharedConstants)
+	{
+		auto* ConstantBuffer = ((FWinD3D12ConstantBuffer*)InConstantBuffer);
+		// Array may contain NULL entries to pad out to proper 
+		if (ConstantBuffer && ConstantBuffer->CommitConstantsToDevice(bDiscardSharedConstants))
+		{
+			ID3D12Resource* DeviceBuffer = ConstantBuffer->GetConstantBuffer();
+			StateCache.SetConstantBuffer<ShaderFrequencyT>(DeviceBuffer, Index);
+		}
+	}
+};
diff --git a/Engine/Source/Runtime/Windows/D3D12RHI/Private/WinRT/WinRTD3D12Device.cpp b/Engine/Source/Runtime/Windows/D3D12RHI/Private/WinRT/WinRTD3D12Device.cpp
new file mode 100644
--- /dev/null
+++ b/Engine/Source/Runtime/Windows/D3D12RHI/Private/WinRT/WinRTD3D12Device.cpp
@@ -0,0 +1,153 @@
+// Copyright 1998-2014 Epic Games, Inc. All Rights Reserved.
+
+#include "D3D12RHIPrivate.h"
+
+bool FD3D12DynamicRHIModule::IsSupported()
+{
+	// figure out a way to detect the only allowable video card
+	return true;
+}
+
+FDynamicRHI* FD3D12DynamicRHIModule::CreateRHI()
+{
+	FDynamicRHI* DynamicRHI = new FD3D12DynamicRHI(NULL,D3D_FEATURE_LEVEL_11_0,ChosenAdapter.AdapterIndex);
+	if (DynamicRHI)
+	{
+		// Initialize the RHI capabilities.
+		GRHIShaderPlatform = SP_PCD3D_SM5;
+
+		// This will be set in the D3D11DynamicRHI::Init call!
+		//	GIsRHIInitialized = true;
+	}
+	return DynamicRHI;
+}
+
+/**
+ *	Retrieve available screen resolutions.
+ *
+ *	@param	Resolutions			TArray<FScreenResolutionRHI> parameter that will be filled in.
+ *	@param	bIgnoreRefreshRate	If true, ignore refresh rates.
+ *
+ *	@return	bool				true if successfully filled the array
+ */
+bool FD3D12DynamicRHI::RHIGetAvailableResolutions(FScreenResolutionArray& Resolutions, bool bIgnoreRefreshRate)
+{
+	int32 MinAllowableResolutionX = 0;
+	int32 MinAllowableResolutionY = 0;
+	int32 MaxAllowableResolutionX = 10480;
+	int32 MaxAllowableResolutionY = 10480;
+	int32 MinAllowableRefreshRate = 0;
+	int32 MaxAllowableRefreshRate = 10480;
+
+	if (MaxAllowableResolutionX == 0)
+	{
+		MaxAllowableResolutionX = 10480;
+	}
+	if (MaxAllowableResolutionY == 0)
+	{
+		MaxAllowableResolutionY = 10480;
+	}
+	if (MaxAllowableRefreshRate == 0)
+	{
+		MaxAllowableRefreshRate = 10480;
+	}
+
+	HRESULT hr = S_OK;
+	TRefCountPtr<IDXGIAdapter> Adapter;
+	hr = DXGIFactory->EnumAdapters(ChosenAdapter,Adapter.GetInitReference());
+
+	if( DXGI_ERROR_NOT_FOUND == hr )
+		return false;
+	if( FAILED(hr) )
+		return false;
+
+	// get the description of the adapter
+	DXGI_ADAPTER_DESC AdapterDesc;
+	VERIFYD3D11RESULT(Adapter->GetDesc(&AdapterDesc));
+
+	int32 CurrentOutput = 0;
+	do 
+	{
+		TRefCountPtr<IDXGIOutput> Output;
+		hr = Adapter->EnumOutputs(CurrentOutput,Output.GetInitReference());
+		if(DXGI_ERROR_NOT_FOUND == hr)
+			break;
+		if(FAILED(hr))
+			return false;
+
+		// TODO: GetDisplayModeList is a terribly SLOW call.  It can take up to a second per invocation.
+		//  We might want to work around some DXGI badness here.
+		DXGI_FORMAT Format = DXGI_FORMAT_R8G8B8A8_UNORM;
+		uint32 NumModes = 0;
+		hr = Output->GetDisplayModeList(Format, 0, &NumModes, NULL);
+		if(hr == DXGI_ERROR_NOT_FOUND)
+		{
+			continue;
+		}
+		else if(hr == DXGI_ERROR_NOT_CURRENTLY_AVAILABLE)
+		{
+			UE_LOG(LogD3D12RHI, Fatal,
+				TEXT("This application cannot be run over a remote desktop configuration")
+				);
+			return false;
+		}
+
+		DXGI_MODE_DESC* ModeList = new DXGI_MODE_DESC[ NumModes ];
+		VERIFYD3D11RESULT(Output->GetDisplayModeList(Format, 0, &NumModes, ModeList));
+
+		for(uint32 m = 0;m < NumModes;m++)
+		{
+			if (((int32)ModeList[m].Width >= MinAllowableResolutionX) &&
+				((int32)ModeList[m].Width <= MaxAllowableResolutionX) &&
+				((int32)ModeList[m].Height >= MinAllowableResolutionY) &&
+				((int32)ModeList[m].Height <= MaxAllowableResolutionY)
+				)
+			{
+				bool bAddIt = true;
+				if (bIgnoreRefreshRate == false)
+				{
+					if (((int32)ModeList[m].RefreshRate.Numerator < MinAllowableRefreshRate * ModeList[m].RefreshRate.Denominator) ||
+						((int32)ModeList[m].RefreshRate.Numerator > MaxAllowableRefreshRate * ModeList[m].RefreshRate.Denominator)
+						)
+					{
+						continue;
+					}
+				}
+				else
+				{
+					// See if it is in the list already
+					for (int32 CheckIndex = 0; CheckIndex < Resolutions.Num(); CheckIndex++)
+					{
+						FScreenResolutionRHI& CheckResolution = Resolutions[CheckIndex];
+						if ((CheckResolution.Width == ModeList[m].Width) &&
+							(CheckResolution.Height == ModeList[m].Height))
+						{
+							// Already in the list...
+							bAddIt = false;
+							break;
+						}
+					}
+				}
+
+				if (bAddIt)
+				{
+					// Add the mode to the list
+					int32 Temp2Index = Resolutions.AddZeroed();
+					FScreenResolutionRHI& ScreenResolution = Resolutions[Temp2Index];
+
+					ScreenResolution.Width = ModeList[m].Width;
+					ScreenResolution.Height = ModeList[m].Height;
+					ScreenResolution.RefreshRate = ModeList[m].RefreshRate.Numerator / ModeList[m].RefreshRate.Denominator;
+				}
+			}
+		}
+
+		delete[] ModeList;
+
+		++CurrentOutput;
+
+	// TODO: Cap at 1 for default output
+	} while(CurrentOutput < 1);
+
+	return true;
+}
diff --git a/Engine/Source/Runtime/Windows/D3D12RHI/Private/WinRT/WinRTD3D12Viewport.cpp b/Engine/Source/Runtime/Windows/D3D12RHI/Private/WinRT/WinRTD3D12Viewport.cpp
new file mode 100644
--- /dev/null
+++ b/Engine/Source/Runtime/Windows/D3D12RHI/Private/WinRT/WinRTD3D12Viewport.cpp
@@ -0,0 +1,118 @@
+// Copyright 1998-2014 Epic Games, Inc. All Rights Reserved.
+
+#include "D3D12RHIPrivate.h"
+#include "RenderCore.h"
+
+#include "AllowWinRTPlatformTypes.h"
+	#include <DirectXMath.h>
+// #undef D3DCOLORVALUE_DEFINED
+// 
+// #pragma message("*** Including dxgiformat.h")
+// 	#include <dxgiformat.h>
+// #pragma message("*** Including dxgitype.h")
+// 	#include <dxgitype.h>
+// #pragma message("*** Including dxgi.h")
+// 	#include <dxgi.h>
+// #pragma message("*** Including dxgi1_2.h")
+	#include <dxgi1_2.h>
+// 	//	#include <xswaplib.h>
+#include "HideWinRTPlatformTypes.h"
+
+#pragma warning(disable : 4946)	// reinterpret_cast used between related classes: 'Platform::Object' and ...
+
+using namespace Windows::Foundation;
+using namespace Windows::ApplicationModel::Core;
+using namespace Windows::UI::Core;
+
+extern FD3D12Texture2D* GetSwapChainSurface(FD3D12DynamicRHI* D3DRHI,IDXGISwapChain* SwapChain);
+
+extern void D3D11TextureAllocated2D( FD3D12Texture2D& Texture );
+
+FD3D12Viewport::FD3D12Viewport(FD3D12DynamicRHI* InD3DRHI, HWND InWindowHandle, uint32 InSizeX, uint32 InSizeY, bool bInIsFullscreen, EPixelFormat InPreferredPixelFormat) :
+	D3DRHI(InD3DRHI),
+	LastFlipTime(0),
+	LastFrameComplete(0),
+	LastCompleteTime(0),
+	SyncCounter(0),
+	bSyncedLastFrame(false),
+	WindowHandle(InWindowHandle),
+	SizeX(InSizeX),
+	SizeY(InSizeY),
+	bIsFullscreen(bInIsFullscreen),
+	PixelFormat(InPreferredPixelFormat),
+	bIsValid(true)
+{
+	D3DRHI->Viewports.Add(this);
+
+	// Ensure that the D3D device has been created.
+	D3DRHI->InitD3DDevice();
+
+	// Create a backbuffer/swapchain for each viewport
+	// First, retrieve the underlying DXGI Device from the D3D Device
+	TRefCountPtr<IDXGIDevice1> DXGIDevice;
+	VERIFYD3D11RESULT(D3DRHI->GetDevice()->QueryInterface(__uuidof(DXGIDevice), (void**)DXGIDevice.GetInitReference() ));
+
+	IDXGIAdapter* pdxgiAdapter;
+	VERIFYD3D11RESULT(DXGIDevice->GetAdapter(&pdxgiAdapter));
+
+	IDXGIFactory2* pdxgiFactory;
+	VERIFYD3D11RESULT(pdxgiAdapter->GetParent(__uuidof(pdxgiFactory), reinterpret_cast< void** >(&pdxgiFactory)));
+
+	// Create the swapchain.
+	DXGI_SWAP_CHAIN_DESC1 swapChainDesc = {0};
+	swapChainDesc.Width = SizeX;
+	swapChainDesc.Height = SizeY;
+	swapChainDesc.Format = GetRenderTargetFormat(PixelFormat);
+	swapChainDesc.Stereo = false; 
+	swapChainDesc.SampleDesc.Count = 1;                          // don't use multi-sampling
+	swapChainDesc.SampleDesc.Quality = 0;
+	swapChainDesc.BufferUsage = DXGI_USAGE_RENDER_TARGET_OUTPUT | DXGI_USAGE_SHADER_INPUT;
+	swapChainDesc.BufferCount = 2;                               // use two buffers to enable flip effect
+	swapChainDesc.Scaling = DXGI_SCALING_STRETCH;
+	swapChainDesc.SwapEffect = DXGI_SWAP_EFFECT_FLIP_SEQUENTIAL; // we recommend using this swap effect for all applications
+	swapChainDesc.Flags = 0;
+
+	IDXGISwapChain1* CreatedSwapChain = NULL;
+
+	VERIFYD3D11RESULT(pdxgiFactory->CreateSwapChainForCoreWindow(
+		D3DRHI->GetDevice(), 
+		reinterpret_cast< IUnknown* >(CoreWindow::GetForCurrentThread()),
+		&swapChainDesc,
+		NULL,
+		(IDXGISwapChain1**)&CreatedSwapChain
+		)
+		);
+
+	*(SwapChain.GetInitReference()) = CreatedSwapChain;
+
+	// Set the DXGI message hook to not change the window behind our back.
+	//	D3DRHI->GetFactory()->MakeWindowAssociation(WindowHandle,DXGI_MWA_NO_WINDOW_CHANGES);
+
+	// Create a RHI surface to represent the viewport's back buffer.
+	BackBuffer = GetSwapChainSurface(D3DRHI,SwapChain);
+	BackBuffer->GetResource()->SetSubresourceUsage(D3D12_RESOURCE_BARRIER_ALL_SUBRESOURCES, D3D12_RESOURCE_STATE_PRESENT);
+}
+
+void FD3D12Viewport::ConditionalResetSwapChain(bool bIgnoreFocus)
+{
+	if(!bIsValid)
+	{
+		const bool bIsFocused = true;
+		const bool bIsIconic = false;
+		if(bIgnoreFocus || (bIsFocused && !bIsIconic) )
+		{
+			FlushRenderingCommands();
+
+			HRESULT Result = SwapChain->SetFullscreenState(bIsFullscreen,NULL);
+			if(SUCCEEDED(Result))
+			{
+				bIsValid = true;
+			}
+			else
+			{
+				// Even though the docs say SetFullscreenState always returns S_OK, that doesn't always seem to be the case.
+				UE_LOG(LogD3D12RHI, Log, TEXT("IDXGISwapChain::SetFullscreenState returned %08x; waiting for the next frame to try again."),Result);
+			}
+		}
+	}
+}
diff --git a/Engine/Source/Runtime/Windows/D3D12RHI/Private/Windows/D3D12RHIBasePrivate.h b/Engine/Source/Runtime/Windows/D3D12RHI/Private/Windows/D3D12RHIBasePrivate.h
new file mode 100644
--- /dev/null
+++ b/Engine/Source/Runtime/Windows/D3D12RHI/Private/Windows/D3D12RHIBasePrivate.h
@@ -0,0 +1,67 @@
+// Copyright 1998-2014 Epic Games, Inc. All Rights Reserved.
+
+/*=============================================================================
+	D3D12BaseRHIPrivate.h: Private D3D RHI definitions for Windows.
+=============================================================================*/
+
+#pragma once
+
+// Assume D3DX is available
+#ifndef WITH_D3DX_LIBS
+#define WITH_D3DX_LIBS	1
+#endif
+
+// Disable macro redefinition warning for compatibility with Windows SDK 8+
+#pragma warning(push)
+#if _MSC_VER >= 1700
+#pragma warning(disable : 4005)	// macro redefinition
+#endif
+
+// D3D headers.
+#if PLATFORM_64BITS
+#pragma pack(push,16)
+#else
+#pragma pack(push,8)
+#endif
+#define D3D_OVERLOADS 1
+#include "AllowWindowsPlatformTypes.h"
+#include <d3d12.h>
+#include "d3dx12.h"
+#include <d3d12sdklayers.h>
+#include "HideWindowsPlatformTypes.h"
+
+#undef DrawText
+
+#pragma pack(pop)
+#pragma warning(pop)
+
+typedef ID3D12GraphicsCommandList FD3D12CommandList;
+typedef ID3D12CommandQueue FD3D12CommandQueue;
+typedef ID3D12CommandAllocator FD3D12CommandAllocator;
+
+
+class FD3D12Device;
+
+class FD3D12DeviceChild
+{
+protected:
+    FD3D12Device* Parent;
+
+public:
+    FD3D12DeviceChild(FD3D12Device* InParent = nullptr) : Parent(InParent) {}
+
+    inline FD3D12Device* GetParentDevice() const
+	{ 
+        // If this fires an object was likely created with a default constructor i.e in an STL container
+        // and is therefore an orphan
+        check(Parent != nullptr);
+        return Parent;
+	}
+
+	// To be used with delayed setup
+    inline void SetParentDevice(FD3D12Device* InParent)
+	{
+		check(Parent == nullptr);
+		Parent = InParent;
+	}
+};
diff --git a/Engine/Source/Runtime/Windows/D3D12RHI/Private/Windows/D3D12RHIPrivateUtil.h b/Engine/Source/Runtime/Windows/D3D12RHI/Private/Windows/D3D12RHIPrivateUtil.h
new file mode 100644
--- /dev/null
+++ b/Engine/Source/Runtime/Windows/D3D12RHI/Private/Windows/D3D12RHIPrivateUtil.h
@@ -0,0 +1,24 @@
+// Copyright 1998-2014 Epic Games, Inc. All Rights Reserved.
+
+/*=============================================================================
+	D3D12RHIPrivateUtil.h: Private D3D RHI Utility definitions for Windows.
+=============================================================================*/
+
+#pragma once
+
+#include "WindowsD3D12ConstantBuffer.h"
+
+struct FD3DRHIUtil
+{
+	template <EShaderFrequency ShaderFrequencyT>
+	static FORCEINLINE void CommitConstants(FD3D12DynamicHeapAllocator& UploadHeap, FD3D12ConstantBuffer* InConstantBuffer, FD3D12StateCache& StateCache, uint32 Index, bool bDiscardSharedConstants)
+	{
+		auto* ConstantBuffer = ((FWinD3D12ConstantBuffer*)InConstantBuffer);
+		// Array may contain NULL entries to pad out to proper 
+        if (ConstantBuffer && ConstantBuffer->CommitConstantsToDevice(UploadHeap, bDiscardSharedConstants))
+		{
+			FD3D12ResourceLocation* DeviceBufferLocation = ConstantBuffer->GetConstantBufferLocation();
+			StateCache.SetConstantBuffer<ShaderFrequencyT>(Index, DeviceBufferLocation, nullptr);
+		}
+	}
+};
diff --git a/Engine/Source/Runtime/Windows/D3D12RHI/Private/Windows/D3D12StateCache.h b/Engine/Source/Runtime/Windows/D3D12RHI/Private/Windows/D3D12StateCache.h
new file mode 100644
--- /dev/null
+++ b/Engine/Source/Runtime/Windows/D3D12RHI/Private/Windows/D3D12StateCache.h
@@ -0,0 +1,10 @@
+// Copyright 1998-2014 Epic Games, Inc. All Rights Reserved.
+
+// Implementation of Device Context State Caching to improve draw
+//	thread performance by removing redundant device context calls.
+
+#pragma once
+
+#include "D3D12StateCachePrivate.h"
+
+typedef FD3D12StateCacheBase FD3D12StateCache;
diff --git a/Engine/Source/Runtime/Windows/D3D12RHI/Private/Windows/WindowsD3D12ConstantBuffer.cpp b/Engine/Source/Runtime/Windows/D3D12RHI/Private/Windows/WindowsD3D12ConstantBuffer.cpp
new file mode 100644
--- /dev/null
+++ b/Engine/Source/Runtime/Windows/D3D12RHI/Private/Windows/WindowsD3D12ConstantBuffer.cpp
@@ -0,0 +1,90 @@
+// Copyright 1998-2014 Epic Games, Inc. All Rights Reserved.
+
+/*=============================================================================
+	WindowsD3D12ConstantBuffer.cpp: D3D Constant Buffer functions
+=============================================================================*/
+
+#include "D3D12RHIPrivate.h"
+#include "WindowsD3D12ConstantBuffer.h"
+
+void FWinD3D12ConstantBuffer::InitDynamicRHI()
+{
+    CurrentBuffer = new FD3D12ResourceLocation;
+	FD3D12ConstantBuffer::InitDynamicRHI();
+}
+
+void FWinD3D12ConstantBuffer::ReleaseDynamicRHI()
+{
+	FD3D12ConstantBuffer::ReleaseDynamicRHI();
+    CurrentBuffer = nullptr;
+}
+
+bool FWinD3D12ConstantBuffer::CommitConstantsToDevice(FD3D12DynamicHeapAllocator& UploadHeap, bool bDiscardSharedConstants)
+{
+	// New circular buffer system for faster constant uploads.  Avoids CopyResource and speeds things up considerably
+	if (CurrentUpdateSize == 0)
+	{
+		return false;
+	}
+
+	SCOPE_CYCLE_COUNTER(STAT_D3D12GlobalConstantBufferUpdateTime);
+
+	if ( bDiscardSharedConstants )
+	{
+		// If we're discarding shared constants, just use constants that have been updated since the last Commit.
+		TotalUpdateSize = CurrentUpdateSize;
+	}
+	else
+	{
+		// If we're re-using shared constants, use all constants.
+		TotalUpdateSize = FMath::Max( CurrentUpdateSize, TotalUpdateSize );
+	}
+
+	// This only used for the IndexSlot==0 constant buffer on the vertex shader.
+	// Which will have an indeterminate number of constant values that are generated by material shaders.
+
+    // Get the next constant buffer
+    void *pData = UploadHeap.FastAlloc(TotalUpdateSize, D3D12_CONSTANT_BUFFER_DATA_PLACEMENT_ALIGNMENT, CurrentBuffer);
+
+	check(((uint64)(ShadowData)& 0xf) == 0);
+	FMemory::Memcpy(pData, ShadowData, TotalUpdateSize);
+
+	CurrentUpdateSize = 0;
+
+	return true;
+}
+
+void FD3D12CommandContext::InitConstantBuffers()
+{
+	// Allocate shader constant buffers.  All shader types can have access to all buffers.
+	//  Index==0 is reserved for "custom" params, and the rest are reserved by the system for Common
+	//	constants
+	VSConstantBuffers.Empty(MAX_CONSTANT_BUFFER_SLOTS);
+	PSConstantBuffers.Empty(MAX_CONSTANT_BUFFER_SLOTS);
+	HSConstantBuffers.Empty(MAX_CONSTANT_BUFFER_SLOTS);
+	DSConstantBuffers.Empty(MAX_CONSTANT_BUFFER_SLOTS);
+	GSConstantBuffers.Empty(MAX_CONSTANT_BUFFER_SLOTS);
+	CSConstantBuffers.Empty(MAX_CONSTANT_BUFFER_SLOTS);
+	for(int32 BufferIndex = 0;BufferIndex < MAX_CONSTANT_BUFFER_SLOTS;BufferIndex++)
+	{
+		uint32 Size = GConstantBufferSizes[BufferIndex];
+		uint32 SubBuffers = 1;
+		if(BufferIndex == GLOBAL_CONSTANT_BUFFER_INDEX)
+		{
+			SubBuffers = 5;
+		}
+
+		// Vertex shader can have subbuffers for index==0.  This is from Epic's original design for the auto-fit of size to
+		//	reduce the update costs of the buffer.
+
+		// New circular buffer system for faster constant uploads.  Avoids CopyResource and speeds things up considerably
+		VSConstantBuffers.Add(new FWinD3D12ConstantBuffer(GetParentDevice(), Size, SubBuffers));
+		PSConstantBuffers.Add(new FWinD3D12ConstantBuffer(GetParentDevice(), Size, SubBuffers));
+		HSConstantBuffers.Add(new FWinD3D12ConstantBuffer(GetParentDevice(), Size));
+		DSConstantBuffers.Add(new FWinD3D12ConstantBuffer(GetParentDevice(), Size));
+		GSConstantBuffers.Add(new FWinD3D12ConstantBuffer(GetParentDevice(), Size));
+		CSConstantBuffers.Add(new FWinD3D12ConstantBuffer(GetParentDevice(), Size));
+	}
+}
+
+DEFINE_STAT(STAT_D3D12GlobalConstantBufferUpdateTime);
diff --git a/Engine/Source/Runtime/Windows/D3D12RHI/Private/Windows/WindowsD3D12ConstantBuffer.h b/Engine/Source/Runtime/Windows/D3D12RHI/Private/Windows/WindowsD3D12ConstantBuffer.h
new file mode 100644
--- /dev/null
+++ b/Engine/Source/Runtime/Windows/D3D12RHI/Private/Windows/WindowsD3D12ConstantBuffer.h
@@ -0,0 +1,36 @@
+// Copyright 1998-2014 Epic Games, Inc. All Rights Reserved.
+
+/*=============================================================================
+	WindowsD3D12ConstantBuffer.h: D3D Constant Buffer functions
+=============================================================================*/
+
+#pragma once
+
+class FWinD3D12ConstantBuffer : public FD3D12ConstantBuffer
+{
+public:
+	FWinD3D12ConstantBuffer(FD3D12Device* InParent, uint32 InSize = 0, uint32 SubBuffers = 1) :
+        FD3D12ConstantBuffer(InParent, InSize, SubBuffers)
+	{
+	}
+
+	// FRenderResource interface.
+	virtual void	InitDynamicRHI() override;
+	virtual void	ReleaseDynamicRHI() override;
+
+	/**
+	* Get the current buffer
+	*/
+	FD3D12ResourceLocation* GetConstantBufferLocation() const
+	{
+		return CurrentBuffer.GetReference();
+	}
+
+	/**
+	* Unlocks the constant buffer so the data can be transmitted to the device
+	*/
+    bool CommitConstantsToDevice(FD3D12DynamicHeapAllocator& UploadHeap, bool bDiscardSharedConstants);
+
+private:
+    TRefCountPtr<FD3D12ResourceLocation> CurrentBuffer;
+};
