UE4 - D3D12RHI (not working)

diff --git a/Engine/Source/Runtime/Windows/D3D12RHI/Private/Windows/WindowsD3D12Device.cpp b/Engine/Source/Runtime/Windows/D3D12RHI/Private/Windows/WindowsD3D12Device.cpp
new file mode 100644
--- /dev/null
+++ b/Engine/Source/Runtime/Windows/D3D12RHI/Private/Windows/WindowsD3D12Device.cpp
@@ -0,0 +1,901 @@
+// Copyright 1998-2014 Epic Games, Inc. All Rights Reserved.
+
+/*=============================================================================
+	D3D12Device.cpp: D3D device RHI implementation.
+=============================================================================*/
+
+#include "D3D12RHIPrivate.h"
+#include "AllowWindowsPlatformTypes.h"
+#include <delayimp.h>
+#include "HideWindowsPlatformTypes.h"
+#include "HardwareInfo.h"
+
+#pragma comment(lib, "d3d12.lib")
+
+extern bool D3D12RHI_ShouldCreateWithD3DDebug();
+extern bool D3D12RHI_ShouldAllowAsyncResourceCreation();
+
+static TAutoConsoleVariable<int32> CVarGraphicsAdapter(
+	TEXT("r.D3D12GraphicsAdapter"),
+	-1,
+	TEXT("User request to pick a specific graphics adapter (e.g. when using a integrated graphics card with a descrete one)\n")
+	TEXT("At the moment this only works on Direct3D 11.\n")
+	TEXT(" -2: Take the first one that fulfills the criteria\n")
+	TEXT(" -1: Favour non integrated because there are usually faster\n")
+	TEXT("  0: Adpater #0\n")
+	TEXT("  1: Adpater #1, ..."),
+	ECVF_RenderThreadSafe);
+
+namespace D3D12RHI
+{
+
+/**
+ * Console variables used by the D3D11 RHI device.
+ */
+namespace RHIConsoleVariables
+{
+	int32 FeatureSetLimit = -1;
+	static FAutoConsoleVariableRef CVarFeatureSetLimit(
+		TEXT("D3D12RHI.FeatureSetLimit"),
+		FeatureSetLimit,
+		TEXT("If set to 10, limit D3D RHI to D3D10 feature level. Otherwise, it will use default. Changing this at run-time has no effect. (default is -1)")
+		);
+};
+}
+using namespace D3D12RHI;
+
+/** This function is used as a SEH filter to catch only delay load exceptions. */
+static bool IsDelayLoadException(PEXCEPTION_POINTERS ExceptionPointers)
+{
+#if WINVER > 0x502	// Windows SDK 7.1 doesn't define VcppException
+	switch (ExceptionPointers->ExceptionRecord->ExceptionCode)
+	{
+	case VcppException(ERROR_SEVERITY_ERROR, ERROR_MOD_NOT_FOUND):
+	case VcppException(ERROR_SEVERITY_ERROR, ERROR_PROC_NOT_FOUND):
+		return EXCEPTION_EXECUTE_HANDLER;
+	default:
+		return EXCEPTION_CONTINUE_SEARCH;
+	}
+#else
+	return EXCEPTION_EXECUTE_HANDLER;
+#endif
+}
+
+// We suppress warning C6322: Empty _except block. Appropriate checks are made upon returning. 
+#if USING_CODE_ANALYSIS
+	MSVC_PRAGMA(warning(push))
+	MSVC_PRAGMA(warning(disable:6322))
+#endif	// USING_CODE_ANALYSIS
+/**
+ * Since CreateDXGIFactory is a delay loaded import from the D3D11 DLL, if the user
+ * doesn't have Vista/DX10, calling CreateDXGIFactory will throw an exception.
+ * We use SEH to detect that case and fail gracefully.
+ */
+static void SafeCreateDXGIFactory(IDXGIFactory4** DXGIFactory)
+{
+#if !D3D11_CUSTOM_VIEWPORT_CONSTRUCTOR
+	__try
+	{
+		CreateDXGIFactory(__uuidof(IDXGIFactory4), (void**)DXGIFactory);
+	}
+	__except (IsDelayLoadException(GetExceptionInformation()))
+	{
+	}
+#endif	//!D3D11_CUSTOM_VIEWPORT_CONSTRUCTOR
+}
+
+/**
+ * Returns the highest D3D feature level we are allowed to created based on
+ * command line parameters.
+ */
+static D3D_FEATURE_LEVEL GetAllowedD3DFeatureLevel()
+{
+	// Default to D3D11 
+	D3D_FEATURE_LEVEL AllowedFeatureLevel = D3D_FEATURE_LEVEL_11_0;
+
+	// Use a feature level 10 if specified on the command line.
+	if (FParse::Param(FCommandLine::Get(), TEXT("d3d10")) ||
+		FParse::Param(FCommandLine::Get(), TEXT("dx10")) ||
+		FParse::Param(FCommandLine::Get(), TEXT("sm4")) ||
+		RHIConsoleVariables::FeatureSetLimit == 10)
+	{
+		AllowedFeatureLevel = D3D_FEATURE_LEVEL_10_0;
+	}
+	return AllowedFeatureLevel;
+}
+
+/**
+ * Attempts to create a D3D12 device for the adapter using at most MaxFeatureLevel.
+ * If creation is successful, true is returned and the supported feature level is set in OutFeatureLevel.
+ */
+static bool SafeTestD3D12CreateDevice(IDXGIAdapter* Adapter, D3D_FEATURE_LEVEL MaxFeatureLevel, D3D_FEATURE_LEVEL* OutFeatureLevel)
+{
+	ID3D12Device* D3DDevice = NULL;
+
+	// Use a debug device if specified on the command line.
+	if(D3D12RHI_ShouldCreateWithD3DDebug())
+	{
+		ID3D12Debug* DebugController = nullptr;
+		VERIFYD3D11RESULT(D3D12GetDebugInterface(IID_PPV_ARGS(&DebugController)));
+		DebugController->EnableDebugLayer();
+		DebugController->Release();
+	}
+
+	D3D_FEATURE_LEVEL RequestedFeatureLevels[] =
+	{
+		D3D_FEATURE_LEVEL_11_0,
+		D3D_FEATURE_LEVEL_10_0
+	};
+
+	int32 FirstAllowedFeatureLevel = 0;
+	int32 NumAllowedFeatureLevels = ARRAY_COUNT(RequestedFeatureLevels);
+	while (FirstAllowedFeatureLevel < NumAllowedFeatureLevels)
+	{
+		if (RequestedFeatureLevels[FirstAllowedFeatureLevel] == MaxFeatureLevel)
+		{
+			break;
+		}
+		FirstAllowedFeatureLevel++;
+	}
+	NumAllowedFeatureLevels -= FirstAllowedFeatureLevel;
+
+	if (NumAllowedFeatureLevels == 0)
+	{
+		return false;
+	}
+
+	__try
+	{
+		// We don't want software renderer. Ideally we specify D3D_DRIVER_TYPE_HARDWARE on creation but
+		// when we specify an adapter we need to specify D3D_DRIVER_TYPE_UNKNOWN (otherwise the call fails).
+		// We cannot check the device type later (seems this is missing functionality in D3D).
+		if (SUCCEEDED(D3D12CreateDevice(
+			Adapter,
+			RequestedFeatureLevels[FirstAllowedFeatureLevel],
+			IID_PPV_ARGS(&D3DDevice)
+			)))
+		{
+			*OutFeatureLevel = RequestedFeatureLevels[FirstAllowedFeatureLevel];
+			D3DDevice->Release();
+			return true;
+		}
+	}
+	__except (IsDelayLoadException(GetExceptionInformation()))
+	{
+	}
+
+	return false;
+}
+
+// Re-enable C6322
+#if USING_CODE_ANALYSIS
+	MSVC_PRAGMA(warning(pop)) 
+#endif // USING_CODE_ANALYSIS
+
+bool FD3D12DynamicRHIModule::IsSupported()
+{
+	// if not computed yet
+    if (ChosenAdapter.IsValid() == 0)
+	{
+		FindAdapter();
+	}
+	if (ChosenAdapter.IsValid() == 0)
+    {
+        return false;
+    }
+    else
+    {
+        // The hardware must support at least 10.0 (usually 11_0, 10_0 or 10_1).
+        return ChosenAdapter.IsValid()
+            && ChosenAdapter.MaxSupportedFeatureLevel != D3D_FEATURE_LEVEL_9_1
+            && ChosenAdapter.MaxSupportedFeatureLevel != D3D_FEATURE_LEVEL_9_2
+            && ChosenAdapter.MaxSupportedFeatureLevel != D3D_FEATURE_LEVEL_9_3;
+    }
+}
+
+namespace D3D12RHI
+{
+const TCHAR* GetFeatureLevelString(D3D_FEATURE_LEVEL FeatureLevel)
+{
+	switch (FeatureLevel)
+	{
+		case D3D_FEATURE_LEVEL_9_1:		return TEXT("9_1");
+		case D3D_FEATURE_LEVEL_9_2:		return TEXT("9_2");
+		case D3D_FEATURE_LEVEL_9_3:		return TEXT("9_3");
+		case D3D_FEATURE_LEVEL_10_0:	return TEXT("10_0");
+		case D3D_FEATURE_LEVEL_10_1:	return TEXT("10_1");
+		case D3D_FEATURE_LEVEL_11_0:	return TEXT("11_0");
+	}
+	return TEXT("X_X");
+}
+}
+
+static uint32 CountAdapterOutputs(TRefCountPtr<IDXGIAdapter>& Adapter)
+{
+	uint32 OutputCount = 0;
+	for (;;)
+	{
+		TRefCountPtr<IDXGIOutput> Output;
+		HRESULT hr = Adapter->EnumOutputs(OutputCount, Output.GetInitReference());
+		if (FAILED(hr))
+		{
+			break;
+		}
+		++OutputCount;
+	}
+
+	return OutputCount;
+}
+
+void FD3D12DynamicRHIModule::FindAdapter()
+{
+    // Once we chosen one we don't need to do it again.
+    check(ChosenAdapter.IsValid() == 0);
+
+	// Try to create the DXGIFactory.  This will fail if we're not running Vista.
+	TRefCountPtr<IDXGIFactory4> DXGIFactory;
+	SafeCreateDXGIFactory(DXGIFactory.GetInitReference());
+	if (!DXGIFactory)
+	{
+		return;
+	}
+
+	bool bAllowPerfHUD = true;
+
+#if UE_BUILD_SHIPPING || UE_BUILD_TEST
+	bAllowPerfHUD = false;
+#endif
+
+	int32 CVarValue = CVarGraphicsAdapter.GetValueOnGameThread();
+
+	const bool bFavorNonIntegrated = CVarValue == -1;
+
+	TRefCountPtr<IDXGIAdapter> TempAdapter;
+	D3D_FEATURE_LEVEL MaxAllowedFeatureLevel = GetAllowedD3DFeatureLevel();
+
+	FD3D12Adapter FirstWithoutIntegratedAdapter;
+	FD3D12Adapter FirstAdapter;
+
+	bool bIsAnyAMD = false;
+	bool bIsAnyIntel = false;
+	bool bIsAnyNVIDIA = false;
+
+	// Enumerate the DXGIFactory's adapters.
+	for (uint32 AdapterIndex = 0; DXGIFactory->EnumAdapters(AdapterIndex, TempAdapter.GetInitReference()) != DXGI_ERROR_NOT_FOUND; ++AdapterIndex)
+	{
+		// Check that if adapter supports D3D11.
+		if (TempAdapter)
+		{
+			D3D_FEATURE_LEVEL ActualFeatureLevel = (D3D_FEATURE_LEVEL)0;
+			if (SafeTestD3D12CreateDevice(TempAdapter, MaxAllowedFeatureLevel, &ActualFeatureLevel))
+			{
+				// Log some information about the available D3D12 adapters.
+				DXGI_ADAPTER_DESC AdapterDesc;
+				VERIFYD3D11RESULT(TempAdapter->GetDesc(&AdapterDesc));
+				uint32 OutputCount = CountAdapterOutputs(TempAdapter);
+
+				UE_LOG(LogD3D12RHI, Log,
+					TEXT("Found D3D12 adapter %u: %s (Feature Level %s)"),
+					AdapterIndex,
+					AdapterDesc.Description,
+					GetFeatureLevelString(ActualFeatureLevel)
+					);
+				UE_LOG(LogD3D12RHI, Log,
+					TEXT("Adapter has %uMB of dedicated video memory, %uMB of dedicated system memory, and %uMB of shared system memory, %d output[s]"),
+					(uint32)(AdapterDesc.DedicatedVideoMemory / (1024*1024)),
+					(uint32)(AdapterDesc.DedicatedSystemMemory / (1024*1024)),
+					(uint32)(AdapterDesc.SharedSystemMemory / (1024*1024)),
+					OutputCount
+					);
+
+				bool bIsAMD = AdapterDesc.VendorId == 0x1002;
+				bool bIsIntel = AdapterDesc.VendorId == 0x8086;
+				bool bIsNVIDIA = AdapterDesc.VendorId == 0x10DE;
+                bool bIsWARP = FParse::Param(FCommandLine::Get(), TEXT("warp"));
+
+				if (bIsAMD) bIsAnyAMD = true;
+                if (bIsIntel) bIsAnyIntel = true;
+				if (bIsNVIDIA) bIsAnyNVIDIA = true;
+
+				// Simple heuristic but without profiling it's hard to do better
+				const bool bIsIntegrated = bIsIntel;
+				// PerfHUD is for performance profiling
+				const bool bIsPerfHUD = !FCString::Stricmp(AdapterDesc.Description, TEXT("NVIDIA PerfHUD"));
+
+				FD3D12Adapter CurrentAdapter(AdapterIndex, ActualFeatureLevel);
+
+				if (!OutputCount && !bIsWARP)
+				{
+                    // Add special check to support WARP, which does not have an output associated with it.
+
+					// This device has no outputs. Reject it, 
+					// http://msdn.microsoft.com/en-us/library/windows/desktop/bb205075%28v=vs.85%29.aspx#WARP_new_for_Win8
+					continue;
+				}
+
+				if(bIsPerfHUD && !bAllowPerfHUD)
+				{
+					// we don't allow the PerfHUD adapter
+					continue;
+				}
+
+				if (CVarValue >= 0 && AdapterIndex != CVarValue)
+				{
+					// the user wants a specific adapter, not this one
+					continue;
+				}
+
+				if (!bIsIntegrated && !FirstWithoutIntegratedAdapter.IsValid())
+				{
+					FirstWithoutIntegratedAdapter = CurrentAdapter;
+				}
+
+				if (!FirstAdapter.IsValid())
+				{
+					FirstAdapter = CurrentAdapter;
+				}
+			}
+		}
+	}
+
+	if (bFavorNonIntegrated && (bIsAnyAMD || bIsAnyNVIDIA))
+	{
+		ChosenAdapter = FirstWithoutIntegratedAdapter;
+
+		// We assume Intel is integrated graphics (slower than discrete) than NVIDIA or AMD cards and rather take a different one
+		if (!ChosenAdapter.IsValid())
+		{
+			ChosenAdapter = FirstAdapter;
+		}
+	}
+	else
+	{
+		ChosenAdapter = FirstAdapter;
+	}
+
+	if (ChosenAdapter.IsValid())
+	{
+		UE_LOG(LogD3D12RHI, Log, TEXT("Chosen D3D12 Adapter Id = %u"), ChosenAdapter.AdapterIndex);
+	}
+	else
+	{
+		UE_LOG(LogD3D12RHI, Error, TEXT("Failed to choose a D3D12 Adapter."));
+	}
+}
+
+FDynamicRHI* FD3D12DynamicRHIModule::CreateRHI()
+{
+	TRefCountPtr<IDXGIFactory4> DXGIFactory;
+	SafeCreateDXGIFactory(DXGIFactory.GetInitReference());
+	check(DXGIFactory);
+    return new FD3D12DynamicRHI(DXGIFactory, ChosenAdapter);
+}
+
+void FD3D12DynamicRHI::Init()
+{
+#if UE_BUILD_DEBUG	
+	SubmissionLockStalls = 0;
+    DrawCount = 0;
+    PresentCount = 0;
+#endif
+	InitD3DDevices();
+
+#if !WITH_EDITOR
+	if (GRHISupportsRHIThread)
+	{
+		SetupRecursiveResources();
+	}
+#endif
+}
+
+void FD3D12DynamicRHI::PerRHISetup(FD3D12Device* MainDevice)
+{
+    bool bIsWARP = FParse::Param(FCommandLine::Get(), TEXT("warp"));
+
+    check(!GIsRHIInitialized);
+
+    DXGI_ADAPTER_DESC* AdapterDesc = MainDevice->GetD3DAdapterDesc();
+   
+    GTexturePoolSize = 0;
+
+    // TODO MS: We need to do this after CreateDevice due to a bug in QueryVideoMemoryInfo
+    GRHIAdapterName = AdapterDesc->Description;
+    GRHIVendorId = AdapterDesc->VendorId;
+
+    extern int64 GDedicatedVideoMemory;
+    extern int64 GDedicatedSystemMemory;
+    extern int64 GSharedSystemMemory;
+    extern int64 GTotalGraphicsMemory;
+
+    // Issue: 32bit windows doesn't report 64bit value, we take what we get.
+    GDedicatedVideoMemory = int64(AdapterDesc->DedicatedVideoMemory);
+    GDedicatedSystemMemory = int64(AdapterDesc->DedicatedSystemMemory);
+    GSharedSystemMemory = int64(AdapterDesc->SharedSystemMemory);
+
+    // Total amount of system memory, clamped to 8 GB
+    int64 TotalPhysicalMemory = FMath::Min(int64(FPlatformMemory::GetConstants().TotalPhysicalGB), 8ll) * (1024ll * 1024ll * 1024ll);
+
+    // Consider 50% of the shared memory but max 25% of total system memory.
+    int64 ConsideredSharedSystemMemory = FMath::Min(GSharedSystemMemory / 2ll, TotalPhysicalMemory / 4ll);
+
+    IDXGIAdapter3* DxgiAdapter3 = MainDevice->GetAdapter3();
+    TRefCountPtr<IDXGIAdapter> EnumAdapter;
+
+	// Hack to handle a kernel bug where QueryVideomMemoryInfo can't be called on WARP
+	DXGI_QUERY_VIDEO_MEMORY_INFO LocalVideoMemoryInfo = {};
+	if (!bIsWARP)
+	{
+		GTotalGraphicsMemory = 0;
+		if (DXGIFactory->EnumAdapters(MainDevice->GetAdapterIndex(), EnumAdapter.GetInitReference()) != DXGI_ERROR_NOT_FOUND)
+		{
+			if (EnumAdapter)
+			{
+				VERIFYD3D11RESULT(DxgiAdapter3->QueryVideoMemoryInfo(0, DXGI_MEMORY_SEGMENT_GROUP_LOCAL, &LocalVideoMemoryInfo));
+				GTotalGraphicsMemory = LocalVideoMemoryInfo.Budget;
+			}
+		}
+
+	}
+	else
+	{
+		GTotalGraphicsMemory = GDedicatedSystemMemory;
+	}
+
+	GMaxTextureMipCount = FMath::CeilLogTwo(GMaxTextureDimensions) + 1;
+	GMaxTextureMipCount = FMath::Min<int32>(MAX_TEXTURE_MIP_COUNT, GMaxTextureMipCount);
+	// If we only have <= 2GB of GPU memory to work with, we should cut down on the mip count
+	// "GMaxTextureMipCount = 9" allows Intel to run the 4.8 Elemental demo, ideally this handling 
+	// is unnecessary as the upper engine should be able to look at GTotalGraphicsMemory to determine 
+	// the appropriate mip level but that doesn't appear to happen correctly on start-up.
+	if (LocalVideoMemoryInfo.Budget <= 2ll * 1024ll * 1024ll * 1024ll)
+	{
+		GMaxTextureMipCount = FMath::Min<int32>(GMaxTextureMipCount, 9);
+	}
+
+	if (sizeof(SIZE_T) < 8)
+	{
+		// Clamp to 1 GB if we're less than 64-bit
+		GTotalGraphicsMemory = FMath::Min(GTotalGraphicsMemory, 1024ll * 1024ll * 1024ll);
+	}
+	else
+	{
+		// Clamp to 1.9 GB if we're 64-bit
+		GTotalGraphicsMemory = FMath::Min(GTotalGraphicsMemory, 1945ll * 1024ll * 1024ll);
+	}
+
+	if (!bIsWARP)
+	{
+		VERIFYD3D11RESULT(DxgiAdapter3->SetVideoMemoryReservation(0, DXGI_MEMORY_SEGMENT_GROUP_LOCAL, FMath::Min((int64)LocalVideoMemoryInfo.AvailableForReservation, GTotalGraphicsMemory)));
+	}
+
+
+
+    if (GPoolSizeVRAMPercentage > 0)
+    {
+        float PoolSize = float(GPoolSizeVRAMPercentage) * 0.01f * float(GTotalGraphicsMemory);
+
+        // Truncate GTexturePoolSize to MB (but still counted in bytes)
+        GTexturePoolSize = int64(FGenericPlatformMath::TruncToFloat(PoolSize / 1024.0f / 1024.0f)) * 1024 * 1024;
+
+        UE_LOG(LogRHI, Log, TEXT("Texture pool is %llu MB (%d%% of %llu MB)"),
+            GTexturePoolSize / 1024 / 1024,
+            GPoolSizeVRAMPercentage,
+            GTotalGraphicsMemory / 1024 / 1024);
+    }
+
+#if (UE_BUILD_SHIPPING && WITH_EDITOR) && PLATFORM_WINDOWS && !PLATFORM_64BITS
+    // Disable PIX for windows in the shipping editor builds
+    D3DPERF_SetOptions(1);
+#endif
+
+    // Multi-threaded resource creation is always supported in DX12.
+    GRHISupportsAsyncTextureCreation = true;
+
+    GShaderPlatformForFeatureLevel[ERHIFeatureLevel::ES2] = SP_PCD3D_ES2;
+	GShaderPlatformForFeatureLevel[ERHIFeatureLevel::ES3_1] = SP_PCD3D_ES3_1;
+    GShaderPlatformForFeatureLevel[ERHIFeatureLevel::SM4] = SP_PCD3D_SM4;
+    GShaderPlatformForFeatureLevel[ERHIFeatureLevel::SM5] = SP_PCD3D_SM5;
+
+    if (IsRHIDeviceIntel())
+    {
+        // MS: To work around current known issue on Intel, TextureQuality be '0'.
+        // SetQualityLevels responsible for setting this value, must be re-run now after device creation.
+        Scalability::FQualityLevels currentQualityLevels = Scalability::GetQualityLevels();
+        Scalability::SetQualityLevels(currentQualityLevels);
+    }
+
+    // Notify all initialized FRenderResources that there's a valid RHI device to create their RHI resources for now.
+    for (TLinkedList<FRenderResource*>::TIterator ResourceIt(FRenderResource::GetResourceList()); ResourceIt; ResourceIt.Next())
+    {
+        ResourceIt->InitDynamicRHI();
+    }
+    for (TLinkedList<FRenderResource*>::TIterator ResourceIt(FRenderResource::GetResourceList()); ResourceIt; ResourceIt.Next())
+    {
+        ResourceIt->InitRHI();
+    }
+
+    FHardwareInfo::RegisterHardwareInfo(NAME_RHI, TEXT("D3D12"));
+
+    GRHISupportsTextureStreaming = true;
+	GRHISupportsFirstInstance = true;
+
+    // Set the RHI initialized flag.
+    GIsRHIInitialized = true;
+}
+
+
+void FD3D12Device::InitD3DDevice()
+{
+	check(IsInGameThread());
+
+	// Wait for the rendering thread to go idle.
+	SCOPED_SUSPEND_RENDERING_THREAD(false);
+
+	// If the device we were using has been removed, release it and the resources we created for it.
+	if (bDeviceRemoved)
+	{
+		check(Direct3DDevice);
+
+		HRESULT hRes = Direct3DDevice->GetDeviceRemovedReason();
+
+		const TCHAR* Reason = TEXT("?");
+		switch (hRes)
+		{
+			case DXGI_ERROR_DEVICE_HUNG:			Reason = TEXT("HUNG"); break;
+			case DXGI_ERROR_DEVICE_REMOVED:			Reason = TEXT("REMOVED"); break;
+			case DXGI_ERROR_DEVICE_RESET:			Reason = TEXT("RESET"); break;
+			case DXGI_ERROR_DRIVER_INTERNAL_ERROR:	Reason = TEXT("INTERNAL_ERROR"); break;
+			case DXGI_ERROR_INVALID_CALL:			Reason = TEXT("INVALID_CALL"); break;
+		}
+
+		bDeviceRemoved = false;
+
+		// Cleanup the D3D device.
+		CleanupD3DDevice();
+
+		// We currently don't support removed devices because FTexture2DResource can't recreate its RHI resources from scratch.
+		// We would also need to recreate the viewport swap chains from scratch.
+		UE_LOG(LogD3D12RHI, Fatal, TEXT("The Direct3D 12 device that was being used has been removed (Error: %d '%s').  Please restart the game."), hRes, Reason);
+	}
+
+	// If we don't have a device yet, either because this is the first viewport, or the old device was removed, create a device.
+	if (!Direct3DDevice)
+	{
+		// Determine the adapter and device type to use.
+		TRefCountPtr<IDXGIAdapter> Adapter;
+		
+		// In Direct3D 11, if you are trying to create a hardware or a software device, set pAdapter != NULL which constrains the other inputs to be:
+		//		DriverType must be D3D_DRIVER_TYPE_UNKNOWN 
+		//		Software must be NULL. 
+		D3D_DRIVER_TYPE DriverType = D3D_DRIVER_TYPE_UNKNOWN;	
+
+		// Use a debug device if specified on the command line.
+		const bool bWithD3DDebug = D3D12RHI_ShouldCreateWithD3DDebug();
+
+		if (bWithD3DDebug)
+		{
+			TRefCountPtr<ID3D12Debug> DebugController;
+			VERIFYD3D11RESULT(D3D12GetDebugInterface(IID_PPV_ARGS(DebugController.GetInitReference())));
+			DebugController->EnableDebugLayer();
+
+			UE_LOG(LogD3D12RHI, Log, TEXT("InitD3DDevice: -D3DDebug = %s"), bWithD3DDebug ? TEXT("on") : TEXT("off"));
+		}
+        
+        TRefCountPtr<IDXGIAdapter> EnumAdapter;
+
+        if (DXGIFactory->EnumAdapters(GetAdapterIndex(), EnumAdapter.GetInitReference()) != DXGI_ERROR_NOT_FOUND)
+        {
+            if (EnumAdapter)
+            {
+                if (SUCCEEDED(EnumAdapter->GetDesc(&AdapterDesc)))
+                {
+                    Adapter = EnumAdapter;
+
+                    const bool bIsPerfHUD = !FCString::Stricmp(AdapterDesc.Description, TEXT("NVIDIA PerfHUD"));
+
+                    if (bIsPerfHUD)
+                    {
+                        DriverType = D3D_DRIVER_TYPE_REFERENCE;
+                    }
+
+                    VERIFYD3D11RESULT(EnumAdapter->QueryInterface(_uuidof(DxgiAdapter3), (void **)DxgiAdapter3.GetInitReference()));
+                }
+                else
+                {
+                    check(!"Internal error, GetDesc() failed but before it worked")
+                }
+            }
+        }
+        else
+        {
+            check(!"Internal error, EnumAdapters() failed but before it worked")
+        }
+
+        D3D_FEATURE_LEVEL ActualFeatureLevel = GetFeatureLevel();
+		
+		if (FParse::Param(FCommandLine::Get(), TEXT("warp")))
+		{
+			TRefCountPtr<IDXGIAdapter> WarpAdapter;
+			VERIFYD3D11RESULT(DXGIFactory->EnumWarpAdapter(IID_PPV_ARGS(WarpAdapter.GetInitReference())));
+
+			// Creating the Direct3D WARP device.
+			VERIFYD3D11RESULT(D3D12CreateDevice(
+				WarpAdapter,
+				GetFeatureLevel(),
+				IID_PPV_ARGS(Direct3DDevice.GetInitReference())
+				));
+
+		}
+		else
+		{
+			// Creating the Direct3D device.
+			VERIFYD3D11RESULT(D3D12CreateDevice(
+				Adapter,
+				GetFeatureLevel(),
+				IID_PPV_ARGS(Direct3DDevice.GetInitReference())
+				));
+		}
+
+        TRefCountPtr<ID3D12RootSignature> NullRS;
+        {
+            CD3DX12_ROOT_SIGNATURE_DESC NullRSDesc;
+            NullRSDesc.Init(0, nullptr, 0);
+            TRefCountPtr<ID3DBlob> Blob;
+            VERIFYD3D11RESULT(D3D12SerializeRootSignature(&NullRSDesc, D3D_ROOT_SIGNATURE_VERSION_1, Blob.GetInitReference(), nullptr));
+            VERIFYD3D11RESULT(Direct3DDevice->CreateRootSignature(0, Blob->GetBufferPointer(), Blob->GetBufferSize(), IID_PPV_ARGS(NullRS.GetInitReference())));
+        }
+
+#if UE_BUILD_DEBUG	
+		//break on debug
+		TRefCountPtr<ID3D12Debug> d3dDebug;
+		if (SUCCEEDED(Direct3DDevice->QueryInterface(__uuidof(ID3D12Debug), (void**)d3dDebug.GetInitReference())))
+		{
+			TRefCountPtr<ID3D12InfoQueue> d3dInfoQueue;
+			if (SUCCEEDED(d3dDebug->QueryInterface(__uuidof(ID3D12InfoQueue), (void**)d3dInfoQueue.GetInitReference())))
+			{
+				d3dInfoQueue->SetBreakOnSeverity(D3D12_MESSAGE_SEVERITY_CORRUPTION, true);
+				d3dInfoQueue->SetBreakOnSeverity(D3D12_MESSAGE_SEVERITY_ERROR, true);
+				//d3dInfoQueue->SetBreakOnSeverity(D3D12_MESSAGE_SEVERITY_WARNING, true);
+			}
+		}
+#endif
+
+		// We should get the feature level we asked for as earlier we checked to ensure it is supported.
+		check(ActualFeatureLevel == GetFeatureLevel());
+
+		D3D12_FEATURE_DATA_D3D12_OPTIONS D3D12Caps;
+		VERIFYD3D11RESULT(GetDevice()->CheckFeatureSupport(D3D12_FEATURE_D3D12_OPTIONS, &D3D12Caps, sizeof(D3D12Caps)));
+		ResourceHeapTier = D3D12Caps.ResourceHeapTier;
+
+		// Init offline descriptor allocators
+		RTVAllocator.Init(GetDevice());
+		DSVAllocator.Init(GetDevice());
+		SRVAllocator.Init(GetDevice());
+		UAVAllocator.Init(GetDevice());
+		SamplerAllocator.Init(GetDevice());
+		CBVAllocator.Init(GetDevice());
+
+		// Init the occlusion query heap
+		OcclusionQueryHeap.Init();
+
+		// Create the main set of command lists used for rendering a frame
+		CommandListManager.Create(Direct3DDevice, D3D12_COMMAND_LIST_TYPE_DIRECT);
+		CopyCommandListManager.Create(Direct3DDevice, D3D12_COMMAND_LIST_TYPE_COPY);
+
+#if !(UE_BUILD_SHIPPING && WITH_EDITOR)
+		// Add some filter outs for known debug spew messages (that we don't care about)
+		if (D3D12RHI_ShouldCreateWithD3DDebug())
+		{
+			ID3D12InfoQueue *pd3dInfoQueue = NULL;
+            VERIFYD3D11RESULT(Direct3DDevice->QueryInterface(__uuidof(ID3D12InfoQueue), (void**)&pd3dInfoQueue));
+			if (pd3dInfoQueue)
+			{
+				D3D12_INFO_QUEUE_FILTER NewFilter;
+				FMemory::Memzero(&NewFilter, sizeof(NewFilter));
+
+				// Turn off info msgs as these get really spewy
+				D3D12_MESSAGE_SEVERITY DenySeverity = D3D12_MESSAGE_SEVERITY_INFO;
+				NewFilter.DenyList.NumSeverities = 1;
+				NewFilter.DenyList.pSeverityList = &DenySeverity;
+
+				// Be sure to carefully comment the reason for any additions here!  Someone should be able to look at it later and get an idea of whether it is still necessary.
+				D3D12_MESSAGE_ID DenyIds[] = {
+					// OMSETRENDERTARGETS_INVALIDVIEW - d3d will complain if depth and color targets don't have the exact same dimensions, but actually
+					//	if the color target is smaller then things are ok.  So turn off this error.  There is a manual check in FD3D12DynamicRHI::SetRenderTarget
+					//	that tests for depth smaller than color and MSAA settings to match.
+					D3D12_MESSAGE_ID_OMSETRENDERTARGETS_INVALIDVIEW, 
+
+					// QUERY_BEGIN_ABANDONING_PREVIOUS_RESULTS - The RHI exposes the interface to make and issue queries and a separate interface to use that data.
+					//		Currently there is a situation where queries are issued and the results may be ignored on purpose.  Filtering out this message so it doesn't
+					//		swarm the debug spew and mask other important warnings
+					//D3D12_MESSAGE_ID_QUERY_BEGIN_ABANDONING_PREVIOUS_RESULTS,
+					//D3D12_MESSAGE_ID_QUERY_END_ABANDONING_PREVIOUS_RESULTS,
+
+					// D3D12_MESSAGE_ID_CREATEINPUTLAYOUT_EMPTY_LAYOUT - This is a warning that gets triggered if you use a null vertex declaration,
+					//       which we want to do when the vertex shader is generating vertices based on ID.
+					D3D12_MESSAGE_ID_CREATEINPUTLAYOUT_EMPTY_LAYOUT,
+
+					// D3D12_MESSAGE_ID_COMMAND_LIST_DRAW_INDEX_BUFFER_TOO_SMALL - This warning gets triggered by Slate draws which are actually using a valid index range.
+					//		The invalid warning seems to only happen when VS 2012 is installed.  Reported to MS.  
+					//		There is now an assert in DrawIndexedPrimitive to catch any valid errors reading from the index buffer outside of range.
+                    D3D12_MESSAGE_ID_COMMAND_LIST_DRAW_INDEX_BUFFER_TOO_SMALL,
+
+					// D3D12_MESSAGE_ID_DEVICE_DRAW_RENDERTARGETVIEW_NOT_SET - This warning gets triggered by shadow depth rendering because the shader outputs
+					//		a color but we don't bind a color render target. That is safe as writes to unbound render targets are discarded.
+					//		Also, batched elements triggers it when rendering outside of scene rendering as it outputs to the GBuffer containing normals which is not bound.
+					//(D3D12_MESSAGE_ID)3146081, // D3D12_MESSAGE_ID_DEVICE_DRAW_RENDERTARGETVIEW_NOT_SET,
+                    // BUGBUG: There is a D3D12_MESSAGE_ID_DEVICE_DRAW_DEPTHSTENCILVIEW_NOT_SET, why not one for RT?
+
+                    // D3D12_MESSAGE_ID_CLEARDEPTHSTENCILVIEW_MISMATCHINGCLEARVALUE/D3D12_MESSAGE_ID_CLEARRENDERTARGETVIEW_MISMATCHINGCLEARVALUE - 
+                    //      This warning gets triggered by ClearDepthStencilView/ClearRenderTargetView because when the resource was created
+                    //      it wasn't passed an optimized clear color (see CreateCommitedResource). This shows up a lot and is very noisy.
+                    D3D12_MESSAGE_ID_CLEARRENDERTARGETVIEW_MISMATCHINGCLEARVALUE,
+                    D3D12_MESSAGE_ID_CLEARDEPTHSTENCILVIEW_MISMATCHINGCLEARVALUE,
+
+					// D3D12_MESSAGE_ID_EXECUTECOMMANDLISTS_GPU_WRITTEN_READBACK_RESOURCE_MAPPED - This warning gets triggered by ExecuteCommandLists.
+					//		if it contains a readback resource that still has mapped subresources when executing a command list that performs a copy operation to the resource.
+					//		This may be ok if any data read from the readback resources was flushed by calling Unmap() after the resourcecopy operation completed.
+					//		We intentionally keep the readback resources persistently mapped.
+					D3D12_MESSAGE_ID_EXECUTECOMMANDLISTS_GPU_WRITTEN_READBACK_RESOURCE_MAPPED
+				};
+
+				NewFilter.DenyList.NumIDs = sizeof(DenyIds) / sizeof(D3D12_MESSAGE_ID);
+				NewFilter.DenyList.pIDList = (D3D12_MESSAGE_ID*)&DenyIds;
+
+				pd3dInfoQueue->PushStorageFilter(&NewFilter);
+
+				// Break on D3D debug errors.
+				pd3dInfoQueue->SetBreakOnSeverity(D3D12_MESSAGE_SEVERITY_ERROR, true);
+
+				// Enable this to break on a specific id in order to quickly get a callstack
+				//pd3dInfoQueue->SetBreakOnID(D3D12_MESSAGE_ID_DEVICE_DRAW_CONSTANT_BUFFER_TOO_SMALL, true);
+
+				if (FParse::Param(FCommandLine::Get(), TEXT("d3dbreakonwarning")))
+				{
+					pd3dInfoQueue->SetBreakOnSeverity(D3D12_MESSAGE_SEVERITY_WARNING, true);
+				}
+
+				pd3dInfoQueue->Release();
+			}
+		}
+#endif
+
+		SetupAfterDeviceCreation();
+	}
+}
+
+/**
+ *	Retrieve available screen resolutions.
+ *
+ *	@param	Resolutions			TArray<FScreenResolutionRHI> parameter that will be filled in.
+ *	@param	bIgnoreRefreshRate	If true, ignore refresh rates.
+ *
+ *	@return	bool				true if successfully filled the array
+ */
+bool FD3D12DynamicRHI::RHIGetAvailableResolutions(FScreenResolutionArray& Resolutions, bool bIgnoreRefreshRate)
+{
+	int32 MinAllowableResolutionX = 0;
+	int32 MinAllowableResolutionY = 0;
+	int32 MaxAllowableResolutionX = 10480;
+	int32 MaxAllowableResolutionY = 10480;
+	int32 MinAllowableRefreshRate = 0;
+	int32 MaxAllowableRefreshRate = 10480;
+
+	if (MaxAllowableResolutionX == 0)
+	{
+		MaxAllowableResolutionX = 10480;
+	}
+	if (MaxAllowableResolutionY == 0)
+	{
+		MaxAllowableResolutionY = 10480;
+	}
+	if (MaxAllowableRefreshRate == 0)
+	{
+		MaxAllowableRefreshRate = 10480;
+	}
+
+	HRESULT hr = S_OK;
+	TRefCountPtr<IDXGIAdapter> Adapter;
+    //TODO: should only be called on display out device
+    hr = DXGIFactory->EnumAdapters(GetRHIDevice()->GetAdapterIndex(), Adapter.GetInitReference());
+
+	if (DXGI_ERROR_NOT_FOUND == hr)
+		return false;
+	if (FAILED(hr))
+		return false;
+
+	// get the description of the adapter
+	DXGI_ADAPTER_DESC AdapterDesc;
+	if (FAILED(Adapter->GetDesc(&AdapterDesc)))
+	{
+		return false;
+	}
+
+	int32 CurrentOutput = 0;
+	do 
+	{
+		TRefCountPtr<IDXGIOutput> Output;
+		hr = Adapter->EnumOutputs(CurrentOutput, Output.GetInitReference());
+		if (DXGI_ERROR_NOT_FOUND == hr)
+			break;
+		if (FAILED(hr))
+			return false;
+
+		// TODO: GetDisplayModeList is a terribly SLOW call.  It can take up to a second per invocation.
+		//  We might want to work around some DXGI badness here.
+		DXGI_FORMAT Format = DXGI_FORMAT_R8G8B8A8_UNORM;
+		uint32 NumModes = 0;
+		hr = Output->GetDisplayModeList(Format, 0, &NumModes, NULL);
+		if (hr == DXGI_ERROR_NOT_FOUND)
+		{
+			continue;
+		}
+		else if (hr == DXGI_ERROR_NOT_CURRENTLY_AVAILABLE)
+		{
+			UE_LOG(LogD3D12RHI, Fatal,
+				TEXT("This application cannot be run over a remote desktop configuration")
+				);
+			return false;
+		}
+
+		checkf(NumModes > 0, TEXT("No display modes found for the standard format DXGI_FORMAT_R8G8B8A8_UNORM!"));
+
+		DXGI_MODE_DESC* ModeList = new DXGI_MODE_DESC[NumModes];
+		VERIFYD3D11RESULT(Output->GetDisplayModeList(Format, 0, &NumModes, ModeList));
+
+		for (uint32 m = 0; m < NumModes; m++)
+		{
+			if (((int32)ModeList[m].Width >= MinAllowableResolutionX) &&
+				((int32)ModeList[m].Width <= MaxAllowableResolutionX) &&
+				((int32)ModeList[m].Height >= MinAllowableResolutionY) &&
+				((int32)ModeList[m].Height <= MaxAllowableResolutionY)
+				)
+			{
+				bool bAddIt = true;
+				if (bIgnoreRefreshRate == false)
+				{
+					if (((int32)ModeList[m].RefreshRate.Numerator < MinAllowableRefreshRate * ModeList[m].RefreshRate.Denominator) ||
+						((int32)ModeList[m].RefreshRate.Numerator > MaxAllowableRefreshRate * ModeList[m].RefreshRate.Denominator)
+						)
+					{
+						continue;
+					}
+				}
+				else
+				{
+					// See if it is in the list already
+					for (int32 CheckIndex = 0; CheckIndex < Resolutions.Num(); CheckIndex++)
+					{
+						FScreenResolutionRHI& CheckResolution = Resolutions[CheckIndex];
+						if ((CheckResolution.Width == ModeList[m].Width) &&
+							(CheckResolution.Height == ModeList[m].Height))
+						{
+							// Already in the list...
+							bAddIt = false;
+							break;
+						}
+					}
+				}
+
+				if (bAddIt)
+				{
+					// Add the mode to the list
+					int32 Temp2Index = Resolutions.AddZeroed();
+					FScreenResolutionRHI& ScreenResolution = Resolutions[Temp2Index];
+
+					ScreenResolution.Width = ModeList[m].Width;
+					ScreenResolution.Height = ModeList[m].Height;
+					ScreenResolution.RefreshRate = ModeList[m].RefreshRate.Numerator / ModeList[m].RefreshRate.Denominator;
+				}
+			}
+		}
+
+		delete[] ModeList;
+
+		++CurrentOutput;
+
+	// TODO: Cap at 1 for default output
+	} while (CurrentOutput < 1);
+
+	return true;
+}
diff --git a/Engine/Source/Runtime/Windows/D3D12RHI/Private/Windows/WindowsD3D12Viewport.cpp b/Engine/Source/Runtime/Windows/D3D12RHI/Private/Windows/WindowsD3D12Viewport.cpp
new file mode 100644
--- /dev/null
+++ b/Engine/Source/Runtime/Windows/D3D12RHI/Private/Windows/WindowsD3D12Viewport.cpp
@@ -0,0 +1,106 @@
+// Copyright 1998-2014 Epic Games, Inc. All Rights Reserved.
+
+/*=============================================================================
+	D3D12Viewport.cpp: D3D viewport RHI implementation.
+=============================================================================*/
+
+#include "D3D12RHIPrivate.h"
+#include "RenderCore.h"
+
+#include "AllowWindowsPlatformTypes.h"
+#include <dwmapi.h>
+
+extern FD3D12Texture2D* GetSwapChainSurface(FD3D12Device* Parent, EPixelFormat PixelFormat, IDXGISwapChain3* SwapChain, const uint32 &backBufferIndex);
+
+FD3D12Viewport::FD3D12Viewport(class FD3D12Device* InParent, HWND InWindowHandle, uint32 InSizeX, uint32 InSizeY, bool bInIsFullscreen, EPixelFormat InPreferredPixelFormat) :
+	LastFlipTime(0),
+	LastFrameComplete(0),
+	LastCompleteTime(0),
+	SyncCounter(0),
+	bSyncedLastFrame(false),
+	WindowHandle(InWindowHandle),
+    MaximumFrameLatency(3),
+	SizeX(InSizeX),
+	SizeY(InSizeY),
+	bIsFullscreen(bInIsFullscreen),
+	PixelFormat(InPreferredPixelFormat),
+	bIsValid(true),
+	bIsBenchmarkMode(false),
+	NumBackBuffers(DefaultNumBackBuffers),
+    FD3D12DeviceChild(InParent)
+{
+	check(IsInGameThread());
+    GetParentDevice()->GetViewports().Add(this);
+}
+
+//Init for a Viewport that will do the presenting
+void FD3D12Viewport::Init(IDXGIFactory4* Factory, bool AssociateWindow)
+{
+    // Ensure that the D3D devices have been created.
+    // TODO: is this really necessary?
+    //D3DRHI->InitD3DDevices();
+
+	bIsBenchmarkMode = FParse::Param(FCommandLine::Get(), TEXT("benchmark"));
+
+	DXGI_SWAP_CHAIN_FLAG swapChainFlags = DXGI_SWAP_CHAIN_FLAG_ALLOW_MODE_SWITCH;
+	if (bIsBenchmarkMode)
+	{
+		swapChainFlags = static_cast<DXGI_SWAP_CHAIN_FLAG>(static_cast<uint32>(swapChainFlags) | DXGI_SWAP_CHAIN_FLAG_FRAME_LATENCY_WAITABLE_OBJECT);
+		NumBackBuffers = MaxNumBackBuffers;
+	}
+
+	DXGI_SWAP_CHAIN_DESC SwapChainDesc;
+	FMemory::Memzero(&SwapChainDesc, sizeof(DXGI_SWAP_CHAIN_DESC));
+
+	// Create the swapchain.
+	{
+		TRefCountPtr<IDXGISwapChain> SwapChain;
+		SwapChainDesc.BufferDesc = SetupDXGI_MODE_DESC();
+		// MSAA Sample count
+		SwapChainDesc.SampleDesc.Count = 1;
+		SwapChainDesc.SampleDesc.Quality = 0;
+		SwapChainDesc.BufferUsage = DXGI_USAGE_RENDER_TARGET_OUTPUT | DXGI_USAGE_SHADER_INPUT;
+		// 1:single buffering, 2:double buffering, 3:triple buffering
+		SwapChainDesc.BufferCount = NumBackBuffers;
+		SwapChainDesc.OutputWindow = WindowHandle;
+		SwapChainDesc.Windowed = !bIsFullscreen;
+		// DXGI_SWAP_EFFECT_DISCARD / DXGI_SWAP_EFFECT_SEQUENTIAL
+		SwapChainDesc.SwapEffect = DXGI_SWAP_EFFECT_FLIP_DISCARD;
+		SwapChainDesc.Flags = swapChainFlags;
+		ID3D12CommandQueue* CommandQueue = GetParentDevice()->GetCommandListManager().GetD3DCommandQueue();
+
+		VERIFYD3D11RESULT(Factory->CreateSwapChain(CommandQueue, &SwapChainDesc, SwapChain.GetInitReference()));
+
+		VERIFYD3D11RESULT(SwapChain->QueryInterface(IID_PPV_ARGS(SwapChain3.GetInitReference())));
+	}
+
+	if (bIsBenchmarkMode)
+	{
+		VERIFYD3D11RESULT(SwapChain3->SetMaximumFrameLatency(NumBackBuffers - 1));
+	}
+
+    if (AssociateWindow)
+    {
+        // Set the DXGI message hook to not change the window behind our back.
+        Factory->MakeWindowAssociation(WindowHandle, DXGI_MWA_NO_WINDOW_CHANGES);
+    }
+
+    // Create a RHI surface to represent the viewport's back buffer.
+    for (uint32 i = 0; i < NumBackBuffers; ++i)
+    {
+        BackBuffers[i] = GetSwapChainSurface(GetParentDevice(), PixelFormat, SwapChain3, i);
+    }
+
+    if (AssociateWindow)
+    {
+        // Tell the window to redraw when they can.
+        // @todo: For Slate viewports, it doesn't make sense to post WM_PAINT messages (we swallow those.)
+        ::PostMessage(WindowHandle, WM_PAINT, 0, 0);
+    }
+}
+
+void FD3D12Viewport::ConditionalResetSwapChain(bool bIgnoreFocus)
+{
+}
+
+#include "HideWindowsPlatformTypes.h"
diff --git a/Engine/Source/Runtime/Windows/D3D12RHI/Public/D3D12ConstantBuffer.h b/Engine/Source/Runtime/Windows/D3D12RHI/Public/D3D12ConstantBuffer.h
new file mode 100644
--- /dev/null
+++ b/Engine/Source/Runtime/Windows/D3D12RHI/Public/D3D12ConstantBuffer.h
@@ -0,0 +1,75 @@
+// Copyright 1998-2014 Epic Games, Inc. All Rights Reserved.
+
+/*=============================================================================
+	D3D12ConstantBuffer.h: Public D3D Constant Buffer definitions.
+	=============================================================================*/
+
+#ifndef __D3D11ConstantBuffer_H__
+#define __D3D11ConstantBuffer_H__
+
+#include "RenderResource.h"
+
+/** Size of the default constant buffer. */
+#define MAX_GLOBAL_CONSTANT_BUFFER_SIZE		4096
+
+// !!! These offsets must match the cbuffer register definitions in Common.usf !!!
+enum ED3D11ShaderOffsetBuffer
+{
+	/** Default constant buffer. */
+	GLOBAL_CONSTANT_BUFFER_INDEX = 0,
+	MAX_CONSTANT_BUFFER_SLOTS
+};
+
+namespace D3D12RHI
+{
+	/** Sizes of constant buffers defined in ED3D11ShaderOffsetBuffer. */
+	extern const uint32 GConstantBufferSizes[MAX_CONSTANT_BUFFER_SLOTS];
+}
+using namespace D3D12RHI;
+
+/**
+ * A D3D constant buffer
+ */
+class FD3D12ConstantBuffer : public FRenderResource, public FRHIResource, public FD3D12DeviceChild
+{
+public:
+
+	// New circular buffer system for faster constant uploads.  Avoids CopyResource and speeds things up considerably
+	FD3D12ConstantBuffer(FD3D12Device* InParent, uint32 InSize = 0, uint32 SubBuffers = 1);
+	virtual ~FD3D12ConstantBuffer();
+
+	// FRenderResource interface.
+	virtual void	InitDynamicRHI() override;
+	virtual void	ReleaseDynamicRHI() override;
+
+	/**
+	* Updates a variable in the constant buffer.
+	* @param Data - The data to copy into the constant buffer
+	* @param Offset - The offset in the constant buffer to place the data at
+	* @param InSize - The size of the data being copied
+	*/
+	void UpdateConstant(const uint8* Data, uint16 Offset, uint16 InSize)
+	{
+		// Check that the data we are shadowing fits in the allocated shadow data
+		check((uint32)Offset + (uint32)InSize <= MaxSize);
+		FMemory::Memcpy(ShadowData + Offset, Data, InSize);
+		CurrentUpdateSize = FMath::Max((uint32)(Offset + InSize), CurrentUpdateSize);
+	}
+
+protected:
+	uint32	MaxSize;
+	uint8*	ShadowData;
+
+	/** Size of all constants that has been updated since the last call to Commit. */
+	uint32	CurrentUpdateSize;
+
+	/**
+	 * Size of all constants that has been updated since the last Discard.
+	 * Includes "shared" constants that don't necessarily gets updated between every Commit.
+	 */
+	uint32	TotalUpdateSize;
+};
+
+DECLARE_CYCLE_STAT_EXTERN(TEXT("Global Constant buffer update time"), STAT_D3D12GlobalConstantBufferUpdateTime, STATGROUP_D3D12RHI, );
+
+#endif
diff --git a/Engine/Source/Runtime/Windows/D3D12RHI/Public/D3D12RHI.h b/Engine/Source/Runtime/Windows/D3D12RHI/Public/D3D12RHI.h
new file mode 100644
--- /dev/null
+++ b/Engine/Source/Runtime/Windows/D3D12RHI/Public/D3D12RHI.h
@@ -0,0 +1,12 @@
+// Copyright 1998-2014 Epic Games, Inc. All Rights Reserved.
+
+/*=============================================================================
+	D3D12RHI.h: Public D3D RHI definitions.
+=============================================================================*/
+
+#ifndef _INC_D3D12RHI
+#define _INC_D3D12RHI
+
+DECLARE_STATS_GROUP(TEXT("D3D12RHI"), STATGROUP_D3D12RHI, STATCAT_Advanced);
+
+#endif
