UE4 - D3D12RHI (not working)

diff --git a/Engine/Source/Runtime/Windows/D3D12RHI/Public/D3D12Resources.h b/Engine/Source/Runtime/Windows/D3D12RHI/Public/D3D12Resources.h
new file mode 100644
--- /dev/null
+++ b/Engine/Source/Runtime/Windows/D3D12RHI/Public/D3D12Resources.h
@@ -0,0 +1,3132 @@
+// Copyright 1998-2014 Epic Games, Inc. All Rights Reserved.
+
+/*=============================================================================
+	D3D12Resources.h: D3D resource RHI definitions.
+	=============================================================================*/
+
+#pragma once
+
+#include "BoundShaderStateCache.h"
+#include "D3D12ShaderResources.h"
+
+class FD3D12Resource;
+class FD3D12StateCacheBase;
+typedef FD3D12StateCacheBase FD3D12StateCache;
+
+class FD3D12PendingResourceBarrier
+{
+public:
+	
+	TRefCountPtr<FD3D12Resource> Resource;
+	D3D12_RESOURCE_STATES        State;
+	uint32                       SubResource;
+};
+
+class FD3D12CommandListManager;
+class FD3D12CommandContext;
+
+class FD3D12CommandListHandle
+{
+public:
+	static const uint32	ResourceBarrierScratchSpaceSize = 256;
+private:
+
+	class FD3D12CommandListData
+	{
+	public:
+
+		FD3D12CommandListData (ID3D12Device* Direct3DDevice, D3D12_COMMAND_LIST_TYPE CommandListType, FD3D12CommandListManager* InCommandListManager)
+			: CommandListManager (InCommandListManager)
+			, FenceIndex (0)
+			, CurrentGeneration(0)
+			, IsClosed (false)
+			, PendingResourceBarriers()
+			, CurrentOwningContext(nullptr)
+		{
+			VERIFYD3D11RESULT(Direct3DDevice->CreateCommandAllocator(CommandListType, IID_PPV_ARGS(CommandAllocator.GetInitReference())));
+			VERIFYD3D11RESULT(Direct3DDevice->CreateCommandList(0, CommandListType, CommandAllocator, nullptr, IID_PPV_ARGS(CommandList.GetInitReference())));
+
+			// Initially start with all lists closed.  We'll open them as we allocate them.
+			Close();
+
+			PendingResourceBarriers.Reserve (256);
+		}
+
+		virtual ~FD3D12CommandListData()
+		{
+			CommandAllocator.SafeRelease();
+			CommandList.SafeRelease();
+		}
+
+		void Close()
+		{
+			if (!IsClosed)
+			{
+				VERIFYD3D11RESULT(CommandList->Close());
+				IsClosed = true;
+			}
+		}
+
+		// Reset the command list to the beginning.  Requires execution on the GPU to have completed (verified by a fence.)
+		void Reset (ID3D12PipelineState* InitialPSO = nullptr)
+		{
+			// Reset (can't reset the allocator until the GPU is done with it)
+			VERIFYD3D11RESULT(CommandAllocator->Reset());
+			VERIFYD3D11RESULT(CommandList->Reset(CommandAllocator, InitialPSO));
+
+			FenceIndex = (uint64)-1;
+			IsClosed   = false;
+			CurrentGeneration++;
+		}
+
+		uint32 AddRef() const
+		{
+			int32 NewValue = NumRefs.Increment();
+			check(NewValue > 0); 
+			return uint32(NewValue);
+		}
+		uint32 Release() const
+		{
+			int32 NewValue = NumRefs.Decrement();
+			check(NewValue >= 0);
+			return uint32(NewValue);
+		}
+
+		mutable FThreadSafeCounter				NumRefs;
+		FD3D12CommandListManager*               CommandListManager;
+		FD3D12CommandContext*					CurrentOwningContext;
+		TRefCountPtr<FD3D12CommandList>			CommandList;		// Raw D3D command list pointer
+		TRefCountPtr<FD3D12CommandAllocator>	CommandAllocator;	// Command allocator paired with the command list
+		uint64									FenceIndex;			// -1 when active; positive when submitted for execution
+		uint64									CurrentGeneration;
+		bool                                    IsClosed;
+
+		// Array of resources who's state needs to be synced between submits.
+		TArray<FD3D12PendingResourceBarrier>	PendingResourceBarriers;
+		// Scratch space to accelarate filling out resource barrier descriptors
+		D3D12_RESOURCE_BARRIER					ResourceBarrierScratchSpace[ResourceBarrierScratchSpaceSize];
+		/**
+		*	A map of all D3D resources, and their states, that were state transitioned with tracking.
+		*/
+		class FCommandListResourceState
+		{
+		private:
+			TMap<FD3D12Resource*, CResourceState> ResourceStates;
+			void inline ConditionalInitalize(FD3D12Resource* pResource, CResourceState& ResourceState);
+
+		public:
+			CResourceState& GetResourceState(FD3D12Resource* pResource);
+
+			// Empty the command list's resource state map after the command list is executed
+			void Empty();
+		};
+
+		FCommandListResourceState TrackedResourceState;
+	};
+
+public:
+
+	FD3D12CommandListHandle() : CommandListData (nullptr) {}
+
+	FD3D12CommandListHandle (const FD3D12CommandListHandle& CL) : CommandListData (CL.CommandListData)
+	{
+		if (CommandListData)
+			CommandListData->AddRef();
+	}
+
+	virtual ~FD3D12CommandListHandle()
+	{
+		if (CommandListData && CommandListData->Release() == 0)
+		{
+			delete CommandListData;
+		}
+	}
+	
+	FD3D12CommandListHandle& operator = (const FD3D12CommandListHandle* CL)
+	{
+		if (this != CL)
+		{
+			if (CommandListData && CommandListData->Release() == 0)
+			{
+				delete CommandListData;
+			}
+
+			CommandListData = nullptr;
+
+			if (CL && CL->CommandListData)
+			{
+				CommandListData = CL->CommandListData;
+
+				CommandListData->AddRef();
+			}
+		}
+
+		return *this;
+	}
+
+	FD3D12CommandListHandle& operator = (const FD3D12CommandListHandle& CL)
+	{
+		if (this != &CL)
+		{
+			if (CommandListData && CommandListData->Release() == 0)
+			{
+				delete CommandListData;
+			}
+
+			CommandListData = nullptr;
+
+			if (CL.CommandListData)
+			{
+				CommandListData = CL.CommandListData;
+
+				CommandListData->AddRef();
+			}
+		}
+
+		return *this;
+	}
+
+	bool operator!() const
+	{
+		return CommandListData == 0;
+	}
+
+	inline friend bool operator==(const FD3D12CommandListHandle& lhs, const FD3D12CommandListHandle& rhs)
+	{
+		return lhs.CommandListData == rhs.CommandListData;
+	}
+
+	inline friend bool operator==(const FD3D12CommandListHandle& lhs, const FD3D12CommandListData* rhs)
+	{
+		return lhs.CommandListData == rhs;
+	}
+
+	inline friend bool operator==(const FD3D12CommandListData* lhs, const FD3D12CommandListHandle& rhs)
+	{
+		return lhs == rhs.CommandListData;
+	}
+
+	inline friend bool operator!=(const FD3D12CommandListHandle& lhs, const FD3D12CommandListData* rhs)
+	{
+		return lhs.CommandListData != rhs;
+	}
+
+	inline friend bool operator!=(const FD3D12CommandListData* lhs, const FD3D12CommandListHandle& rhs)
+	{
+		return lhs != rhs.CommandListData;
+	}
+
+	FD3D12CommandList* operator->() const
+	{
+		check (CommandListData && !CommandListData->IsClosed);
+
+		return CommandListData->CommandList; 
+	}
+
+	void Create (ID3D12Device* Direct3DDevice, D3D12_COMMAND_LIST_TYPE CommandListType, FD3D12CommandListManager* InCommandListManager)
+	{
+		check (!CommandListData);
+
+		CommandListData = new FD3D12CommandListData (Direct3DDevice, CommandListType, InCommandListManager);
+
+		CommandListData->AddRef();
+	}
+
+	void Execute (bool WaitForCompletion = false);
+
+	void Close()
+	{
+		check (CommandListData);
+		CommandListData->Close();
+	}
+
+	// Reset the command list to the beginning.  Requires execution on the GPU to have completed (verified by a fence.)
+	void Reset (ID3D12PipelineState* InitialPSO = nullptr)
+	{
+		check (CommandListData);
+		CommandListData->Reset();
+	}
+
+	FD3D12CommandList* CommandList() const
+	{
+		check (CommandListData);
+		return CommandListData->CommandList.GetReference();
+	}
+
+	uint64 FenceIndex() const
+	{
+		check (CommandListData);
+		return CommandListData->FenceIndex;
+	}
+
+	void SetFenceIndex(uint64 FenceIndex)
+	{
+		check (CommandListData);
+		CommandListData->FenceIndex = FenceIndex;
+	}
+
+	uint64 CurrentGeneration() const
+	{
+		check (CommandListData);
+		return CommandListData->CurrentGeneration;
+	}
+
+	// Get the state of a resource on this command lists.
+	// This is only used for resources that require state tracking.
+	CResourceState& GetResourceState(FD3D12Resource* pResource)
+	{
+		check(CommandListData);
+
+		return CommandListData->TrackedResourceState.GetResourceState(pResource);
+	}
+
+	void AddPendingResourceBarrier (FD3D12Resource* Resource, D3D12_RESOURCE_STATES State, uint32 SubResource)
+	{
+		check (CommandListData);
+
+		FD3D12PendingResourceBarrier PRB = { Resource, State, SubResource };
+
+		CommandListData->PendingResourceBarriers.Add (PRB);
+	}
+
+	TArray<FD3D12PendingResourceBarrier>& PendingResourceBarriers()
+	{
+		check (CommandListData);
+		return CommandListData->PendingResourceBarriers;
+	}
+
+	// Empty all the resource states being tracked on this command list
+	void EmptyTrackedResourceState()
+	{
+		check(CommandListData);
+
+		CommandListData->TrackedResourceState.Empty();
+	}
+
+	void SetCurrentOwningContext(FD3D12CommandContext* context)
+	{
+		CommandListData->CurrentOwningContext = context;
+	}
+
+	FD3D12CommandContext* GetCurrentOwningContext()
+	{
+		return CommandListData->CurrentOwningContext;
+	}
+
+	D3D12_RESOURCE_BARRIER* GetResourceBarrierScratchSpace()
+	{
+		return CommandListData->ResourceBarrierScratchSpace;
+	}
+
+private:
+
+	FD3D12CommandListHandle& operator*()
+	{
+		return *this;
+	}
+
+	FD3D12CommandListData* CommandListData;
+};
+
+class FD3D12FrameSyncPoint
+{
+	FD3D12FrameSyncPoint()
+		: FrameFence(0)
+	{
+	}
+
+private:
+	friend class FD3D12CommandListManager;
+
+	uint64 FrameFence;
+};
+
+class FD3D12CLSyncPoint
+{
+public:
+
+	FD3D12CLSyncPoint() : Generation (0) {}
+
+	FD3D12CLSyncPoint(FD3D12CommandListHandle& CL) : CommandList (CL), Generation (CL.CommandList() ? CL.CurrentGeneration() : 0) {}
+
+	FD3D12CLSyncPoint(const FD3D12CLSyncPoint& SyncPoint) : CommandList (SyncPoint.CommandList), Generation (SyncPoint.Generation) {}
+
+	FD3D12CLSyncPoint& operator = (FD3D12CommandListHandle& CL)
+	{
+		CommandList = CL;
+		Generation = (CL != nullptr) ? CL.CurrentGeneration() : 0;
+
+		return *this;
+	}
+
+	FD3D12CLSyncPoint& operator = (const FD3D12CLSyncPoint& SyncPoint)
+	{
+		CommandList = SyncPoint.CommandList;
+		Generation = SyncPoint.Generation;
+
+		return *this;
+	}
+
+	bool operator!() const
+	{
+		return CommandList == 0;
+	}
+
+private:
+
+	friend class FD3D12CommandListManager;
+
+	FD3D12CommandListHandle CommandList;
+	uint64                  Generation;
+};
+
+class FD3D12CommandContext;
+
+class FD3D12RefCount
+{
+public:
+	FD3D12RefCount()
+	{
+	}
+	virtual ~FD3D12RefCount() 
+	{ 
+		check(NumRefs.GetValue() == 0);
+	}
+	uint32 AddRef() const
+	{
+		int32 NewValue = NumRefs.Increment();
+		check(NewValue > 0); 
+		return uint32(NewValue);
+	}
+	uint32 Release() const
+	{
+		int32 NewValue = NumRefs.Decrement();
+		if (NewValue == 0)
+		{
+			delete this;
+		}
+		check(NewValue >= 0);
+		return uint32(NewValue);
+	}
+	uint32 GetRefCount() const
+	{
+		int32 CurrentValue = NumRefs.GetValue();
+		check(CurrentValue >= 0); 
+		return uint32(CurrentValue);
+	}
+private:
+	mutable FThreadSafeCounter NumRefs;
+};
+
+#if SUPPORTS_MEMORY_RESIDENCY
+struct FD3D12ResidencyHandle
+{
+	static const uint32 INVALID_INDEX = 0x7FFFFFFF;
+
+	FD3D12ResidencyHandle()
+		: LastUpdatedFence (0)
+		, IsResident (1)
+		, Index (INVALID_INDEX)
+	{}
+
+	uint32 LastUpdatedFence;
+	uint32 IsResident : 1;
+	uint32 Index      : 31;
+};
+#endif
+
+/** Find the appropriate depth-stencil typeless DXGI format for the given format. */
+inline DXGI_FORMAT FindDepthStencilParentDXGIFormat(DXGI_FORMAT InFormat)
+{
+	switch (InFormat)
+	{
+	case DXGI_FORMAT_D24_UNORM_S8_UINT:
+		return DXGI_FORMAT_R24G8_TYPELESS;
+#if DEPTH_32_BIT_CONVERSION
+		// Changing Depth Buffers to 32 bit on Dingo as D24S8 is actually implemented as a 32 bit buffer in the hardware
+	case DXGI_FORMAT_D32_FLOAT_S8X24_UINT:
+		return DXGI_FORMAT_R32G8X24_TYPELESS;
+#endif
+	case DXGI_FORMAT_D32_FLOAT:
+		return DXGI_FORMAT_R32_TYPELESS;
+	case DXGI_FORMAT_D16_UNORM:
+		return DXGI_FORMAT_R16_TYPELESS;
+	};
+	return InFormat;
+}
+
+static uint8 GetPlaneSliceFromViewFormat(DXGI_FORMAT ResourceFormat, DXGI_FORMAT ViewFormat)
+{
+	// Currently, the only planar resources used are depth-stencil formats
+	switch (FindDepthStencilParentDXGIFormat(ResourceFormat))
+	{
+	case DXGI_FORMAT_R24G8_TYPELESS:
+		switch (ViewFormat)
+		{
+		case DXGI_FORMAT_R24_UNORM_X8_TYPELESS:
+			return 0;
+		case DXGI_FORMAT_X24_TYPELESS_G8_UINT:
+			return 1;
+		}
+		break;
+	case DXGI_FORMAT_R32G8X24_TYPELESS:
+		switch (ViewFormat)
+		{
+		case DXGI_FORMAT_R32_FLOAT_X8X24_TYPELESS:
+			return 0;
+		case DXGI_FORMAT_X32_TYPELESS_G8X24_UINT:
+			return 1;
+		}
+		break;
+	}
+
+	return 0;
+}
+
+static uint8 GetPlaneCount(DXGI_FORMAT Format)
+{
+	// Currently, the only planar resources used are depth-stencil formats
+	// Note there is a D3D12 helper for this, D3D12GetFormatPlaneCount
+	switch (FindDepthStencilParentDXGIFormat(Format))
+	{
+	case DXGI_FORMAT_R24G8_TYPELESS:
+	case DXGI_FORMAT_R32G8X24_TYPELESS:
+		return 2;
+	default:
+		return 1;
+	}
+}
+
+class FD3D12Resource : public FD3D12RefCount
+{
+private:
+	TRefCountPtr<ID3D12Resource> Resource;
+	TRefCountPtr<ID3D12Heap> Heap;
+
+	D3D12_RESOURCE_DESC Desc;
+	uint8 PlaneCount;
+	uint16 SubresourceCount;
+	CResourceState* pResourceState;
+	D3D12_RESOURCE_STATES DefaultResourceState;
+	bool bRequiresResourceStateTracking;
+	D3D12_HEAP_TYPE HeapType;
+
+#if UE_BUILD_DEBUG
+	static int64 TotalResourceCount;
+	static int64 NoStateTrackingResourceCount;
+#endif
+
+#if SUPPORTS_MEMORY_RESIDENCY
+	FD3D12ResidencyHandle ResidencyHandle;
+	uint32                ResourceSize;
+
+	friend class FD3D12ResourceResidencyManager;
+
+	void UpdateResourceSize();
+	uint32 GetResourceSize() { return ResourceSize; }
+	ID3D12Resource* GetResourceUntracked() const { return Resource.GetReference(); }
+#endif
+
+public:
+	explicit FD3D12Resource(ID3D12Resource* InResource, D3D12_RESOURCE_DESC const& InDesc, ID3D12Heap *InHeap = nullptr, D3D12_HEAP_TYPE InHeapType = D3D12_HEAP_TYPE_DEFAULT)
+		: Resource(InResource)
+		, Heap(InHeap)
+		, Desc(InDesc)
+		, PlaneCount(::GetPlaneCount(Desc.Format))
+		, SubresourceCount(0)
+		, pResourceState(nullptr)
+		, DefaultResourceState(D3D12_RESOURCE_STATE_TBD)
+		, bRequiresResourceStateTracking(true)
+		, HeapType(InHeapType)
+	{
+#if UE_BUILD_DEBUG
+		FPlatformAtomics::_InterlockedIncrement(&TotalResourceCount);
+#endif
+
+		InitalizeResourceState();
+
+#if SUPPORTS_MEMORY_RESIDENCY
+		UpdateResourceSize();
+		UpdateResidency();
+#endif
+	}
+
+	virtual ~FD3D12Resource();
+
+#if SUPPORTS_MEMORY_RESIDENCY
+	void UpdateResidency();
+
+	ID3D12Resource* GetResource() { UpdateResidency(); return Resource.GetReference(); }
+#else
+	operator ID3D12Resource&() { return *Resource; }
+	ID3D12Resource* GetResource() const { return Resource.GetReference(); }
+#endif
+
+	D3D12_RESOURCE_DESC const& GetDesc() const { return Desc; }
+	D3D12_HEAP_TYPE GetHeapType() const { return HeapType; }
+	uint16 GetMipLevels() const { return Desc.MipLevels; }
+	uint16 GetArraySize() const { return (Desc.Dimension == D3D12_RESOURCE_DIMENSION_TEXTURE3D) ? 1 : Desc.DepthOrArraySize; }
+	uint8 GetPlaneCount() const { return PlaneCount; }
+	uint16 GetSubresourceCount() const { return SubresourceCount; }
+	CResourceState* GetResourceState() 
+	{
+		// This state is used as the resource's "global" state between command lists. It's only needed for resources that
+		// require state tracking.
+		return pResourceState;
+	}
+	D3D12_RESOURCE_STATES GetDefaultResourceState() const { check(!bRequiresResourceStateTracking); return DefaultResourceState; }
+	bool RequiresResourceStateTracking() const { return bRequiresResourceStateTracking; }
+
+private:
+	void InitalizeResourceState()
+	{
+		SubresourceCount = GetMipLevels() * GetArraySize() * GetPlaneCount();
+		ComputeDefaultResourceState();
+
+		if (bRequiresResourceStateTracking)
+		{
+			// Only a few resources (~1%) actually need resource state tracking
+			pResourceState = new CResourceState();
+			pResourceState->Initialize(SubresourceCount);
+			pResourceState->SetResourceState(D3D12_RESOURCE_STATE_COMMON);
+		}
+	}
+
+	void ComputeDefaultResourceState()
+	{
+		const bool bSRV = (Desc.Flags & D3D12_RESOURCE_FLAG_DENY_SHADER_RESOURCE) == 0;
+		const bool bDSV = (Desc.Flags & D3D12_RESOURCE_FLAG_ALLOW_DEPTH_STENCIL) != 0;
+		const bool bRTV = (Desc.Flags & D3D12_RESOURCE_FLAG_ALLOW_RENDER_TARGET) != 0;
+		const bool bUAV = (Desc.Flags & D3D12_RESOURCE_FLAG_ALLOW_UNORDERED_ACCESS) != 0;
+
+		const bool bWriteable = bDSV || bRTV || bUAV;
+		const bool bSRVOnly = bSRV && !bWriteable;
+		const bool bBuffer = Desc.Dimension == D3D12_RESOURCE_DIMENSION_BUFFER;
+
+		if (bBuffer)
+		{
+			if (!bWriteable)
+			{
+				// Buffer used for input, like Vertex/Index buffer.
+				// Don't bother tracking state for this resource.
+#if UE_BUILD_DEBUG
+				FPlatformAtomics::_InterlockedIncrement(&NoStateTrackingResourceCount);
+#endif
+				DefaultResourceState = (HeapType == D3D12_HEAP_TYPE_READBACK) ? D3D12_RESOURCE_STATE_COPY_DEST : D3D12_RESOURCE_STATE_GENERIC_READ;
+				bRequiresResourceStateTracking = false;
+				return;
+				
+			}
+		}
+		else
+		{
+			if (bSRVOnly)
+			{
+				// Texture used only as a SRV.
+				// Don't bother tracking state for this resource.
+#if UE_BUILD_DEBUG
+				FPlatformAtomics::_InterlockedIncrement(&NoStateTrackingResourceCount);
+#endif
+				DefaultResourceState = D3D12_RESOURCE_STATE_NON_PIXEL_SHADER_RESOURCE | D3D12_RESOURCE_STATE_PIXEL_SHADER_RESOURCE;
+				bRequiresResourceStateTracking = false;
+				return;
+			}
+		}
+	}
+};
+
+// FD3D12ResourceBlockInfo
+class FD3D12ResourceBlockInfo : public FD3D12RefCount
+{
+public:
+	FD3D12ResourceBlockInfo()
+		: Address(nullptr)
+		, Offset(0)
+		, Bucket(0)
+		, ResourceHeap(nullptr)
+		, Allocator(nullptr)
+	{
+	}
+
+	FD3D12ResourceBlockInfo(FD3D12Resource *Resource, void *InAddress, uint64 InOffset, uint32 InBucket, class FD3D12ResourceAllocator *InAllocator)
+		: Address(InAddress)
+		, Offset(InOffset)
+		, Bucket(InBucket)
+		, ResourceHeap(Resource)
+		, Allocator(InAllocator)
+	{
+	}
+
+	void *Address;
+	uint64 Offset;
+	uint32 Bucket;
+	TRefCountPtr<FD3D12Resource> ResourceHeap;
+	uint64 FrameFence;
+	class FD3D12ResourceAllocator *Allocator;
+};
+
+// FD3D12ResourceLocation
+class FD3D12ResourceLocation : public FD3D12RefCount, public FD3D12DeviceChild
+{
+protected:
+	TRefCountPtr<FD3D12Resource> Resource;
+	TRefCountPtr<FD3D12ResourceLocation> LinkedResourceLocation;
+	uint64 Offset;
+	uint64 Padding;
+	TRefCountPtr<FD3D12ResourceBlockInfo> BlockInfo;
+	uint64 EffectiveBufferSize;	// This member is used to track the size of the buffer the app wanted because we may have increased the size behind the scene (GetDesc11() will return the wrong size).
+
+	friend class FD3D12DynamicHeapAllocator;
+	friend class FD3D12DefaultBufferPool;
+
+	virtual void InternalReleaseResource();
+
+public:
+	FD3D12ResourceLocation()
+		: Resource(nullptr)
+		, Offset(0)
+		, Padding(0)
+		, EffectiveBufferSize(0)
+		, BlockInfo(nullptr)
+        , FD3D12DeviceChild(nullptr)
+	{
+	}
+
+	FD3D12ResourceLocation(FD3D12Device* InParent, const uint64 &InEffectiveSize)
+		: Resource(nullptr)
+		, Offset(0)
+		, Padding(0)
+		, EffectiveBufferSize(InEffectiveSize)
+		, BlockInfo(nullptr)
+        , FD3D12DeviceChild(InParent)
+	{
+	}
+
+    FD3D12ResourceLocation(FD3D12Device* InParent, FD3D12Resource *InResource, uint64 InOffset = 0)
+		: Resource(InResource)
+		, Offset(InOffset)
+		, Padding(0)
+		, BlockInfo(nullptr)
+        , FD3D12DeviceChild(InParent)
+	{
+		check(nullptr == BlockInfo);
+		BlockInfo = nullptr;
+
+		if (InResource)
+		{
+			D3D12_RESOURCE_DESC const& desc = InResource->GetDesc();
+			EffectiveBufferSize = (uint32)desc.Width;
+		}
+		else
+		{
+			EffectiveBufferSize = 0;
+		}
+	}
+
+	virtual ~FD3D12ResourceLocation();
+
+	void LinkToResourceLocation(FD3D12ResourceLocation *InLinkedResource)
+	{
+		LinkedResourceLocation = InLinkedResource;
+	}
+
+	void SetBlockInfo(FD3D12ResourceBlockInfo* InBlockInfo, FD3D12StateCache* StateCache)
+	{
+		const FD3D12Resource *OldResource = Resource;
+		FD3D12Resource *NewResource = InBlockInfo->ResourceHeap;
+		BlockInfo = InBlockInfo;
+
+		uint64 OldOffset = Offset;
+		Offset = InBlockInfo->Offset;
+
+		if (OldResource != NewResource)
+		{
+			Resource = NewResource;
+			if (OldResource)
+			{
+                if (StateCache != nullptr)
+				    UpdateStateCache(*StateCache);
+			}
+		}
+		else
+		{
+			if (OldOffset != Offset)
+			{
+                if (StateCache != nullptr)
+                    UpdateStateCache(*StateCache);
+            }
+		}
+	}
+
+	void SetBlockInfoNoUpdate(FD3D12ResourceBlockInfo *InBlockInfo)
+	{
+		Resource = InBlockInfo->ResourceHeap;
+		BlockInfo = InBlockInfo;
+		Offset = InBlockInfo->Offset;
+	}
+
+	FD3D12ResourceBlockInfo *GetBlockInfo()
+	{
+		return BlockInfo.GetReference();
+	}
+
+	FD3D12Resource *GetResource() const
+	{
+		check(nullptr == BlockInfo || Resource == BlockInfo->ResourceHeap);
+		return Resource;
+	}
+
+	uint64 GetOffset() const
+	{
+		return Offset + Padding;
+	}
+
+	void SetPadding(uint64 InPadding)
+	{
+		Padding = InPadding;
+	}
+
+	void SetFromD3DResource(FD3D12Resource *InResource, uint64 InOffset, uint32 InEffectiveBufferSize);
+
+	void SetEffectiveBufferSize(uint64 InSize)
+	{
+		EffectiveBufferSize = InSize;
+	}
+
+	uint64 GetEffectiveBufferSize() const
+	{
+		return EffectiveBufferSize;
+	}
+
+	void Clear();
+
+	void UpdateStateCache(FD3D12StateCache& StateCache);
+    void UpdateDefaultStateCache();
+};
+
+template<typename _Type>
+struct TFencedObject
+{
+    typedef _Type FencedObjectType;
+    TFencedObject()
+    {
+    }
+
+	TFencedObject(const FencedObjectType &InObject, FD3D12CommandListHandle InCommandList)
+            : Object(InObject)
+			, SyncPoint(InCommandList)
+	{
+	}
+
+    FencedObjectType Object;
+	FD3D12CLSyncPoint SyncPoint;
+};
+
+template< typename _Type>
+class TFencedObjectPool : public FD3D12DeviceChild
+{
+public:
+	typedef _Type FencedObjectType;
+
+private:
+    TDoubleLinkedList<TFencedObject<_Type>> Pool;
+
+public:
+    void AddFencedObject(const FencedObjectType &Object, FD3D12CommandListHandle FenceValue)
+    {
+        // Check that fence values do not decrease
+        //check(Pool.Num() == 0 || FenceValue >= Pool.GetTail()->GetValue().FenceValue);
+
+        Pool.AddTail(TFencedObject<_Type>(Object, FenceValue));
+    }
+
+    bool TestCommandListCompletion(const FD3D12CLSyncPoint& SyncPoint, uint64 FenceOffset)
+    {
+        return GetParentDevice()->GetCommandListManager().IsFinished(SyncPoint, FenceOffset);
+    }
+
+    bool TryAllocate(FencedObjectType &Object, uint64 FenceOffset = 0)
+    {
+        if(Pool.Num() > 0)
+        {
+            auto Node = Pool.GetHead();
+            const TFencedObject<_Type> &Element = Node->GetValue();
+			if (TestCommandListCompletion(Element.SyncPoint, FenceOffset))
+            {
+                Object = Element.Object;
+                Pool.RemoveNode(Node);
+                return true;
+            }
+        }
+
+		return false;
+	}
+
+    uint32 Size()
+    {
+        return (uint32)Pool.Num();
+    }
+
+    void Clear()
+    {
+        Pool.Empty();
+    }
+
+    TFencedObjectPool(FD3D12Device* InParent) :
+        FD3D12DeviceChild(InParent){};
+};
+
+class FD3D12DeferredDeletionQueue : public FD3D12DeviceChild
+{
+	typedef TPair<FD3D12Resource*, uint64> FencedObjectType;
+	FThreadsafeQueue<FencedObjectType> DeferredReleaseQueue;
+
+public:
+	void EnqueueResource(FD3D12Resource *pResource);
+
+	bool ReleaseResources();
+
+	void Clear()
+	{
+		// Make sure the queue is empty after a clear.
+		// TODO: Perhaps just ignore the frame fence and clear all resources?
+		const bool bQueueIsEmpty = ReleaseResources();
+		check(bQueueIsEmpty);
+	}
+
+    FD3D12DeferredDeletionQueue(FD3D12Device* InParent);
+};
+
+template <>
+struct TTypeTraits<D3D12_INPUT_ELEMENT_DESC> : public TTypeTraitsBase < D3D12_INPUT_ELEMENT_DESC >
+{
+	enum { IsBytewiseComparable = true };
+};
+
+/** Convenience typedef: preallocated array of D3D11 input element descriptions. */
+typedef TArray<D3D12_INPUT_ELEMENT_DESC, TFixedAllocator<MaxVertexElementCount> > FD3D12VertexElements;
+
+/** This represents a vertex declaration that hasn't been combined with a specific shader to create a bound shader. */
+class FD3D12VertexDeclaration : public FRHIVertexDeclaration
+{
+public:
+	/** Elements of the vertex declaration. */
+	FD3D12VertexElements VertexElements;
+
+	/** Initialization constructor. */
+	explicit FD3D12VertexDeclaration(const FD3D12VertexElements& InElements)
+		: VertexElements(InElements)
+	{
+	}
+};
+
+/** This represents a vertex shader that hasn't been combined with a specific declaration to create a bound shader. */
+class FD3D12VertexShader : public FRHIVertexShader
+{
+public:
+	enum { StaticFrequency = SF_Vertex };
+
+	/** The shader's bytecode. */
+	FD3D12ShaderBytecode ShaderBytecode;
+
+	FD3D12ShaderResourceTable ShaderResourceTable;
+
+	/** The vertex shader's bytecode, with custom data in the last byte. */
+	TArray<uint8> Code;
+
+	// TEMP remove with removal of bound shader state
+	int32 Offset;
+
+	bool bShaderNeedsGlobalConstantBuffer;
+	uint32 SRVCount;
+	uint32 CBCount;
+	uint32 SamplerCount;
+	uint32 UAVCount;
+};
+
+class FD3D12GeometryShader : public FRHIGeometryShader
+{
+public:
+	enum { StaticFrequency = SF_Geometry };
+
+	/** The shader's bytecode. */
+	FD3D12ShaderBytecode ShaderBytecode;
+
+	FD3D12ShaderResourceTable ShaderResourceTable;
+
+	/** The shader's bytecode, with custom data in the last byte. */
+	TArray<uint8> Code;
+
+	/** The shader's stream output description. */
+	D3D12_STREAM_OUTPUT_DESC StreamOutput;
+	D3D12_SO_DECLARATION_ENTRY *pStreamOutEntries;
+	uint32 *pStreamOutStrides;
+
+	bool bShaderNeedsGlobalConstantBuffer;
+	bool bShaderNeedsStreamOutput;
+	uint32 SRVCount;
+	uint32 CBCount;
+	uint32 SamplerCount;
+	uint32 UAVCount;
+
+	FD3D12GeometryShader()
+		: bShaderNeedsStreamOutput(false)
+		, pStreamOutEntries(nullptr)
+		, pStreamOutStrides(nullptr)
+	{
+		FMemory::Memzero(&StreamOutput, sizeof(StreamOutput));
+	}
+
+	virtual ~FD3D12GeometryShader()
+	{
+		if (pStreamOutEntries)
+		{
+			delete[] pStreamOutEntries;
+			pStreamOutEntries = nullptr;
+		}
+
+		if (pStreamOutStrides)
+		{
+			delete[] pStreamOutStrides;
+			pStreamOutStrides = nullptr;
+		}
+	}
+};
+
+class FD3D12HullShader : public FRHIHullShader
+{
+public:
+	enum { StaticFrequency = SF_Hull };
+
+	/** The shader's bytecode. */
+	FD3D12ShaderBytecode ShaderBytecode;
+
+	FD3D12ShaderResourceTable ShaderResourceTable;
+
+	/** The shader's bytecode, with custom data in the last byte. */
+	TArray<uint8> Code;
+
+	bool bShaderNeedsGlobalConstantBuffer;
+	uint32 SRVCount;
+	uint32 CBCount;
+	uint32 SamplerCount;
+	uint32 UAVCount;
+};
+
+class FD3D12DomainShader : public FRHIDomainShader
+{
+public:
+	enum { StaticFrequency = SF_Domain };
+
+	/** The shader's bytecode. */
+	FD3D12ShaderBytecode ShaderBytecode;
+
+	FD3D12ShaderResourceTable ShaderResourceTable;
+
+	/** The shader's bytecode, with custom data in the last byte. */
+	TArray<uint8> Code;
+
+	bool bShaderNeedsGlobalConstantBuffer;
+	uint32 SRVCount;
+	uint32 CBCount;
+	uint32 SamplerCount;
+	uint32 UAVCount;
+};
+
+class FD3D12PixelShader : public FRHIPixelShader
+{
+public:
+	enum { StaticFrequency = SF_Pixel };
+
+	/** The shader's bytecode. */
+	FD3D12ShaderBytecode ShaderBytecode;
+
+	/** The shader's bytecode, with custom data in the last byte. */
+	TArray<uint8> Code;
+
+	FD3D12ShaderResourceTable ShaderResourceTable;
+
+	bool bShaderNeedsGlobalConstantBuffer;
+	uint32 SRVCount;
+	uint32 CBCount;
+	uint32 SamplerCount;
+	uint32 UAVCount;
+};
+
+class FD3D12ComputeShader : public FRHIComputeShader
+{
+public:
+	enum { StaticFrequency = SF_Compute };
+
+	/** The shader's bytecode. */
+	FD3D12ShaderBytecode ShaderBytecode;
+
+	/** The shader's bytecode, with custom data in the last byte. */
+	TArray<uint8> Code;
+
+	FD3D12ShaderResourceTable ShaderResourceTable;
+
+	bool bShaderNeedsGlobalConstantBuffer;
+	uint32 SRVCount;
+	uint32 CBCount;
+	uint32 SamplerCount;
+	uint32 UAVCount;
+};
+
+/**
+* Combined shader state and vertex definition for rendering geometry.
+* Each unique instance consists of a vertex decl, vertex shader, and pixel shader.
+*/
+class FD3D12BoundShaderState : public FRHIBoundShaderState
+{
+public:
+
+#if D3D12_SUPPORTS_PARALLEL_RHI_EXECUTE
+	FCachedBoundShaderStateLink_Threadsafe CacheLink;
+#else
+	FCachedBoundShaderStateLink CacheLink;
+#endif
+
+	D3D12_INPUT_LAYOUT_DESC InputLayout;
+
+	bool bShaderNeedsGlobalConstantBuffer[SF_NumFrequencies];
+    uint64 UniqueID;
+
+
+	/** Initialization constructor. */
+	FD3D12BoundShaderState(
+		FVertexDeclarationRHIParamRef InVertexDeclarationRHI,
+		FVertexShaderRHIParamRef InVertexShaderRHI,
+		FPixelShaderRHIParamRef InPixelShaderRHI,
+		FHullShaderRHIParamRef InHullShaderRHI,
+		FDomainShaderRHIParamRef InDomainShaderRHI,
+		FGeometryShaderRHIParamRef InGeometryShaderRHI
+		);
+
+	virtual ~FD3D12BoundShaderState();
+
+	/**
+	* Get the shader for the given frequency.
+	*/
+	FORCEINLINE FD3D12VertexShader*   GetVertexShader() const   { return (FD3D12VertexShader*)CacheLink.GetVertexShader(); }
+	FORCEINLINE FD3D12PixelShader*    GetPixelShader() const    { return (FD3D12PixelShader*)CacheLink.GetPixelShader(); }
+	FORCEINLINE FD3D12HullShader*     GetHullShader() const     { return (FD3D12HullShader*)CacheLink.GetHullShader(); }
+	FORCEINLINE FD3D12DomainShader*   GetDomainShader() const   { return (FD3D12DomainShader*)CacheLink.GetDomainShader(); }
+	FORCEINLINE FD3D12GeometryShader* GetGeometryShader() const { return (FD3D12GeometryShader*)CacheLink.GetGeometryShader(); }
+};
+
+/** The base class of resources that may be bound as shader resources. */
+class FD3D12BaseShaderResource : public FD3D12DeviceChild, public IRefCountedObject
+{
+public:
+	TRefCountPtr<FD3D12ResourceLocation> ResourceLocation;
+
+public:
+    FD3D12BaseShaderResource(FD3D12ResourceLocation *InResourceLocation, FD3D12Device* InParent) :
+        FD3D12DeviceChild(InParent)
+	{
+		ResourceLocation = InResourceLocation;
+	}
+
+	FD3D12BaseShaderResource(FD3D12Resource *InResource, FD3D12Device* InParent) :
+        FD3D12DeviceChild(InParent)
+	{
+		ResourceLocation = new FD3D12ResourceLocation(GetParentDevice(),InResource);
+	}
+};
+
+enum ViewSubresourceSubsetFlags
+{
+	ViewSubresourceSubsetFlags_None = 0x0,
+	ViewSubresourceSubsetFlags_DepthOnlyDsv = 0x1,
+	ViewSubresourceSubsetFlags_StencilOnlyDsv = 0x2,
+	ViewSubresourceSubsetFlags_DepthAndStencilDsv = (ViewSubresourceSubsetFlags_DepthOnlyDsv | ViewSubresourceSubsetFlags_StencilOnlyDsv),
+};
+
+/** Class to track subresources in a view */
+struct CBufferView {};
+class CSubresourceSubset
+{
+public:
+	CSubresourceSubset() {}
+	inline explicit CSubresourceSubset(const CBufferView&) :
+		m_BeginArray(0),
+		m_EndArray(1),
+		m_BeginMip(0),
+		m_EndMip(1),
+		m_BeginPlane(0),
+		m_EndPlane(1)
+	{
+	}
+	inline explicit CSubresourceSubset(const D3D12_SHADER_RESOURCE_VIEW_DESC& Desc, DXGI_FORMAT ResourceFormat) :
+		m_BeginArray(0),
+		m_EndArray(1),
+		m_BeginMip(0),
+		m_EndMip(1),
+		m_BeginPlane(0),
+		m_EndPlane(1)
+	{
+		switch (Desc.ViewDimension)
+		{
+		default: ASSUME(0 && "Corrupt Resource Type on Shader Resource View"); break;
+
+		case (D3D12_SRV_DIMENSION_BUFFER) :
+			break;
+
+		case (D3D12_SRV_DIMENSION_TEXTURE1D) :
+			m_BeginMip = uint8(Desc.Texture1D.MostDetailedMip);
+			m_EndMip = uint8(m_BeginMip + Desc.Texture1D.MipLevels);
+			m_BeginPlane = GetPlaneSliceFromViewFormat(ResourceFormat, Desc.Format);
+			m_EndPlane = m_BeginPlane + 1;
+			break;
+
+		case (D3D12_SRV_DIMENSION_TEXTURE1DARRAY) :
+			m_BeginArray = uint16(Desc.Texture1DArray.FirstArraySlice);
+			m_EndArray = uint16(m_BeginArray + Desc.Texture1DArray.ArraySize);
+			m_BeginMip = uint8(Desc.Texture1DArray.MostDetailedMip);
+			m_EndMip = uint8(m_BeginMip + Desc.Texture1DArray.MipLevels);
+			m_BeginPlane = GetPlaneSliceFromViewFormat(ResourceFormat, Desc.Format);
+			m_EndPlane = m_BeginPlane + 1;
+			break;
+
+		case (D3D12_SRV_DIMENSION_TEXTURE2D) :
+			m_BeginMip = uint8(Desc.Texture2D.MostDetailedMip);
+			m_EndMip = uint8(m_BeginMip + Desc.Texture2D.MipLevels);
+			m_BeginPlane = uint8(Desc.Texture2D.PlaneSlice);
+			m_EndPlane = uint8(Desc.Texture2D.PlaneSlice + 1);
+			break;
+
+		case (D3D12_SRV_DIMENSION_TEXTURE2DARRAY) :
+			m_BeginArray = uint16(Desc.Texture2DArray.FirstArraySlice);
+			m_EndArray = uint16(m_BeginArray + Desc.Texture2DArray.ArraySize);
+			m_BeginMip = uint8(Desc.Texture2DArray.MostDetailedMip);
+			m_EndMip = uint8(m_BeginMip + Desc.Texture2DArray.MipLevels);
+			m_BeginPlane = uint8(Desc.Texture2DArray.PlaneSlice);
+			m_EndPlane = uint8(Desc.Texture2DArray.PlaneSlice + 1);
+			break;
+
+		case (D3D12_SRV_DIMENSION_TEXTURE2DMS) :
+			m_BeginPlane = GetPlaneSliceFromViewFormat(ResourceFormat, Desc.Format);
+			m_EndPlane = m_BeginPlane + 1;
+			break;
+
+		case (D3D12_SRV_DIMENSION_TEXTURE2DMSARRAY) :
+			m_BeginArray = uint16(Desc.Texture2DMSArray.FirstArraySlice);
+			m_EndArray = uint16(m_BeginArray + Desc.Texture2DMSArray.ArraySize);
+			m_BeginPlane = GetPlaneSliceFromViewFormat(ResourceFormat, Desc.Format);
+			m_EndPlane = m_BeginPlane + 1;
+			break;
+
+		case (D3D12_SRV_DIMENSION_TEXTURE3D) :
+			m_EndArray = uint16(-1); //all slices
+			m_BeginMip = uint8(Desc.Texture3D.MostDetailedMip);
+			m_EndMip = uint8(m_BeginMip + Desc.Texture3D.MipLevels);
+			break;
+
+		case (D3D12_SRV_DIMENSION_TEXTURECUBE) :
+			m_BeginMip = uint8(Desc.TextureCube.MostDetailedMip);
+			m_EndMip = uint8(m_BeginMip + Desc.TextureCube.MipLevels);
+			m_BeginArray = 0;
+			m_EndArray = 6;
+			m_BeginPlane = GetPlaneSliceFromViewFormat(ResourceFormat, Desc.Format);
+			m_EndPlane = m_BeginPlane + 1;
+			break;
+
+		case (D3D12_SRV_DIMENSION_TEXTURECUBEARRAY) :
+			m_BeginArray = uint16(Desc.TextureCubeArray.First2DArrayFace);
+			m_EndArray = uint16(m_BeginArray + Desc.TextureCubeArray.NumCubes * 6);
+			m_BeginMip = uint8(Desc.TextureCubeArray.MostDetailedMip);
+			m_EndMip = uint8(m_BeginMip + Desc.TextureCubeArray.MipLevels);
+			m_BeginPlane = GetPlaneSliceFromViewFormat(ResourceFormat, Desc.Format);
+			m_EndPlane = m_BeginPlane + 1;
+			break;
+		}
+	}
+	inline explicit CSubresourceSubset(const D3D12_UNORDERED_ACCESS_VIEW_DESC& Desc) :
+		m_BeginArray(0),
+		m_EndArray(1),
+		m_BeginMip(0),
+		m_BeginPlane(0),
+		m_EndPlane(1)
+	{
+		switch (Desc.ViewDimension)
+		{
+		default: ASSUME(0 && "Corrupt Resource Type on Unordered Access View"); break;
+
+		case (D3D12_UAV_DIMENSION_BUFFER) : break;
+
+		case (D3D12_UAV_DIMENSION_TEXTURE1D) :
+			m_BeginMip = uint8(Desc.Texture1D.MipSlice);
+			break;
+
+		case (D3D12_UAV_DIMENSION_TEXTURE1DARRAY) :
+			m_BeginArray = uint16(Desc.Texture1DArray.FirstArraySlice);
+			m_EndArray = uint16(m_BeginArray + Desc.Texture1DArray.ArraySize);
+			m_BeginMip = uint8(Desc.Texture1DArray.MipSlice);
+			break;
+
+		case (D3D12_UAV_DIMENSION_TEXTURE2D) :
+			m_BeginMip = uint8(Desc.Texture2D.MipSlice);
+			m_BeginPlane = uint8(Desc.Texture2D.PlaneSlice);
+			m_EndPlane = uint8(Desc.Texture2D.PlaneSlice + 1);
+			break;
+
+		case (D3D12_UAV_DIMENSION_TEXTURE2DARRAY) :
+			m_BeginArray = uint16(Desc.Texture2DArray.FirstArraySlice);
+			m_EndArray = uint16(m_BeginArray + Desc.Texture2DArray.ArraySize);
+			m_BeginMip = uint8(Desc.Texture2DArray.MipSlice);
+			m_BeginPlane = uint8(Desc.Texture2DArray.PlaneSlice);
+			m_EndPlane = uint8(Desc.Texture2DArray.PlaneSlice + 1);
+			break;
+
+		case (D3D12_UAV_DIMENSION_TEXTURE3D) :
+			m_BeginArray = uint16(Desc.Texture3D.FirstWSlice);
+			m_EndArray = uint16(m_BeginArray + Desc.Texture3D.WSize);
+			m_BeginMip = uint8(Desc.Texture3D.MipSlice);
+			break;
+		}
+
+		m_EndMip = m_BeginMip + 1;
+	}
+	inline explicit CSubresourceSubset(const D3D12_RENDER_TARGET_VIEW_DESC& Desc) :
+		m_BeginArray(0),
+		m_EndArray(1),
+		m_BeginMip(0),
+		m_BeginPlane(0),
+		m_EndPlane(1)
+	{
+		switch (Desc.ViewDimension)
+		{
+		default: ASSUME(0 && "Corrupt Resource Type on Render Target View"); break;
+
+		case (D3D12_RTV_DIMENSION_BUFFER) : break;
+
+		case (D3D12_RTV_DIMENSION_TEXTURE1D) :
+			m_BeginMip = uint8(Desc.Texture1D.MipSlice);
+			break;
+
+		case (D3D12_RTV_DIMENSION_TEXTURE1DARRAY) :
+			m_BeginArray = uint16(Desc.Texture1DArray.FirstArraySlice);
+			m_EndArray = uint16(m_BeginArray + Desc.Texture1DArray.ArraySize);
+			m_BeginMip = uint8(Desc.Texture1DArray.MipSlice);
+			break;
+
+		case (D3D12_RTV_DIMENSION_TEXTURE2D) :
+			m_BeginMip = uint8(Desc.Texture2D.MipSlice);
+			m_BeginPlane = uint8(Desc.Texture2D.PlaneSlice);
+			m_EndPlane = uint8(Desc.Texture2D.PlaneSlice + 1);
+			break;
+
+		case (D3D12_RTV_DIMENSION_TEXTURE2DMS) : break;
+
+		case (D3D12_RTV_DIMENSION_TEXTURE2DARRAY) :
+			m_BeginArray = uint16(Desc.Texture2DArray.FirstArraySlice);
+			m_EndArray = uint16(m_BeginArray + Desc.Texture2DArray.ArraySize);
+			m_BeginMip = uint8(Desc.Texture2DArray.MipSlice);
+			m_BeginPlane = uint8(Desc.Texture2DArray.PlaneSlice);
+			m_EndPlane = uint8(Desc.Texture2DArray.PlaneSlice + 1);
+			break;
+
+		case (D3D12_RTV_DIMENSION_TEXTURE2DMSARRAY) :
+			m_BeginArray = uint16(Desc.Texture2DMSArray.FirstArraySlice);
+			m_EndArray = uint16(m_BeginArray + Desc.Texture2DMSArray.ArraySize);
+			break;
+
+		case (D3D12_RTV_DIMENSION_TEXTURE3D) :
+			m_BeginArray = uint16(Desc.Texture3D.FirstWSlice);
+			m_EndArray = uint16(m_BeginArray + Desc.Texture3D.WSize);
+			m_BeginMip = uint8(Desc.Texture3D.MipSlice);
+			break;
+		}
+
+		m_EndMip = m_BeginMip + 1;
+	}
+
+	inline explicit CSubresourceSubset(const D3D12_DEPTH_STENCIL_VIEW_DESC& Desc, DXGI_FORMAT ResourceFormat, ViewSubresourceSubsetFlags Flags) :
+		m_BeginArray(0),
+		m_EndArray(1),
+		m_BeginMip(0),
+		m_BeginPlane(0),
+		m_EndPlane(GetPlaneCount(ResourceFormat))
+	{
+		switch (Desc.ViewDimension)
+		{
+		default: ASSUME(0 && "Corrupt Resource Type on Depth Stencil View"); break;
+
+		case (D3D12_DSV_DIMENSION_TEXTURE1D) :
+			m_BeginMip = uint8(Desc.Texture1D.MipSlice);
+			break;
+
+		case (D3D12_DSV_DIMENSION_TEXTURE1DARRAY) :
+			m_BeginArray = uint16(Desc.Texture1DArray.FirstArraySlice);
+			m_EndArray = uint16(m_BeginArray + Desc.Texture1DArray.ArraySize);
+			m_BeginMip = uint8(Desc.Texture1DArray.MipSlice);
+			break;
+
+		case (D3D12_DSV_DIMENSION_TEXTURE2D) :
+			m_BeginMip = uint8(Desc.Texture2D.MipSlice);
+			break;
+
+		case (D3D12_DSV_DIMENSION_TEXTURE2DMS) : break;
+
+		case (D3D12_DSV_DIMENSION_TEXTURE2DARRAY) :
+			m_BeginArray = uint16(Desc.Texture2DArray.FirstArraySlice);
+			m_EndArray = uint16(m_BeginArray + Desc.Texture2DArray.ArraySize);
+			m_BeginMip = uint8(Desc.Texture2DArray.MipSlice);
+			break;
+
+		case (D3D12_DSV_DIMENSION_TEXTURE2DMSARRAY) :
+			m_BeginArray = uint16(Desc.Texture2DMSArray.FirstArraySlice);
+			m_EndArray = uint16(m_BeginArray + Desc.Texture2DMSArray.ArraySize);
+			break;
+		}
+
+		m_EndMip = m_BeginMip + 1;
+
+		if (m_EndPlane == 2)
+		{
+			if ((Flags & ViewSubresourceSubsetFlags_DepthAndStencilDsv) != ViewSubresourceSubsetFlags_DepthAndStencilDsv)
+			{
+				if (Flags & ViewSubresourceSubsetFlags_DepthOnlyDsv)
+				{
+					m_BeginPlane = 0;
+					m_EndPlane = 1;
+				}
+				else if (Flags & ViewSubresourceSubsetFlags_StencilOnlyDsv)
+				{
+					m_BeginPlane = 1;
+					m_EndPlane = 2;
+				}
+			}
+		}
+	}
+
+	__forceinline bool DoesNotOverlap(const CSubresourceSubset& other) const
+	{
+		if (m_EndArray <= other.m_BeginArray)
+		{
+			return true;
+		}
+
+		if (other.m_EndArray <= m_BeginArray)
+		{
+			return true;
+		}
+
+		if (m_EndMip <= other.m_BeginMip)
+		{
+			return true;
+		}
+
+		if (other.m_EndMip <= m_BeginMip)
+		{
+			return true;
+		}
+
+		if (m_EndPlane <= other.m_BeginPlane)
+		{
+			return true;
+		}
+
+		if (other.m_EndPlane <= m_BeginPlane)
+		{
+			return true;
+		}
+
+		return false;
+	}
+
+protected:
+	uint16 m_BeginArray; // Also used to store Tex3D slices.
+	uint16 m_EndArray; // End - Begin == Array Slices
+	uint8 m_BeginMip;
+	uint8 m_EndMip; // End - Begin == Mip Levels
+	uint8 m_BeginPlane;
+	uint8 m_EndPlane;
+};
+
+template <typename TDesc>
+class FD3D12View;
+
+class CViewSubresourceSubset : public CSubresourceSubset
+{
+	friend class FD3D12View < D3D12_SHADER_RESOURCE_VIEW_DESC > ;
+	friend class FD3D12View < D3D12_RENDER_TARGET_VIEW_DESC > ;
+	friend class FD3D12View < D3D12_DEPTH_STENCIL_VIEW_DESC > ;
+	friend class FD3D12View < D3D12_UNORDERED_ACCESS_VIEW_DESC > ;
+
+public:
+	CViewSubresourceSubset() {}
+	inline explicit CViewSubresourceSubset(const CBufferView&)
+		: CSubresourceSubset(CBufferView())
+		, m_MipLevels(1)
+		, m_ArraySlices(1)
+		, m_MostDetailedMip(0)
+		, m_ViewArraySize(1)
+	{
+	}
+
+	inline CViewSubresourceSubset(uint32 Subresource, uint8 MipLevels, uint16 ArraySize, uint8 PlaneCount)
+		: m_MipLevels(MipLevels)
+		, m_ArraySlices(ArraySize)
+		, m_PlaneCount(PlaneCount)
+	{
+		if (Subresource < uint32(MipLevels) * uint32(ArraySize))
+		{
+			m_BeginArray = Subresource / MipLevels;
+			m_EndArray = m_BeginArray + 1;
+			m_BeginMip = Subresource % MipLevels;
+			m_EndMip = m_EndArray + 1;
+		}
+		else
+		{
+			m_BeginArray = 0;
+			m_BeginMip = 0;
+			if (Subresource == D3D12_RESOURCE_BARRIER_ALL_SUBRESOURCES)
+			{
+				m_EndArray = ArraySize;
+				m_EndMip = MipLevels;
+			}
+			else
+			{
+				m_EndArray = 0;
+				m_EndMip = 0;
+			}
+		}
+		m_MostDetailedMip = m_BeginMip;
+		m_ViewArraySize = m_EndArray - m_BeginArray;
+	}
+
+	inline CViewSubresourceSubset(const D3D12_SHADER_RESOURCE_VIEW_DESC& Desc, uint8 MipLevels, uint16 ArraySize, DXGI_FORMAT ResourceFormat, ViewSubresourceSubsetFlags /*Flags*/)
+		: CSubresourceSubset(Desc, ResourceFormat)
+		, m_MipLevels(MipLevels)
+		, m_ArraySlices(ArraySize)
+		, m_PlaneCount(GetPlaneCount(ResourceFormat))
+	{
+		if (Desc.ViewDimension == D3D12_SRV_DIMENSION_TEXTURE3D)
+		{
+			check(m_BeginArray == 0);
+			m_EndArray = 1;
+		}
+		m_MostDetailedMip = m_BeginMip;
+		m_ViewArraySize = m_EndArray - m_BeginArray;
+		Reduce();
+	}
+	
+	inline CViewSubresourceSubset(const D3D12_UNORDERED_ACCESS_VIEW_DESC& Desc, uint8 MipLevels, uint16 ArraySize, DXGI_FORMAT ResourceFormat, ViewSubresourceSubsetFlags /*Flags*/)
+		: CSubresourceSubset(Desc)
+		, m_MipLevels(MipLevels)
+		, m_ArraySlices(ArraySize)
+		, m_PlaneCount(GetPlaneCount(ResourceFormat))
+	{
+		if (Desc.ViewDimension == D3D12_UAV_DIMENSION_TEXTURE3D)
+		{
+			m_BeginArray = 0;
+			m_EndArray = 1;
+		}
+		m_MostDetailedMip = m_BeginMip;
+		m_ViewArraySize = m_EndArray - m_BeginArray;
+		Reduce();
+	}
+	
+	inline CViewSubresourceSubset(const D3D12_DEPTH_STENCIL_VIEW_DESC& Desc, uint8 MipLevels, uint16 ArraySize, DXGI_FORMAT ResourceFormat, ViewSubresourceSubsetFlags Flags)
+		: CSubresourceSubset(Desc, ResourceFormat, Flags)
+		, m_MipLevels(MipLevels)
+		, m_ArraySlices(ArraySize)
+		, m_PlaneCount(GetPlaneCount(ResourceFormat))
+	{
+		m_MostDetailedMip = m_BeginMip;
+		m_ViewArraySize = m_EndArray - m_BeginArray;
+		Reduce();
+	}
+	
+	inline CViewSubresourceSubset(const D3D12_RENDER_TARGET_VIEW_DESC& Desc, uint8 MipLevels, uint16 ArraySize, DXGI_FORMAT ResourceFormat, ViewSubresourceSubsetFlags /*Flags*/)
+		: CSubresourceSubset(Desc)
+		, m_MipLevels(MipLevels)
+		, m_ArraySlices(ArraySize)
+		, m_PlaneCount(GetPlaneCount(ResourceFormat))
+	{
+		if (Desc.ViewDimension == D3D12_RTV_DIMENSION_TEXTURE3D)
+		{
+			m_BeginArray = 0;
+			m_EndArray = 1;
+		}
+		m_MostDetailedMip = m_BeginMip;
+		m_ViewArraySize = m_EndArray - m_BeginArray;
+		Reduce();
+	}
+
+	template<typename T>
+	static CViewSubresourceSubset FromView(const T* pView)
+	{
+		return CViewSubresourceSubset(
+			pView->Desc(),
+			static_cast<uint8>(pView->GetResource()->GetMipLevels()),
+			static_cast<uint16>(pView->GetResource()->GetArraySize()),
+			static_cast<uint8>(pView->GetResource()->GetPlaneCount())
+			);
+	}
+
+public:
+	class CViewSubresourceIterator;
+
+public:
+	CViewSubresourceIterator begin() const;
+	CViewSubresourceIterator end() const;
+	bool IsWholeResource() const;
+	uint32 ArraySize() const;
+
+	uint8 MostDetailedMip() const;
+	uint16 ViewArraySize() const;
+
+	uint32 MinSubresource() const;
+	uint32 MaxSubresource() const;
+
+private:
+	// Strictly for performance, allows coalescing contiguous subresource ranges into a single range
+	inline void Reduce()
+	{
+		if (m_BeginMip == 0
+			&& m_EndMip == m_MipLevels
+			&& m_BeginArray == 0
+			&& m_EndArray == m_ArraySlices
+			&& m_BeginPlane == 0
+			&& m_EndPlane == m_PlaneCount)
+		{
+			uint32 startSubresource = D3D12CalcSubresource(0, 0, m_BeginPlane, m_MipLevels, m_ArraySlices);
+			uint32 endSubresource = D3D12CalcSubresource(0, 0, m_EndPlane, m_MipLevels, m_ArraySlices);
+
+			// Only coalesce if the full-resolution UINTs fit in the UINT8s used for storage here
+			if (endSubresource < static_cast<uint8>(-1))
+			{
+				m_BeginArray = 0;
+				m_EndArray = 1;
+				m_BeginPlane = 0;
+				m_EndPlane = 1;
+				m_BeginMip = static_cast<uint8>(startSubresource);
+				m_EndMip = static_cast<uint8>(endSubresource);
+			}
+		}
+	}
+
+protected:
+	uint8 m_MipLevels;
+	uint16 m_ArraySlices;
+	uint8 m_PlaneCount;
+	uint8 m_MostDetailedMip;
+	uint16 m_ViewArraySize;
+};
+
+// This iterator iterates over contiguous ranges of subresources within a subresource subset. eg:
+//
+// // For each contiguous subresource range.
+// for( CViewSubresourceIterator it = ViewSubset.begin(); it != ViewSubset.end(); ++it )
+// {
+//      // StartSubresource and EndSubresource members of the iterator describe the contiguous range.
+//      for( uint32 SubresourceIndex = it.StartSubresource(); SubresourceIndex < it.EndSubresource(); SubresourceIndex++ )
+//      {
+//          // Action for each subresource within the current range.
+//      }
+//  }
+//
+class CViewSubresourceSubset::CViewSubresourceIterator
+{
+public:
+	inline CViewSubresourceIterator(CViewSubresourceSubset const& SubresourceSet, uint16 ArraySlice, uint8 PlaneSlice)
+		: m_Subresources(SubresourceSet)
+		, m_CurrentArraySlice(ArraySlice)
+		, m_CurrentPlaneSlice(PlaneSlice)
+	{
+	}
+
+	inline CViewSubresourceSubset::CViewSubresourceIterator& operator++()
+	{
+		check(m_CurrentArraySlice < m_Subresources.m_EndArray);
+
+		if (++m_CurrentArraySlice >= m_Subresources.m_EndArray)
+		{
+			check(m_CurrentPlaneSlice < m_Subresources.m_EndPlane);
+			m_CurrentArraySlice = m_Subresources.m_BeginArray;
+			++m_CurrentPlaneSlice;
+		}
+
+		return *this;
+	}
+
+	inline CViewSubresourceSubset::CViewSubresourceIterator& operator--()
+	{
+		if (m_CurrentArraySlice <= m_Subresources.m_BeginArray)
+		{
+			m_CurrentArraySlice = m_Subresources.m_EndArray;
+
+			check(m_CurrentPlaneSlice > m_Subresources.m_BeginPlane);
+			--m_CurrentPlaneSlice;
+		}
+
+		--m_CurrentArraySlice;
+
+		return *this;
+	}
+
+	inline bool operator==(CViewSubresourceIterator const& other) const
+	{
+		return &other.m_Subresources == &m_Subresources
+			&& other.m_CurrentArraySlice == m_CurrentArraySlice
+			&& other.m_CurrentPlaneSlice == m_CurrentPlaneSlice;
+	}
+
+	inline bool operator!=(CViewSubresourceIterator const& other) const
+	{
+		return !(other == *this);
+	}
+
+	inline uint32 StartSubresource() const
+	{
+		return D3D12CalcSubresource(m_Subresources.m_BeginMip, m_CurrentArraySlice, m_CurrentPlaneSlice, m_Subresources.m_MipLevels, m_Subresources.m_ArraySlices);
+	}
+
+	inline uint32 EndSubresource() const
+	{
+		return D3D12CalcSubresource(m_Subresources.m_EndMip, m_CurrentArraySlice, m_CurrentPlaneSlice, m_Subresources.m_MipLevels, m_Subresources.m_ArraySlices);
+	}
+
+	inline TPair<uint32, uint32> operator*() const
+	{
+		TPair<uint32, uint32> NewPair;
+		NewPair.Key = StartSubresource();
+		NewPair.Value = EndSubresource();
+		return NewPair;
+	}
+
+private:
+	CViewSubresourceSubset const& m_Subresources;
+	uint16 m_CurrentArraySlice;
+	uint8 m_CurrentPlaneSlice;
+};
+
+inline CViewSubresourceSubset::CViewSubresourceIterator CViewSubresourceSubset::begin() const
+{
+	return CViewSubresourceIterator(*this, m_BeginArray, m_BeginPlane);
+}
+
+inline CViewSubresourceSubset::CViewSubresourceIterator CViewSubresourceSubset::end() const
+{
+	return CViewSubresourceIterator(*this, m_BeginArray, m_EndPlane);
+}
+
+inline bool CViewSubresourceSubset::IsWholeResource() const
+{
+	return m_BeginMip == 0 && m_BeginArray == 0 && m_BeginPlane == 0 && (m_EndMip * m_EndArray * m_EndPlane == m_MipLevels * m_ArraySlices * m_PlaneCount);
+}
+
+inline uint32 CViewSubresourceSubset::ArraySize() const
+{
+	return m_ArraySlices;
+}
+
+inline uint8 CViewSubresourceSubset::MostDetailedMip() const
+{
+	return m_MostDetailedMip;
+}
+
+inline uint16 CViewSubresourceSubset::ViewArraySize() const
+{
+	return m_ViewArraySize;
+}
+
+inline uint32 CViewSubresourceSubset::MinSubresource() const
+{
+	return (*begin()).Key;
+}
+
+inline uint32 CViewSubresourceSubset::MaxSubresource() const
+{
+	return (*(--end())).Value;
+}
+
+// The view is either based on a resource location or a resource; not both
+class FD3D12ViewGeneric
+{
+protected:
+	CD3DX12_CPU_DESCRIPTOR_HANDLE Descriptor;
+	uint32 DescriptorHeapIndex;
+	TRefCountPtr<FD3D12ResourceLocation> ResourceLocation;
+	TRefCountPtr<FD3D12Resource> D3DResource;
+
+	explicit FD3D12ViewGeneric(FD3D12ResourceLocation* InResourceLocation)
+		: ResourceLocation(InResourceLocation)
+		, D3DResource(nullptr)
+	{
+		Descriptor.ptr = 0;
+		DescriptorHeapIndex = 0;
+	}
+
+	explicit FD3D12ViewGeneric(FD3D12Resource* InResource)
+		: ResourceLocation(nullptr)
+		, D3DResource(InResource)
+	{
+		Descriptor.ptr = 0;
+		DescriptorHeapIndex = 0;
+	}
+
+	void ResetFromResourceLocation(FD3D12ResourceLocation *InResourceLocation)
+	{
+		ResourceLocation = InResourceLocation;
+		D3DResource = nullptr;
+		Descriptor.ptr = 0;
+		DescriptorHeapIndex = 0;
+	}
+
+	void ResetFromD3DResource(FD3D12Resource *InResource)
+	{
+		D3DResource = InResource;
+		ResourceLocation = nullptr;
+		Descriptor.ptr = 0;
+		DescriptorHeapIndex = 0;
+	}
+
+public:
+	CD3DX12_CPU_DESCRIPTOR_HANDLE GetView() const { check(Descriptor.ptr != 0); return Descriptor; }
+	uint32 GetDescriptorHeapIndex() const { return DescriptorHeapIndex; }
+	FD3D12ResourceLocation* GetResourceLocation() const { return ResourceLocation; }
+	FD3D12Resource* GetResource() const
+	{
+		if (D3DResource)
+		{
+			return D3DResource;
+		}
+		else if (ResourceLocation)
+		{
+			return ResourceLocation->GetResource();
+		}
+
+		return nullptr;
+	}
+
+	friend class FD3D12VertexBuffer;
+};
+
+template <typename TDesc> struct TCreateViewMap;
+template<> struct TCreateViewMap < D3D12_SHADER_RESOURCE_VIEW_DESC > { static decltype(&ID3D12Device::CreateShaderResourceView) GetCreate() { return &ID3D12Device::CreateShaderResourceView; } };
+template<> struct TCreateViewMap < D3D12_RENDER_TARGET_VIEW_DESC > { static decltype(&ID3D12Device::CreateRenderTargetView) GetCreate() { return &ID3D12Device::CreateRenderTargetView; } };
+template<> struct TCreateViewMap < D3D12_DEPTH_STENCIL_VIEW_DESC > { static decltype(&ID3D12Device::CreateDepthStencilView) GetCreate() { return &ID3D12Device::CreateDepthStencilView; } };
+template<> struct TCreateViewMap < D3D12_UNORDERED_ACCESS_VIEW_DESC > { static decltype(&ID3D12Device::CreateUnorderedAccessView) GetCreate() { return &ID3D12Device::CreateUnorderedAccessView; } };
+
+template <typename TDesc>
+class FD3D12View : public FD3D12ViewGeneric, public FD3D12DeviceChild
+{
+protected:
+	CViewSubresourceSubset ViewSubresourceSubset;
+	TDesc Desc;
+
+protected:
+	explicit FD3D12View()
+		: FD3D12ViewGeneric(nullptr)
+        , FD3D12DeviceChild(nullptr)
+	{
+	}
+
+	explicit FD3D12View(FD3D12Device* InParent, const TDesc* InDesc, FD3D12ResourceLocation* InResourceLocation, ViewSubresourceSubsetFlags InFlags)
+		: FD3D12ViewGeneric(InResourceLocation)
+		, ViewSubresourceSubset(*InDesc, 
+			GetResource() ? GetResource()->GetMipLevels() : 0, 
+			GetResource() ? GetResource()->GetArraySize() : 0, 
+			GetResource() ? GetResource()->GetDesc().Format : DXGI_FORMAT_UNKNOWN,
+			InFlags)
+        , FD3D12DeviceChild(InParent)
+	{
+		Init(InDesc);
+	}
+
+	explicit FD3D12View(FD3D12Device* InParent, const TDesc* InDesc, FD3D12Resource* InResource, ViewSubresourceSubsetFlags InFlags)
+		: FD3D12ViewGeneric(InResource)
+		, ViewSubresourceSubset(*InDesc, 
+			GetResource() ? GetResource()->GetMipLevels() : 0, 
+			GetResource() ? GetResource()->GetArraySize() : 0, 
+			GetResource() ? GetResource()->GetDesc().Format : DXGI_FORMAT_UNKNOWN,
+			InFlags)
+        , FD3D12DeviceChild(InParent)
+	{
+		Init(InDesc);
+	}
+
+	virtual ~FD3D12View()
+	{
+		FreeHeapSlot();
+	}
+
+	void Init(const TDesc* InDesc)
+	{
+		if (InDesc)
+		{
+			Desc = *InDesc;
+		}
+		else
+		{
+			FMemory::Memzero(&Desc, sizeof(Desc));
+		}
+		AllocateHeapSlot();
+	}
+
+	void AllocateHeapSlot()
+	{
+		FDescriptorHeapManager& DescriptorAllocator = GetParentDevice()->GetViewDescriptorAllocator<TDesc>();
+		Descriptor = DescriptorAllocator.AllocateHeapSlot(DescriptorHeapIndex);
+		check(Descriptor.ptr != 0);
+	}
+
+	void FreeHeapSlot()
+	{
+		if (Descriptor.ptr)
+		{
+            FDescriptorHeapManager& DescriptorAllocator = GetParentDevice()->GetViewDescriptorAllocator<TDesc>();
+			DescriptorAllocator.FreeHeapSlot(Descriptor, DescriptorHeapIndex);
+			Descriptor.ptr = 0;
+		}
+	}
+
+	void UpdateViewSubresourceSubset(const FD3D12Resource* const InResource)
+	{
+		if (InResource == GetResource())
+		{
+			// Nothing to update
+			return;
+		}
+
+		if (InResource)
+		{
+			const D3D12_RESOURCE_DESC &ResDesc = InResource->GetDesc();
+			ViewSubresourceSubset.m_MipLevels = ResDesc.MipLevels;
+			ViewSubresourceSubset.m_ArraySlices = ResDesc.DepthOrArraySize;
+			ViewSubresourceSubset.m_PlaneCount = GetPlaneCount(ResDesc.Format);
+		}
+		else
+		{
+			// Null resource
+			ViewSubresourceSubset.m_MipLevels = 0;
+			ViewSubresourceSubset.m_ArraySlices = 0;
+			ViewSubresourceSubset.m_PlaneCount = 0;
+		}
+	}
+
+
+public:
+	void CreateView(FD3D12Resource* InResource = nullptr, FD3D12Resource* InCounterResource = nullptr)
+	{
+		if (!InResource)
+		{
+			InResource = GetResource();
+		}
+		else
+		{
+			// Only need to update the view's subresource subset if a new resource is used
+			UpdateViewSubresourceSubset(InResource);
+		}
+
+		check(Descriptor.ptr != 0);
+        (GetParentDevice()->GetDevice()->*TCreateViewMap<TDesc>::GetCreate()) (
+			InResource ? InResource->GetResource() : nullptr, &Desc, Descriptor);
+	}
+
+	void CreateViewWithCounter(FD3D12Resource* InResource, FD3D12Resource* InCounterResource)
+	{
+		if (!InResource)
+		{
+			InResource = GetResource();
+		}
+		else
+		{
+			// Only need to update the view's subresource subset if a new resource is used
+			UpdateViewSubresourceSubset(InResource);
+		}
+
+		check(Descriptor.ptr != 0);
+        (GetParentDevice()->GetDevice()->*TCreateViewMap<TDesc>::GetCreate()) (
+			InResource ? InResource->GetResource() : nullptr,
+			InCounterResource ? InCounterResource->GetResource() : nullptr, &Desc, Descriptor);
+	}
+
+	const TDesc& GetDesc() const { return Desc; }
+	const CViewSubresourceSubset& GetViewSubresourceSubset() const { return ViewSubresourceSubset; }
+
+	template< class T >
+	inline bool DoesNotOverlap(const FD3D12View< T >& Other) const
+	{
+		return ViewSubresourceSubset.DoesNotOverlap(Other.GetViewSubresourceSubset());
+	}
+};
+
+extern LONGLONG D3D12ShaderResourceViewSequenceNumber;
+
+/** Shader resource view class. */
+class FD3D12ShaderResourceView : public FRHIShaderResourceView, public FD3D12View < D3D12_SHADER_RESOURCE_VIEW_DESC >
+{
+	uint64 SequenceNumber;
+	bool bIsBuffer;
+	uint32 Stride;
+
+public:
+    FD3D12ShaderResourceView(FD3D12Device* InParent, D3D12_SHADER_RESOURCE_VIEW_DESC *InSRVDesc, FD3D12ResourceLocation* InResourceLocation, uint32 InStride = 1)
+		: FD3D12View(InParent, InSRVDesc, InResourceLocation, ViewSubresourceSubsetFlags_None)
+		, SequenceNumber(InterlockedIncrement64(&D3D12ShaderResourceViewSequenceNumber))
+		, bIsBuffer(InSRVDesc->ViewDimension == D3D12_SRV_DIMENSION_BUFFER)
+		, Stride(InStride)
+	{
+		if (bIsBuffer && !!InResourceLocation)
+		{
+			check(InResourceLocation->GetOffset() / Stride == InSRVDesc->Buffer.FirstElement);
+		}
+
+		CreateView();
+	}
+
+	void Rename(FD3D12ResourceLocation *InResourceLocation, CD3DX12_CPU_DESCRIPTOR_HANDLE InDescriptor, uint32 InDescriptorHeapIndex)
+	{
+		ResourceLocation = InResourceLocation;
+		if (bIsBuffer)
+		{
+			Desc.Buffer.FirstElement = InResourceLocation->GetOffset() / Stride;
+		}
+
+		D3DResource = nullptr;
+
+		SequenceNumber = InterlockedIncrement64(&D3D12ShaderResourceViewSequenceNumber);
+
+        // If no descriptor and heap index were provided...
+        if(InDescriptor.ptr == 0)
+        {
+            // Allocate a new heap slot
+            AllocateHeapSlot();
+
+			// Create a new view
+            CreateView();
+        }
+        else
+        {
+            // Otherwise use the provided descriptor and index
+            Descriptor.ptr = InDescriptor.ptr;
+            DescriptorHeapIndex = InDescriptorHeapIndex;
+        }
+
+        ResourceLocation->UpdateDefaultStateCache();
+    }
+
+	FORCEINLINE uint64 GetSequenceNumber() { return SequenceNumber; }
+	FORCEINLINE const D3D12_SHADER_RESOURCE_VIEW_DESC& GetDesc()
+	{
+		return Desc;
+	}
+};
+
+class FD3D12UnorderedAccessView : public FRHIUnorderedAccessView, public FD3D12View < D3D12_UNORDERED_ACCESS_VIEW_DESC >
+{
+public:
+
+	TRefCountPtr<FD3D12Resource> CounterResource;
+	bool CounterResourceInitialized;
+
+	FD3D12UnorderedAccessView(FD3D12Device* InParent, D3D12_UNORDERED_ACCESS_VIEW_DESC *InUAVDesc, FD3D12ResourceLocation* InResourceLocation, TRefCountPtr<FD3D12Resource> InCounterResource = nullptr)
+		: FD3D12View(InParent, InUAVDesc, InResourceLocation, ViewSubresourceSubsetFlags_None)
+		, CounterResource(InCounterResource)
+		, CounterResourceInitialized(false)
+	{
+		CreateViewWithCounter(nullptr, CounterResource);
+	}
+	
+};
+
+class FD3D12RenderTargetView : public FD3D12View<D3D12_RENDER_TARGET_VIEW_DESC>, public FRHIResource
+{
+public:
+    FD3D12RenderTargetView(FD3D12Device* InParent, D3D12_RENDER_TARGET_VIEW_DESC *InRTVDesc, FD3D12Resource* InResource)
+		: FD3D12View(InParent, InRTVDesc, InResource, ViewSubresourceSubsetFlags_None)
+	{
+		CreateView();
+	}
+};
+
+class FD3D12DepthStencilView : public FD3D12View<D3D12_DEPTH_STENCIL_VIEW_DESC>, public FRHIResource
+{
+	const bool bHasDepth;
+	const bool bHasStencil;
+	CViewSubresourceSubset DepthOnlyViewSubresourceSubset;
+	CViewSubresourceSubset StencilOnlyViewSubresourceSubset;
+
+public:
+    FD3D12DepthStencilView(FD3D12Device* InParent, D3D12_DEPTH_STENCIL_VIEW_DESC *InDSVDesc, FD3D12Resource* InResource, bool InHasStencil)
+		: FD3D12View(InParent, InDSVDesc, InResource, ViewSubresourceSubsetFlags_DepthAndStencilDsv)
+		, bHasDepth(true)				// Assume all DSVs have depth bits in their format
+		, bHasStencil(InHasStencil)		// Only some DSVs have stencil bits in their format
+		, DepthOnlyViewSubresourceSubset(*InDSVDesc, 0, 0, DXGI_FORMAT_UNKNOWN, ViewSubresourceSubsetFlags_DepthOnlyDsv)
+		, StencilOnlyViewSubresourceSubset(*InDSVDesc, 0, 0, DXGI_FORMAT_UNKNOWN, ViewSubresourceSubsetFlags_StencilOnlyDsv)
+	{
+		CreateView();
+
+		// Create individual subresource subsets for each plane
+		if (bHasDepth)
+		{
+			DepthOnlyViewSubresourceSubset = CViewSubresourceSubset(*InDSVDesc,
+				InResource ? InResource->GetMipLevels() : 0,
+				InResource ? InResource->GetArraySize() : 0,
+				InResource ? InResource->GetDesc().Format : DXGI_FORMAT_UNKNOWN,
+				ViewSubresourceSubsetFlags_DepthOnlyDsv);
+		}
+
+		if (bHasStencil)
+		{
+			StencilOnlyViewSubresourceSubset = CViewSubresourceSubset(*InDSVDesc,
+				InResource ? InResource->GetMipLevels() : 0,
+				InResource ? InResource->GetArraySize() : 0,
+				InResource ? InResource->GetDesc().Format : DXGI_FORMAT_UNKNOWN,
+				ViewSubresourceSubsetFlags_StencilOnlyDsv);
+		}
+	}
+
+	bool HasDepth() const
+	{
+		return bHasDepth;
+	}
+
+	bool HasStencil() const
+	{
+		return bHasStencil;
+	}
+
+	CViewSubresourceSubset& GetDepthOnlyViewSubresourceSubset()
+	{
+		check(bHasDepth);
+		return DepthOnlyViewSubresourceSubset;
+	}
+
+	CViewSubresourceSubset& GetStencilOnlyViewSubresourceSubset()
+	{
+		check(bHasStencil);
+		return StencilOnlyViewSubresourceSubset;
+	}
+};
+
+/** Texture base class. */
+class FD3D12TextureBase : public FD3D12BaseShaderResource
+{
+public:
+
+	FD3D12TextureBase(
+	    class FD3D12Device* InParent,
+		FD3D12Resource* InResource,
+		int32 InRTVArraySize,
+		bool bInCreatedRTVsPerSlice,
+		const TArray<TRefCountPtr<FD3D12RenderTargetView> >& InRenderTargetViews,
+		TRefCountPtr<FD3D12DepthStencilView>* InDepthStencilViews
+		)
+        : FD3D12BaseShaderResource(InResource, InParent)
+		, MemorySize(0)
+		, BaseShaderResource(this)
+		, RTVArraySize(InRTVArraySize)
+		, bCreatedRTVsPerSlice(bInCreatedRTVsPerSlice)
+		, RenderTargetViews(InRenderTargetViews)
+		, NumDepthStencilViews(0)
+	{
+		// Set the DSVs for all the access type combinations
+		if (InDepthStencilViews != nullptr)
+		{
+			for (uint32 Index = 0; Index < FExclusiveDepthStencil::MaxIndex; Index++)
+			{
+				DepthStencilViews[Index] = InDepthStencilViews[Index];
+				// New Monolithic Graphics drivers have optional "fast calls" replacing various D3d functions
+				// You can't use fast version of XXSetShaderResources (called XXSetFastShaderResource) on dynamic or d/s targets
+				if (DepthStencilViews[Index] != NULL)
+					NumDepthStencilViews++;
+			}
+		}
+	}
+
+	virtual ~FD3D12TextureBase() {}
+
+	int32 GetMemorySize() const
+	{
+		return MemorySize;
+	}
+
+	void SetMemorySize(int32 InMemorySize)
+	{
+		MemorySize = InMemorySize;
+	}
+
+	// Accessors.
+	FD3D12Resource* GetResource() const { return ResourceLocation->GetResource(); }
+	uint64 GetOffset() const { return ResourceLocation->GetOffset(); }
+	FD3D12ShaderResourceView* GetShaderResourceView() const { return ShaderResourceView; }
+	FD3D12BaseShaderResource* GetBaseShaderResource() const { return BaseShaderResource; }
+	void SetShaderResourceView(FD3D12ShaderResourceView *InShaderResourceView) { ShaderResourceView = InShaderResourceView; }
+	void SetRenderTargetViews(const TArray<TRefCountPtr<FD3D12RenderTargetView> >& InRenderTargetViews) { RenderTargetViews = InRenderTargetViews; }
+
+	/**
+	 * Get the render target view for the specified mip and array slice.
+	 * An array slice of -1 is used to indicate that no array slice should be required.
+	 */
+	FD3D12RenderTargetView* GetRenderTargetView(int32 MipIndex, int32 ArraySliceIndex) const
+	{
+		int32 ArrayIndex = MipIndex;
+
+		if (bCreatedRTVsPerSlice)
+		{
+			check(ArraySliceIndex >= 0);
+			ArrayIndex = MipIndex * RTVArraySize + ArraySliceIndex;
+		}
+		else
+		{
+			// Catch attempts to use a specific slice without having created the texture to support it
+			check(ArraySliceIndex == -1 || ArraySliceIndex == 0);
+		}
+
+		if ((uint32)ArrayIndex < (uint32)RenderTargetViews.Num())
+		{
+			return RenderTargetViews[ArrayIndex];
+		}
+		return 0;
+	}
+	FD3D12DepthStencilView* GetDepthStencilView(FExclusiveDepthStencil AccessType) const
+	{
+		return DepthStencilViews[AccessType.GetIndex()];
+	}
+
+	// New Monolithic Graphics drivers have optional "fast calls" replacing various D3d functions
+	// You can't use fast version of XXSetShaderResources (called XXSetFastShaderResource) on dynamic or d/s targets
+	bool HasDepthStencilView()
+	{
+		return (NumDepthStencilViews > 0);
+	}
+
+protected:
+
+	/** Amount of memory allocated by this texture, in bytes. */
+	int32 MemorySize;
+
+	/** Pointer to the base shader resource. Usually the object itself, but not for texture references. */
+	FD3D12BaseShaderResource* BaseShaderResource;
+
+	/** A shader resource view of the texture. */
+	TRefCountPtr<FD3D12ShaderResourceView> ShaderResourceView;
+
+	/** A render targetable view of the texture. */
+	TArray<TRefCountPtr<FD3D12RenderTargetView> > RenderTargetViews;
+
+	bool bCreatedRTVsPerSlice;
+
+	int32 RTVArraySize;
+
+	/** A depth-stencil targetable view of the texture. */
+	TRefCountPtr<FD3D12DepthStencilView> DepthStencilViews[FExclusiveDepthStencil::MaxIndex];
+
+	/** Number of Depth Stencil Views - used for fast call tracking. */
+	uint32	NumDepthStencilViews;
+
+};
+
+/** 2D texture (vanilla, cubemap or 2D array) */
+template<typename BaseResourceType>
+class TD3D12Texture2D : public BaseResourceType, public FD3D12TextureBase
+{
+public:
+
+	/** Flags used when the texture was created */
+	uint32 Flags;
+
+	/** Initialization constructor. */
+	TD3D12Texture2D(
+	class FD3D12Device* InParent,
+		FD3D12Resource* InResource,
+		bool bInCreatedRTVsPerSlice,
+		int32 InRTVArraySize,
+		const TArray<TRefCountPtr<FD3D12RenderTargetView> >& InRenderTargetViews,
+		TRefCountPtr<FD3D12DepthStencilView>* InDepthStencilViews,
+		uint32 InSizeX,
+		uint32 InSizeY,
+		uint32 InSizeZ,
+		uint32 InNumMips,
+		uint32 InNumSamples,
+		EPixelFormat InFormat,
+		bool bInCubemap,
+		uint32 InFlags,
+		bool bInPooled,
+		const FClearValueBinding& InClearValue
+#if PLATFORM_SUPPORTS_VIRTUAL_TEXTURES
+		, void* InRawTextureMemory = nullptr
+#endif
+		)
+		: BaseResourceType(
+		InSizeX,
+		InSizeY,
+		InSizeZ,
+		InNumMips,
+		InNumSamples,
+		InFormat,
+		InFlags,
+		InClearValue
+		)
+		, FD3D12TextureBase(
+        InParent,
+		InResource,
+		InRTVArraySize,
+		bInCreatedRTVsPerSlice,
+		InRenderTargetViews,
+		InDepthStencilViews
+		)
+		, Flags(InFlags)
+		, bCubemap(bInCubemap)
+		, bPooled(bInPooled)
+#if PLATFORM_SUPPORTS_VIRTUAL_TEXTURES
+		, RawTextureMemory(InRawTextureMemory)
+#endif
+	{
+		FMemory::Memzero(ReadBackHeapDesc);
+	}
+
+	virtual ~TD3D12Texture2D();
+
+	/**
+	 * Locks one of the texture's mip-maps.
+	 * @return A pointer to the specified texture data.
+	 */
+	void* Lock(uint32 MipIndex, uint32 ArrayIndex, EResourceLockMode LockMode, uint32& DestStride);
+
+	/** Unlocks a previously locked mip-map. */
+	void Unlock(uint32 MipIndex, uint32 ArrayIndex);
+
+	// Accessors.
+	FD3D12Resource* GetResource() const { return (FD3D12Resource*)FD3D12TextureBase::GetResource(); }
+	const D3D12_PLACED_SUBRESOURCE_FOOTPRINT& GetReadBackHeapDesc() const 
+	{ 
+		// This should only be called if SetReadBackHeapDesc() was called with actual contents
+		check(ReadBackHeapDesc.Footprint.Width > 0 && ReadBackHeapDesc.Footprint.Height > 0);
+		
+		return ReadBackHeapDesc; 
+	}
+
+	FD3D12CLSyncPoint GetReadBackSyncPoint() const { return ReadBackSyncPoint; }
+	bool IsCubemap() const { return bCubemap; }
+
+	/** FRHITexture override.  See FRHITexture::GetNativeResource() */
+	virtual void* GetNativeResource() const override
+	{
+		return GetResource();
+	}
+	virtual void* GetNativeShaderResourceView() const override
+	{
+		return GetShaderResourceView();
+	}
+
+	// Modifiers.
+	void SetReadBackHeapDesc(const D3D12_PLACED_SUBRESOURCE_FOOTPRINT &newReadBackHeapDesc) { ReadBackHeapDesc = newReadBackHeapDesc; }
+	void SetReadBackListHandle(FD3D12CommandListHandle listToWaitFor) { ReadBackSyncPoint = listToWaitFor; }
+
+	// IRefCountedObject interface.
+	virtual uint32 AddRef() const
+	{
+		return FRHIResource::AddRef();
+	}
+	virtual uint32 Release() const
+	{
+		return FRHIResource::Release();
+	}
+	virtual uint32 GetRefCount() const
+	{
+		return FRHIResource::GetRefCount();
+	}
+#if PLATFORM_SUPPORTS_VIRTUAL_TEXTURES
+	void* GetRawTextureMemory() const
+	{
+		return RawTextureMemory;
+	}
+#endif
+
+private:
+	D3D12_PLACED_SUBRESOURCE_FOOTPRINT ReadBackHeapDesc;
+	FD3D12CLSyncPoint ReadBackSyncPoint;
+
+	/** Whether the texture is a cube-map. */
+	const uint32 bCubemap : 1;
+	/** Whether the texture can be pooled. */
+	const uint32 bPooled : 1;
+#if PLATFORM_SUPPORTS_VIRTUAL_TEXTURES
+	void* RawTextureMemory;
+#endif
+};
+
+/** 3D Texture */
+class FD3D12Texture3D : public FRHITexture3D, public FD3D12TextureBase
+{
+public:
+
+	/** Initialization constructor. */
+	FD3D12Texture3D(
+	class FD3D12Device* InParent,
+		FD3D12Resource* InResource,
+		const TArray<TRefCountPtr<FD3D12RenderTargetView> >& InRenderTargetViews,
+		uint32 InSizeX,
+		uint32 InSizeY,
+		uint32 InSizeZ,
+		uint32 InNumMips,
+		EPixelFormat InFormat,
+		uint32 InFlags,
+		const FClearValueBinding& InClearValue
+		)
+		: FRHITexture3D(InSizeX, InSizeY, InSizeZ, InNumMips, InFormat, InFlags,InClearValue)
+		, FD3D12TextureBase(
+        InParent,
+		InResource,
+		1,
+		false,
+		InRenderTargetViews,
+		NULL
+		)
+	{
+	}
+
+	virtual ~FD3D12Texture3D();
+
+	// Accessors.
+	FD3D12Resource* GetResource() const { return (FD3D12Resource*)FD3D12TextureBase::GetResource(); }
+
+	// IRefCountedObject interface.
+	virtual uint32 AddRef() const
+	{
+		return FRHIResource::AddRef();
+	}
+	virtual uint32 Release() const
+	{
+		return FRHIResource::Release();
+	}
+	virtual uint32 GetRefCount() const
+	{
+		return FRHIResource::GetRefCount();
+	}
+};
+
+class FD3D12BaseTexture2D : public FRHITexture2D
+{
+public:
+	FD3D12BaseTexture2D(uint32 InSizeX, uint32 InSizeY, uint32 InSizeZ, uint32 InNumMips, uint32 InNumSamples, EPixelFormat InFormat, uint32 InFlags, const FClearValueBinding& InClearValue)
+		: FRHITexture2D(InSizeX, InSizeY, InNumMips, InNumSamples, InFormat, InFlags, InClearValue)
+	{}
+	uint32 GetSizeZ() const { return 0; }
+};
+
+class FD3D12BaseTexture2DArray : public FRHITexture2DArray
+{
+public:
+	FD3D12BaseTexture2DArray(uint32 InSizeX, uint32 InSizeY, uint32 InSizeZ, uint32 InNumMips, uint32 InNumSamples, EPixelFormat InFormat, uint32 InFlags, const FClearValueBinding& InClearValue)
+		: FRHITexture2DArray(InSizeX, InSizeY, InSizeZ, InNumMips, InFormat, InFlags, InClearValue)
+	{
+		check(InNumSamples == 1);
+	}
+};
+
+class FD3D12BaseTextureCube : public FRHITextureCube
+{
+public:
+	FD3D12BaseTextureCube(uint32 InSizeX, uint32 InSizeY, uint32 InSizeZ, uint32 InNumMips, uint32 InNumSamples, EPixelFormat InFormat, uint32 InFlags, const FClearValueBinding& InClearValue)
+		: FRHITextureCube(InSizeX, InNumMips, InFormat, InFlags, InClearValue)
+	{
+		check(InNumSamples == 1);
+	}
+	uint32 GetSizeX() const { return GetSize(); }
+	uint32 GetSizeY() const { return GetSize(); }
+	uint32 GetSizeZ() const { return 0; }
+};
+
+typedef TD3D12Texture2D<FRHITexture>              FD3D12Texture;
+typedef TD3D12Texture2D<FD3D12BaseTexture2D>      FD3D12Texture2D;
+typedef TD3D12Texture2D<FD3D12BaseTexture2DArray> FD3D12Texture2DArray;
+typedef TD3D12Texture2D<FD3D12BaseTextureCube>    FD3D12TextureCube;
+
+/** Texture reference class. */
+class FD3D12TextureReference : public FRHITextureReference, public FD3D12TextureBase
+{
+public:
+	FD3D12TextureReference(class FD3D12Device* InParent, FLastRenderTimeContainer* LastRenderTime)
+		: FRHITextureReference(LastRenderTime)
+        , FD3D12TextureBase(InParent, NULL, 0, false, TArray<TRefCountPtr<FD3D12RenderTargetView> >(), NULL)
+	{
+		BaseShaderResource = NULL;
+	}
+
+	void SetReferencedTexture(FRHITexture* InTexture, FD3D12BaseShaderResource* InBaseShaderResource, FD3D12ShaderResourceView* InSRV)
+	{
+		ShaderResourceView = InSRV;
+		BaseShaderResource = InBaseShaderResource;
+		FRHITextureReference::SetReferencedTexture(InTexture);
+	}
+
+	// IRefCountedObject interface.
+	virtual uint32 AddRef() const
+	{
+		return FRHIResource::AddRef();
+	}
+	virtual uint32 Release() const
+	{
+		return FRHIResource::Release();
+	}
+	virtual uint32 GetRefCount() const
+	{
+		return FRHIResource::GetRefCount();
+	}
+};
+
+/** Given a pointer to a RHI texture that was created by the D3D11 RHI, returns a pointer to the FD3D12TextureBase it encapsulates. */
+inline FD3D12TextureBase* GetD3D11TextureFromRHITexture(FRHITexture* Texture)
+{
+	if (!Texture)
+	{
+		return NULL;
+	}
+	else if (Texture->GetTexture2D())
+	{
+		return static_cast<FD3D12Texture2D*>(Texture);
+	}
+	else if (Texture->GetTextureReference())
+	{
+		return static_cast<FD3D12TextureReference*>(Texture);
+	}
+	else if (Texture->GetTexture2DArray())
+	{
+		return static_cast<FD3D12Texture2DArray*>(Texture);
+	}
+	else if (Texture->GetTexture3D())
+	{
+		return static_cast<FD3D12Texture3D*>(Texture);
+	}
+	else if (Texture->GetTextureCube())
+	{
+		return static_cast<FD3D12TextureCube*>(Texture);
+	}
+	else
+	{
+		UE_LOG(LogD3D12RHI, Fatal, TEXT("Unknown RHI texture type"));
+		return NULL;
+	}
+}
+
+/** D3D12 occlusion query */
+class FD3D12OcclusionQuery : public FRHIRenderQuery
+{
+public:
+
+	/** The query heap resource. */
+    TRefCountPtr<ID3D12QueryHeap> QueryHeap;
+    uint32 HeapIndex;
+
+	/** CPU-visible buffer to store the query result **/
+    TRefCountPtr<ID3D12Resource> ResultBuffer;
+
+	/** The cached query result. */
+	uint64 Result;
+
+	/** true if the query's result is cached. */
+	bool bResultIsCached : 1;
+
+	// todo: memory optimize
+	ERenderQueryType Type;
+
+	class FD3D12CommandContext* OwningContext;
+	FD3D12CommandListHandle OwningCommandList;
+
+	/** Initialization constructor. */
+	FD3D12OcclusionQuery(ID3D12QueryHeap* InQueryHeap, ID3D12Resource* InQueryResultBuffer, ERenderQueryType InQueryType) :
+        QueryHeap(InQueryHeap),
+        ResultBuffer(InQueryResultBuffer),
+        Result(0),
+        Type(InQueryType)
+	{
+        Reset();
+	}
+
+    void Reset()
+    {
+        HeapIndex = -1;
+        bResultIsCached = false;
+		OwningContext = nullptr;
+        OwningCommandList = nullptr;
+    }
+};
+
+/** Updates tracked stats for a buffer. */
+#define D3D12_BUFFER_TYPE_CONSTANT   1
+#define D3D12_BUFFER_TYPE_INDEX      2
+#define D3D12_BUFFER_TYPE_VERTEX     3
+#define D3D12_BUFFER_TYPE_STRUCTURED 4
+extern void UpdateBufferStats(FD3D12ResourceLocation* ResourceLocation, bool bAllocating, uint32 BufferType);
+
+/** Forward declare the constants ring buffer. */
+class FD3D12ConstantsRingBuffer;
+
+/** A ring allocation from the constants ring buffer. */
+struct FRingAllocation
+{
+	FD3D12Resource* Buffer;
+	void* DataPtr;
+	uint32 Offset;
+	uint32 Size;
+
+	FRingAllocation() : Buffer(NULL) {}
+	inline bool IsValid() const { return Buffer != NULL; }
+};
+
+/** Uniform buffer resource class. */
+class FD3D12UniformBuffer : public FRHIUniformBuffer, public FD3D12DeviceChild
+{
+public:
+	/** The handle to the descriptor in the offline descriptor heap */
+	CD3DX12_CPU_DESCRIPTOR_HANDLE OfflineDescriptorHandle;
+
+	/** Unique number which identifies this constant buffer */
+	uint64 SequenceNumber;
+
+	/** The D3D11 constant buffer resource */
+	TRefCountPtr<FD3D12ResourceLocation> ResourceLocation;
+
+	/** Resource table containing RHI references. */
+	TArray<TRefCountPtr<FRHIResource> > ResourceTable;
+
+	/** Cached resources need to retain the associated shader resource for bookkeeping purposes. */
+	struct FResourcePair
+	{
+		FD3D12ResourceLocation* ShaderResourceLocation;
+		IUnknown* D3D11Resource;
+	};
+
+	/** Raw resource table, cached once per frame. */
+	TArray<FResourcePair> RawResourceTable;
+
+	/** The frame in which RawResourceTable was last cached. */
+	uint32 LastCachedFrame;
+
+	/** Index of the descriptor in the offline heap */
+	uint32 OfflineHeapIndex;
+
+	/** Initialization constructor. */
+	FD3D12UniformBuffer(class FD3D12Device* InParent, const FRHIUniformBufferLayout& InLayout, FD3D12ResourceLocation* InResourceLocation, const FRingAllocation& InRingAllocation, uint64 InSequenceNumber)
+		: FRHIUniformBuffer(InLayout)
+		, ResourceLocation(InResourceLocation)
+		, LastCachedFrame((uint32)-1)
+		, OfflineHeapIndex(UINT_MAX)
+		, SequenceNumber(InSequenceNumber)
+        , FD3D12DeviceChild(InParent)
+	{
+		OfflineDescriptorHandle.ptr = 0;
+	}
+
+	virtual ~FD3D12UniformBuffer();
+
+	/** Cache resources if needed. */
+	inline void CacheResources(uint32 InFrameCounter)
+	{
+		if (InFrameCounter == INDEX_NONE || LastCachedFrame != InFrameCounter)
+		{
+			CacheResourcesInternal();
+			LastCachedFrame = InFrameCounter;
+		}
+	}
+
+private:
+	class FD3D12DynamicRHI* D3D12RHI;
+	/** Actually cache resources. */
+	void CacheResourcesInternal();
+};
+
+/** Index buffer resource class that stores stride information. */
+class FD3D12IndexBuffer : public FRHIIndexBuffer, public FD3D12BaseShaderResource
+{
+public:
+
+	FD3D12IndexBuffer(FD3D12Device* InParent, FD3D12ResourceLocation* InResourceLocation, uint32 InStride, uint32 InSize, uint32 InUsage)
+		: FRHIIndexBuffer(InStride, InSize, InUsage)
+        , FD3D12BaseShaderResource(InResourceLocation, InParent)
+	{}
+
+	virtual ~FD3D12IndexBuffer();
+
+	// IRefCountedObject interface.
+	virtual uint32 AddRef() const
+	{
+		return FRHIResource::AddRef();
+	}
+	virtual uint32 Release() const
+	{
+		return FRHIResource::Release();
+	}
+	virtual uint32 GetRefCount() const
+	{
+		return FRHIResource::GetRefCount();
+	}
+};
+
+/** Structured buffer resource class. */
+class FD3D12StructuredBuffer : public FRHIStructuredBuffer, public FD3D12BaseShaderResource
+{
+public:
+	TRefCountPtr<FD3D12Resource> Resource;
+   
+	FD3D12StructuredBuffer(FD3D12Device* InParent, FD3D12ResourceLocation* InResourceLocation, uint32 InStride, uint32 InSize, uint32 InUsage)
+		: FRHIStructuredBuffer(InStride, InSize, InUsage)
+        , FD3D12BaseShaderResource(InResourceLocation, InParent)
+	{
+		Resource = InResourceLocation->GetResource();
+	}
+
+	FD3D12StructuredBuffer(FD3D12Device* InParent, FD3D12Resource* InResource, uint32 InStride, uint32 InSize, uint32 InUsage)
+		: FRHIStructuredBuffer(InStride, InSize, InUsage)
+        , FD3D12BaseShaderResource(InResource, InParent)
+		, Resource(InResource)
+	{}
+
+	virtual ~FD3D12StructuredBuffer();
+
+	// IRefCountedObject interface.
+	virtual uint32 AddRef() const
+	{
+		return FRHIResource::AddRef();
+	}
+	virtual uint32 Release() const
+	{
+		return FRHIResource::Release();
+	}
+	virtual uint32 GetRefCount() const
+	{
+		return FRHIResource::GetRefCount();
+	}
+};
+
+/** Vertex buffer resource class. */
+class FD3D12VertexBuffer : public FRHIVertexBuffer, public FD3D12BaseShaderResource
+{
+	static const uint64 PruneFenceCount = 32;
+
+	struct ResourceViewHandleDesc
+	{
+		ResourceViewHandleDesc()
+			: BlockInfo(nullptr)
+			, DescriptorHeapIndex(0)
+		{
+			SRVHandle.ptr = 0;
+		}
+		ResourceViewHandleDesc(FD3D12ResourceBlockInfo *InBlockInfo, const CD3DX12_CPU_DESCRIPTOR_HANDLE &InSRVHandle, uint32 InDescriptorHeapIndex)
+			: BlockInfo(InBlockInfo)
+			, SRVHandle(InSRVHandle)
+			, DescriptorHeapIndex(InDescriptorHeapIndex) {}
+		TRefCountPtr<FD3D12ResourceBlockInfo> BlockInfo;
+		CD3DX12_CPU_DESCRIPTOR_HANDLE SRVHandle;
+		uint32 DescriptorHeapIndex;
+	};
+
+	TFencedObjectPool<ResourceViewHandleDesc> SRVPool;
+
+public:
+
+	// Current SRV
+	FD3D12ShaderResourceView *DynamicSRV;
+
+	FD3D12VertexBuffer(FD3D12Device* InParent, FD3D12ResourceLocation *InResourceLocation, uint32 InSize, uint32 InUsage)
+		: FRHIVertexBuffer(InSize, InUsage)
+        , FD3D12BaseShaderResource(InResourceLocation, InParent)
+        , DynamicSRV(nullptr)
+        , SRVPool(InParent)
+	{
+	}
+
+	virtual ~FD3D12VertexBuffer();
+
+	void *DynamicLock();
+	void SetDynamicSRV(FD3D12ShaderResourceView *InSRV)
+	{
+		DynamicSRV = InSRV;
+	}
+
+	FD3D12ShaderResourceView *GetDynamicSRV() const
+	{
+		return DynamicSRV;
+	}
+
+	// IRefCountedObject interface.
+	virtual uint32 AddRef() const
+	{
+		return FRHIResource::AddRef();
+	}
+	virtual uint32 Release() const
+	{
+		return FRHIResource::Release();
+	}
+	virtual uint32 GetRefCount() const
+	{
+		return FRHIResource::GetRefCount();
+	}
+};
+
+
+namespace D3D12RHI
+{
+	void ReturnPooledTexture2D(int32 MipCount, EPixelFormat PixelFormat, FD3D12Resource* InResource);
+	void ReleasePooledTextures();
+}
+
+// FD3D12ResourceAllocator
+//
+// Implements an interface for resource allocation and deallocation. Resources are allocated from buckets, which are just
+// a collection of resources of a particular size. Blocks can be an entire resource or a sub allocation from a resource.
+class FD3D12ResourceAllocator : public FD3D12DeviceChild
+{
+protected:
+#if SUB_ALLOCATED_DEFAULT_ALLOCATIONS
+	static const uint32 MIN_HEAP_SIZE = 256 * 1024;
+#else
+	static const uint32 MIN_HEAP_SIZE = 64 * 1024;
+#endif
+
+	static uint32 FORCEINLINE BucketFromSize(uint32 size, uint32 bucketShift)
+	{
+		uint32 bucket = FMath::CeilLogTwo(size);
+		bucket = bucket < bucketShift ? 0 : bucket - bucketShift;
+		return bucket;
+	}
+
+	static uint32 FORCEINLINE BlockSizeFromBufferSize(uint32 bufferSize, uint32 bucketShift)
+	{
+		const uint32 minSize = 1 << bucketShift;
+		return bufferSize > minSize ? FMath::RoundUpToPowerOfTwo(bufferSize) : minSize;
+	}
+
+	FThreadsafeQueue<TRefCountPtr<FD3D12ResourceBlockInfo>> ExpiredBlocks;
+
+public:
+	void virtual CleanupFreeBlocks() = 0;
+	void virtual ExpireBlock(FD3D12ResourceBlockInfo *Block) = 0;
+	void virtual ReleaseAllResources() = 0;
+
+    FD3D12ResourceAllocator(FD3D12Device* InParent) :
+        FD3D12DeviceChild(InParent){}
+	virtual ~FD3D12ResourceAllocator(){}
+};
+
+// FD3D12DynamicHeapAllocator
+//
+// Implements a D3D12 buffer as a ring buffer allocating memory at the tail and freeing memory from the head.
+// Memory allocation and deallocation must happen in the same order since there can be no 
+// fragmentation of memory.  This is designed for allocation of scratch memory such as temporary staging buffers
+// or shadow buffers for dynamic resources.
+class FD3D12DynamicHeapAllocator : public FD3D12ResourceAllocator
+{
+	static const uint32 BucketShift = 6;
+	static const uint32 NumBuckets = 21; // bucket resource sizes range from 64 to 2^24 
+	FThreadsafeQueue<TRefCountPtr<FD3D12ResourceBlockInfo>> AvailableBlocks[NumBuckets];
+	
+	D3D12_HEAP_TYPE HeapType;
+
+	// Fast alloc buffer
+	const uint32 MaxFastAllocBufferSize = 65536 * 16;
+	FD3D12ResourceLocation *GetNextFastAllocBufferLocation(uint32 Size);
+	TRefCountPtr<FD3D12ResourceLocation> FastAllocBuffer;
+	uint32 NextFastAllocOffset;
+	void *FastAllocData;
+
+    // The currently paired command context.  If none is available (as with async texture loads),
+    // the current command list handle can be provided instead.  The command list handle is the
+    // critical part, but the command context can always return the current list handle.
+    // I.e. you need one but not both of these to be valid
+	class FD3D12CommandContext* CurrentCommandContext;
+    FD3D12CommandListHandle CurrentCommandListHandle;
+
+#ifdef UE_BUILD_DEBUG
+	struct
+	{
+		uint64 NumCreatedResources[NumBuckets];
+		uint64 BucketSize[NumBuckets];
+		uint64 BlockSize[NumBuckets];
+		uint64 CurAllocs[NumBuckets];
+		uint64 AvailableAllocs[NumBuckets];
+		uint64 TotalAllocs[NumBuckets];
+	} AllocatorStats;
+#endif
+
+    static HRESULT CreateResource(FD3D12DynamicHeapAllocator* Allocator, D3D12_HEAP_TYPE heapType, uint32 size, FD3D12Resource **ppResource, void **ppBaseAddress);
+
+public:
+	void CleanupFreeBlocks();
+
+	FD3D12DynamicHeapAllocator(FD3D12Device* InParent, D3D12_HEAP_TYPE heapType);
+	virtual ~FD3D12DynamicHeapAllocator();
+
+    // These are mutually exclusive calls!
+	inline void SetCurrentCommandContext(class FD3D12CommandContext* CmdContext)
+    {
+        CurrentCommandContext = CmdContext;
+        check(CurrentCommandListHandle == nullptr);
+    }
+
+    inline void SetCurrentCommandListHandle(FD3D12CommandListHandle CommandListHandle)
+    {
+        check(CurrentCommandContext == nullptr);
+        CurrentCommandListHandle = CommandListHandle;
+    }
+
+	// Allocates <size> bytes from the end of an available resource heap.
+	void Alloc(uint32 size, uint32 alignment, FD3D12ResourceBlockInfo **ppBlock);
+	void *Alloc(uint32 size, uint32 alignment, class FD3D12ResourceLocation *ResourceLocation);
+	void *Alloc(uint32 size, uint32 alignment, class FD3D12ResourceLocation **ResourceLocation);
+
+	// Allocates <size> bytes for immediate single use (assume resource will be released after the next Draw() or Copy*() call).
+	void *FastAlloc(uint32 size, uint32 alignment, class FD3D12ResourceLocation *ResourceLocation);
+
+	void ExpireBlock(FD3D12ResourceBlockInfo *Block);
+
+	void ReleaseAllResources();
+
+	void DumpAllocatorStats(class FOutputDevice& Ar);
+
+
+};
+
+// FD3D12DefaultBufferPool
+//
+class FD3D12DefaultBufferPool : public FD3D12ResourceAllocator
+{
+	static const uint32 BucketShift = 6;
+	static const uint32 NumBuckets = 24;
+	static const uint32 NumPreCreatedBuffersPerBucket = 0;
+	static const uint32 MaxAvailableBlockCountPerBucket = 256;
+	FThreadsafeQueue<TRefCountPtr<FD3D12ResourceBlockInfo>> AvailableBuffers[NumBuckets];
+	D3D12_RESOURCE_DESC resourceDesc;
+	bool bFirstAllocFromPool;
+
+#if UE_BUILD_DEBUG
+	uint64 NumCreatedResources[NumBuckets];
+	uint64 BucketSize[NumBuckets];
+#endif
+
+public:
+	FD3D12DefaultBufferPool(FD3D12Device* InParent);
+    FD3D12DefaultBufferPool();
+	virtual ~FD3D12DefaultBufferPool();
+
+	// Grab a buffer from the available buffers or create a new buffer if none are available
+	HRESULT AllocDefaultResource(const D3D12_RESOURCE_DESC& Desc, D3D12_SUBRESOURCE_DATA *pInitialData, class FD3D12ResourceLocation *ResourceLocation, uint32 Alignment);
+	HRESULT AllocDefaultResource(const D3D12_RESOURCE_DESC& Desc, D3D12_SUBRESOURCE_DATA *pInitialData, class FD3D12Resource *Resource, uint32 Alignment);
+
+	void ExpireBlock(FD3D12ResourceBlockInfo *Block);
+	void CleanupFreeBlocks();
+
+	void ReleaseAllResources();
+};
+
+// FD3D12DefaultBufferAllocator
+//
+class FD3D12DefaultBufferAllocator : public FD3D12DeviceChild
+{
+	const static uint32 MAX_DEFAULT_POOLS = 16; // Should match the max D3D12_RESOURCE_FLAG_FLAG combinations.
+
+	FD3D12DefaultBufferPool* DefaultBufferPools[MAX_DEFAULT_POOLS];
+
+public:
+	HRESULT AllocDefaultResource(const D3D12_RESOURCE_DESC& Desc, D3D12_SUBRESOURCE_DATA *pInitialData, class FD3D12ResourceLocation **ResourceLocation, uint32 Alignment = 0);
+	// Grab a buffer from the available buffers or create a new buffer if none are available
+	HRESULT AllocDefaultResource(const D3D12_RESOURCE_DESC& pDesc, D3D12_SUBRESOURCE_DATA *pInitialData, class FD3D12ResourceLocation *ResourceLocation, uint32 Alignment = 0);
+	void FreeDefaultBufferPools();
+	void CleanupFreeBlocks();
+
+    FD3D12DefaultBufferAllocator(FD3D12Device* InParent) :
+        FD3D12DeviceChild(InParent)
+	{
+		for (uint32 i = 0; i < MAX_DEFAULT_POOLS; ++i)
+			DefaultBufferPools[i] = new FD3D12DefaultBufferPool(InParent);
+		}
+	};
+
+#if SUPPORTS_MEMORY_RESIDENCY
+class FD3D12ResourceResidencyManager : public FD3D12DeviceChild
+	{
+public:
+   
+    const static uint32 MAX_RESOURCES               = 60000;
+	const static uint32 MEMORY_PRESSURE_LEVELS      = 8;
+	const static uint32 STALL_MEMORY_PRESSURE_LEVEL = 7;
+   
+    FD3D12ResourceResidencyManager()
+		: HeadIndex (0)
+		, TailIndex (0)
+    {
+        Resources.AddZeroed (MAX_RESOURCES);
+    }
+   
+	void Init(FD3D12Device* InParent)
+	{
+		Parent = InParent;
+	}
+   
+	void ResourceFreed (FD3D12Resource* Resource);
+    void UpdateResidency (FD3D12Resource* Resource);
+    void Process(uint32 MemoryPressureLevel = 0, uint32 BytesToFree = 0);
+	void FreeMemory (uint32 BytesToFree);
+	void MakeResident();
+   
+private:
+   
+	const static uint32 MEMORY_PRESSURE_FENCE_THRESHOLD[MEMORY_PRESSURE_LEVELS];
+   
+    struct FD3D12Element
+		{
+        uint64          LastUsedFence;
+        FD3D12Resource* Resource;
+	};
+
+	struct FD3D12ResourceStats
+    {
+		FD3D12ResourceStats()
+			: Memory (0)
+			, ResourceCount (0)
+			, PendingMemory (0)
+			, PendingResouceCount (0)
+			, MemoryChurn (0)
+			, ResourceCountChurn (0)
+		{}
+
+		uint64 Memory;
+		uint64 ResourceCount;
+		uint64 PendingMemory;
+		uint64 PendingResouceCount;
+		uint64 MemoryChurn;
+		uint64 ResourceCountChurn;
+    };
+
+    uint32                  HeadIndex;
+    uint32                  TailIndex;
+    TArray<FD3D12Element>   Resources;
+	FThreadsafeQueue<ID3D12Resource*> MakeResidentResources;
+	FCriticalSection        CS;
+	FD3D12ResourceStats		ResidentStats;
+	FD3D12ResourceStats		EvictedStats;
+};
+#endif
+
+template<class T>
+struct TD3D12ResourceTraits
+{
+};
+template<>
+struct TD3D12ResourceTraits<FRHIVertexDeclaration>
+{
+	typedef FD3D12VertexDeclaration TConcreteType;
+};
+template<>
+struct TD3D12ResourceTraits<FRHIVertexShader>
+{
+	typedef FD3D12VertexShader TConcreteType;
+};
+template<>
+struct TD3D12ResourceTraits<FRHIGeometryShader>
+{
+	typedef FD3D12GeometryShader TConcreteType;
+};
+template<>
+struct TD3D12ResourceTraits<FRHIHullShader>
+{
+	typedef FD3D12HullShader TConcreteType;
+};
+template<>
+struct TD3D12ResourceTraits<FRHIDomainShader>
+{
+	typedef FD3D12DomainShader TConcreteType;
+};
+template<>
+struct TD3D12ResourceTraits<FRHIPixelShader>
+{
+	typedef FD3D12PixelShader TConcreteType;
+};
+template<>
+struct TD3D12ResourceTraits<FRHIComputeShader>
+{
+	typedef FD3D12ComputeShader TConcreteType;
+};
+template<>
+struct TD3D12ResourceTraits<FRHIBoundShaderState>
+{
+	typedef FD3D12BoundShaderState TConcreteType;
+};
+template<>
+struct TD3D12ResourceTraits<FRHITexture3D>
+{
+	typedef FD3D12Texture3D TConcreteType;
+};
+template<>
+struct TD3D12ResourceTraits<FRHITexture>
+{
+	typedef FD3D12Texture TConcreteType;
+};
+template<>
+struct TD3D12ResourceTraits<FRHITexture2D>
+{
+	typedef FD3D12Texture2D TConcreteType;
+};
+template<>
+struct TD3D12ResourceTraits<FRHITexture2DArray>
+{
+	typedef FD3D12Texture2DArray TConcreteType;
+};
+template<>
+struct TD3D12ResourceTraits<FRHITextureCube>
+{
+	typedef FD3D12TextureCube TConcreteType;
+};
+template<>
+struct TD3D12ResourceTraits<FRHIRenderQuery>
+{
+	typedef FD3D12OcclusionQuery TConcreteType;
+};
+template<>
+struct TD3D12ResourceTraits<FRHIUniformBuffer>
+{
+	typedef FD3D12UniformBuffer TConcreteType;
+};
+template<>
+struct TD3D12ResourceTraits<FRHIIndexBuffer>
+{
+	typedef FD3D12IndexBuffer TConcreteType;
+};
+template<>
+struct TD3D12ResourceTraits<FRHIStructuredBuffer>
+{
+	typedef FD3D12StructuredBuffer TConcreteType;
+};
+template<>
+struct TD3D12ResourceTraits<FRHIVertexBuffer>
+{
+	typedef FD3D12VertexBuffer TConcreteType;
+};
+template<>
+struct TD3D12ResourceTraits<FRHIShaderResourceView>
+{
+	typedef FD3D12ShaderResourceView TConcreteType;
+};
+template<>
+struct TD3D12ResourceTraits<FRHIUnorderedAccessView>
+{
+	typedef FD3D12UnorderedAccessView TConcreteType;
+};
+template<>
+struct TD3D12ResourceTraits<FRHISamplerState>
+{
+	typedef FD3D12SamplerState TConcreteType;
+};
+template<>
+struct TD3D12ResourceTraits<FRHIRasterizerState>
+{
+	typedef FD3D12RasterizerState TConcreteType;
+};
+template<>
+struct TD3D12ResourceTraits<FRHIDepthStencilState>
+{
+	typedef FD3D12DepthStencilState TConcreteType;
+};
+template<>
+struct TD3D12ResourceTraits<FRHIBlendState>
+{
+	typedef FD3D12BlendState TConcreteType;
+};
\ No newline at end of file
diff --git a/Engine/Source/Runtime/Windows/D3D12RHI/Public/D3D12ShaderResources.h b/Engine/Source/Runtime/Windows/D3D12RHI/Public/D3D12ShaderResources.h
new file mode 100644
--- /dev/null
+++ b/Engine/Source/Runtime/Windows/D3D12RHI/Public/D3D12ShaderResources.h
@@ -0,0 +1,11 @@
+// Copyright 1998-2014 Epic Games, Inc. All Rights Reserved.
+
+/*=============================================================================
+	D3D12Resources.h: D3D resource RHI definitions.
+=============================================================================*/
+
+#pragma once
+
+struct FD3D12ShaderResourceTable : public FBaseShaderResourceTable
+{
+};
diff --git a/Engine/Source/Runtime/Windows/D3D12RHI/Public/D3D12State.h b/Engine/Source/Runtime/Windows/D3D12RHI/Public/D3D12State.h
new file mode 100644
--- /dev/null
+++ b/Engine/Source/Runtime/Windows/D3D12RHI/Public/D3D12State.h
@@ -0,0 +1,42 @@
+// Copyright 1998-2014 Epic Games, Inc. All Rights Reserved.
+
+/*=============================================================================
+	D3D12State.h: D3D state definitions.
+	=============================================================================*/
+
+#pragma once
+
+class FD3D12SamplerState : public FRHISamplerState, public FD3D12DeviceChild
+{
+public:
+	D3D12_CPU_DESCRIPTOR_HANDLE Descriptor;
+	uint32 DescriptorHeapIndex;
+	uint16 ID;
+
+	FD3D12SamplerState(FD3D12Device* InParent, const D3D12_SAMPLER_DESC& Desc, uint16 SamplerID);
+	~FD3D12SamplerState();
+};
+
+class FD3D12RasterizerState : public FRHIRasterizerState
+{
+public:
+	D3D12_RASTERIZER_DESC Desc;
+};
+
+class FD3D12DepthStencilState : public FRHIDepthStencilState
+{
+public:
+
+	D3D12_DEPTH_STENCIL_DESC Desc;
+
+	/* Describes the read/write state of the separate depth and stencil components of the DSV. */
+	FExclusiveDepthStencil AccessType;
+};
+
+class FD3D12BlendState : public FRHIBlendState
+{
+public:
+
+	D3D12_BLEND_DESC Desc;
+};
+
