UE4 - D3D12RHI (not working)

diff --git a/Engine/Source/Runtime/Windows/D3D12RHI/Public/D3D12Util.h b/Engine/Source/Runtime/Windows/D3D12RHI/Public/D3D12Util.h
new file mode 100644
--- /dev/null
+++ b/Engine/Source/Runtime/Windows/D3D12RHI/Public/D3D12Util.h
@@ -0,0 +1,653 @@
+// Copyright 1998-2014 Epic Games, Inc. All Rights Reserved.
+
+/*=============================================================================
+	D3D12Util.h: D3D RHI utility definitions.
+	=============================================================================*/
+
+#pragma once
+
+namespace D3D12RHI
+{
+	/**
+	 * Checks that the given result isn't a failure.  If it is, the application exits with an appropriate error message.
+	 * @param	Result - The result code to check
+	 * @param	Code - The code which yielded the result.
+	 * @param	Filename - The filename of the source file containing Code.
+	 * @param	Line - The line number of Code within Filename.
+	 */
+	extern void VerifyD3D11Result(HRESULT Result, const ANSICHAR* Code, const ANSICHAR* Filename, uint32 Line, ID3D12Device* Device);
+}
+
+namespace D3D12RHI
+{
+	/**
+	* Checks that the given result isn't a failure.  If it is, the application exits with an appropriate error message.
+	* @param	Result - The result code to check
+	* @param	Code - The code which yielded the result.
+	* @param	Filename - The filename of the source file containing Code.
+	* @param	Line - The line number of Code within Filename.
+	*/
+	extern void VerifyD3D12CreateTextureResult(HRESULT D3DResult, const ANSICHAR* Code, const ANSICHAR* Filename, uint32 Line,
+		uint32 SizeX, uint32 SizeY, uint32 SizeZ, uint8 D3DFormat, uint32 NumMips, uint32 Flags);
+
+	/**
+	 * A macro for using VERIFYD3D11RESULT that automatically passes in the code and filename/line.
+	 */
+#define VERIFYD3D11RESULT_EX(x, Device)	{HRESULT hr = x; if (FAILED(hr)) { VerifyD3D11Result(hr,#x,__FILE__,__LINE__, Device); }}
+#define VERIFYD3D11RESULT(x)			{HRESULT hr = x; if (FAILED(hr)) { VerifyD3D11Result(hr,#x,__FILE__,__LINE__, 0); }}
+#define VERIFYD3D11CREATETEXTURERESULT(x,SizeX,SizeY,SizeZ,Format,NumMips,Flags) {HRESULT hr = x; if (FAILED(hr)) { VerifyD3D12CreateTextureResult(hr,#x,__FILE__,__LINE__,SizeX,SizeY,SizeZ,Format,NumMips,Flags); }}
+
+	/**
+	 * Checks that a COM object has the expected number of references.
+	 */
+	extern void VerifyComRefCount(IUnknown* Object, int32 ExpectedRefs, const TCHAR* Code, const TCHAR* Filename, int32 Line);
+#define checkComRefCount(Obj,ExpectedRefs) VerifyComRefCount(Obj,ExpectedRefs,TEXT(#Obj),TEXT(__FILE__),__LINE__)
+
+	/** Returns a string for the provided DXGI format. */
+	const TCHAR* GetD3D11TextureFormatString(DXGI_FORMAT TextureFormat);
+}
+
+using namespace D3D12RHI;
+
+/**
+* Convert from ECubeFace to D3DCUBEMAP_FACES type
+* @param Face - ECubeFace type to convert
+* @return D3D cube face enum value
+*/
+FORCEINLINE uint32 GetD3D11CubeFace(ECubeFace Face)
+{
+	switch (Face)
+	{
+	case CubeFace_PosX:
+	default:
+		return 0;//D3DCUBEMAP_FACE_POSITIVE_X;
+	case CubeFace_NegX:
+		return 1;//D3DCUBEMAP_FACE_NEGATIVE_X;
+	case CubeFace_PosY:
+		return 2;//D3DCUBEMAP_FACE_POSITIVE_Y;
+	case CubeFace_NegY:
+		return 3;//D3DCUBEMAP_FACE_NEGATIVE_Y;
+	case CubeFace_PosZ:
+		return 4;//D3DCUBEMAP_FACE_POSITIVE_Z;
+	case CubeFace_NegZ:
+		return 5;//D3DCUBEMAP_FACE_NEGATIVE_Z;
+	};
+}
+
+/**
+* Calculate a subresource index for a texture
+*/
+FORCEINLINE uint32 CalcSubresource(uint32 MipSlice, uint32 ArraySlice, uint32 MipLevels)
+{
+	return MipSlice + ArraySlice * MipLevels;
+}
+
+class FD3D12Resource;
+
+/**
+ * Keeps track of Locks for D3D11 objects
+ */
+class FD3D12LockedKey
+{
+public:
+	void* SourceObject;
+	uint32 Subresource;
+
+public:
+	FD3D12LockedKey() : SourceObject(NULL)
+		, Subresource(0)
+	{}
+	FD3D12LockedKey(FD3D12Resource* source, uint32 subres = 0) : SourceObject((void*)source)
+		, Subresource(subres)
+	{}
+	FD3D12LockedKey(class FD3D12ResourceLocation* source, uint32 subres = 0) : SourceObject((void*)source)
+		, Subresource(subres)
+	{}
+	bool operator==(const FD3D12LockedKey& Other) const
+	{
+		return SourceObject == Other.SourceObject && Subresource == Other.Subresource;
+	}
+	bool operator!=(const FD3D12LockedKey& Other) const
+	{
+		return SourceObject != Other.SourceObject || Subresource != Other.Subresource;
+	}
+	FD3D12LockedKey& operator=(const FD3D12LockedKey& Other)
+	{
+		SourceObject = Other.SourceObject;
+		Subresource = Other.Subresource;
+		return *this;
+	}
+	uint32 GetHash() const
+	{
+		return PointerHash(SourceObject);
+	}
+
+	/** Hashing function. */
+	friend uint32 GetTypeHash(const FD3D12LockedKey& K)
+	{
+		return K.GetHash();
+	}
+};
+
+/** Information about a D3D resource that is currently locked. */
+struct FD3D12LockedData
+{
+	TRefCountPtr<FD3D12Resource> StagingResource;
+	TRefCountPtr<FD3D12Resource> UploadHeapResource;
+	TRefCountPtr<FD3D12ResourceLocation> UploadHeapLocation;
+	uint32 Pitch;
+	uint32 DepthPitch;
+
+	// constructor
+	FD3D12LockedData()
+		: bAllocDataWasUsed(false)
+	{
+	}
+
+	// 16 byte alignment for best performance  (can be 30x faster than unaligned)
+	void AllocData(uint32 Size)
+	{
+		Data = (uint8*)FMemory::Malloc(Size, 16);
+		bAllocDataWasUsed = true;
+	}
+
+	// Some driver might return aligned memory so we don't enforce the alignment
+	void SetData(void* InData)
+	{
+		check(!bAllocDataWasUsed); Data = (uint8*)InData;
+	}
+
+	uint8* GetData() const
+	{
+		return Data;
+	}
+
+	// only call if AllocData() was used
+	void FreeData()
+	{
+		check(bAllocDataWasUsed);
+		FMemory::Free(Data);
+		Data = 0;
+	}
+
+private:
+	uint8* Data;
+	// then FreeData
+	bool bAllocDataWasUsed;
+};
+
+class FD3D12RenderTargetView;
+class FD3D12DepthStencilView;
+
+/**
+ * Class for retrieving render targets currently bound to the device context.
+ */
+class FD3D12BoundRenderTargets
+{
+public:
+	/** Initialization constructor: requires the state cache. */
+	explicit FD3D12BoundRenderTargets(FD3D12RenderTargetView **RTArray, uint32 NumActiveRTs, FD3D12DepthStencilView* DSView);
+
+	/** Destructor. */
+	~FD3D12BoundRenderTargets();
+
+	/** Accessors. */
+	FORCEINLINE int32 GetNumActiveTargets() const { return NumActiveTargets; }
+	FORCEINLINE FD3D12RenderTargetView* GetRenderTargetView(int32 TargetIndex) { return RenderTargetViews[TargetIndex]; }
+	FORCEINLINE FD3D12DepthStencilView* GetDepthStencilView() { return DepthStencilView; }
+
+private:
+	/** Active render target views. */
+	FD3D12RenderTargetView* RenderTargetViews[MaxSimultaneousRenderTargets];
+
+	/** Active depth stencil view. */
+	FD3D12DepthStencilView* DepthStencilView;
+
+	/** The number of active render targets. */
+	int32 NumActiveTargets;
+};
+
+/**
+ * Class for managing dynamic buffers.
+ */
+class FD3D12DynamicBuffer : public FRenderResource, public FRHIResource, public FD3D12DeviceChild
+{
+public:
+	/** Initialization constructor. */
+	FD3D12DynamicBuffer(FD3D12Device* InParent, class FD3D12DynamicHeapAllocator& UploadHeap);
+	/** Destructor. */
+	~FD3D12DynamicBuffer();
+
+	/** Locks the buffer returning at least Size bytes. */
+	void* Lock(uint32 Size);
+	/** Unlocks the buffer returning the underlying D3D12 buffer to use as a resource. */
+	FD3D12ResourceLocation* Unlock();
+
+	// Begin FRenderResource interface.
+	virtual void InitRHI() override;
+	virtual void ReleaseRHI() override;
+	// End FRenderResource interface.
+
+private:
+	TRefCountPtr<FD3D12ResourceLocation> ResourceLocation;
+	class FD3D12DynamicHeapAllocator& UploadHeapAllocator;
+};
+
+static D3D12_DESCRIPTOR_HEAP_DESC CreateDHD (D3D12_DESCRIPTOR_HEAP_TYPE Type, uint32 NumDescriptorsPerHeap, D3D12_DESCRIPTOR_HEAP_FLAGS Flags)
+{
+	D3D12_DESCRIPTOR_HEAP_DESC DHD = {Type, NumDescriptorsPerHeap, Flags};
+
+	return DHD;
+}
+
+class FDescriptorHeapManager
+{
+public: // Types
+	typedef D3D12_CPU_DESCRIPTOR_HANDLE HeapOffset;
+	typedef decltype(HeapOffset::ptr) HeapOffsetRaw;
+	typedef uint32 HeapIndex;
+
+private: // Types
+	struct SFreeRange { HeapOffsetRaw Start; HeapOffsetRaw End; };
+	struct SHeapEntry
+	{
+		TRefCountPtr<ID3D12DescriptorHeap> m_Heap;
+		TDoubleLinkedList<SFreeRange> m_FreeList;
+
+		SHeapEntry() { }
+	};
+	typedef TArray<SHeapEntry> THeapMap;
+
+public: // Methods
+	FDescriptorHeapManager(D3D12_DESCRIPTOR_HEAP_TYPE Type,
+		uint32 NumDescriptorsPerHeap)
+		: m_Desc(CreateDHD (Type, NumDescriptorsPerHeap, D3D12_DESCRIPTOR_HEAP_FLAG_NONE))
+		, m_DescriptorSize(0)
+		, m_pDevice(nullptr)
+	{
+	}
+
+	void Init(ID3D12Device* pDevice)
+	{
+		m_pDevice = pDevice;
+		m_DescriptorSize = pDevice->GetDescriptorHandleIncrementSize(m_Desc.Type);
+	}
+
+	HeapOffset AllocateHeapSlot(HeapIndex &outIndex)
+	{
+		FScopeLock Lock(&CritSect);
+		if (0 == m_FreeHeaps.Num())
+		{
+			AllocateHeap();
+		}
+		check(0 != m_FreeHeaps.Num());
+		auto Head = m_FreeHeaps.GetHead();
+		outIndex = Head->GetValue();
+		SHeapEntry &HeapEntry = m_Heaps[outIndex];
+		check(0 != HeapEntry.m_FreeList.Num());
+		SFreeRange &Range = HeapEntry.m_FreeList.GetHead()->GetValue();
+		HeapOffset Ret = { Range.Start };
+		Range.Start += m_DescriptorSize;
+
+		if (Range.Start == Range.End)
+		{
+			HeapEntry.m_FreeList.RemoveNode(HeapEntry.m_FreeList.GetHead());
+			if (0 == HeapEntry.m_FreeList.Num())
+			{
+				m_FreeHeaps.RemoveNode(Head);
+			}
+		}
+		return Ret;
+	}
+
+	void FreeHeapSlot(HeapOffset Offset, HeapIndex index)
+	{
+		FScopeLock Lock(&CritSect);
+		SHeapEntry &HeapEntry = m_Heaps[index];
+
+		SFreeRange NewRange = 
+		{
+			Offset.ptr,
+			Offset.ptr + m_DescriptorSize
+		};
+
+		bool bFound = false;
+		for (auto Node = HeapEntry.m_FreeList.GetHead();
+				Node != nullptr && !bFound;
+				Node = Node->GetNextNode())
+		{
+			SFreeRange &Range = Node->GetValue();
+			check(Range.Start < Range.End);
+			if (Range.Start == Offset.ptr + m_DescriptorSize)
+			{
+				Range.Start = Offset.ptr;
+				bFound = true;
+			}
+			else if (Range.End == Offset.ptr)
+			{
+				Range.End += m_DescriptorSize;
+				bFound = true;
+			}
+			else
+			{
+				check(Range.End < Offset.ptr || Range.Start > Offset.ptr);
+				if (Range.Start > Offset.ptr)
+				{
+					HeapEntry.m_FreeList.InsertNode(NewRange, Node);
+					bFound = true;
+				}
+			}
+		}
+
+		if (!bFound)
+		{
+			if (0 == HeapEntry.m_FreeList.Num())
+			{
+				m_FreeHeaps.AddTail(index);
+			}
+			HeapEntry.m_FreeList.AddTail(NewRange);
+		}
+	}
+
+private: // Methods
+	void AllocateHeap()
+	{
+		TRefCountPtr<ID3D12DescriptorHeap> Heap;
+		VERIFYD3D11RESULT( m_pDevice->CreateDescriptorHeap(&m_Desc, IID_PPV_ARGS(Heap.GetInitReference())) );
+		HeapOffset HeapBase = Heap->GetCPUDescriptorHandleForHeapStart();
+
+		// Allocate and initialize a single new entry in the map
+		m_Heaps.SetNum(m_Heaps.Num() + 1); 
+		SHeapEntry& HeapEntry = m_Heaps.Last();
+		HeapEntry.m_FreeList.AddTail({HeapBase.ptr,
+			HeapBase.ptr + m_Desc.NumDescriptors * m_DescriptorSize});
+		HeapEntry.m_Heap = Heap;
+		m_FreeHeaps.AddTail(m_Heaps.Num() - 1);
+	}
+
+private: // Members
+	const D3D12_DESCRIPTOR_HEAP_DESC m_Desc;
+	uint32 m_DescriptorSize;
+	ID3D12Device* m_pDevice; // weak-ref
+
+	THeapMap m_Heaps;
+	TDoubleLinkedList<HeapIndex> m_FreeHeaps;
+	FCriticalSection CritSect;
+};
+
+void LogExecuteCommandLists(uint32 NumCommandLists, ID3D12CommandList *const *ppCommandLists)
+{
+	for (uint32 i = 0; i < NumCommandLists; i++)
+	{
+		ID3D12CommandList *const pCurrentCommandList = ppCommandLists[i];
+		UE_LOG(LogD3D12RHI, Log, TEXT("*** EXECUTE (CmdList: %016llX) %u/%u ***"), pCurrentCommandList, i + 1, NumCommandLists);
+	}
+}
+
+FString ConvertToResourceStateString(uint32 ResourceState)
+{
+	if (ResourceState == 0)
+	{
+		return FString(TEXT("D3D12_RESOURCE_STATE_COMMON"));
+	}
+	
+	TCHAR *ResourceStateNames[] =
+	{
+		TEXT("D3D12_RESOURCE_STATE_VERTEX_AND_CONSTANT_BUFFER"	  ), 
+		TEXT("D3D12_RESOURCE_STATE_INDEX_BUFFER"				  ), 
+		TEXT("D3D12_RESOURCE_STATE_RENDER_TARGET"				  ), 
+		TEXT("D3D12_RESOURCE_STATE_UNORDERED_ACCESS"			  ), 
+		TEXT("D3D12_RESOURCE_STATE_DEPTH_WRITE"					  ), 
+		TEXT("D3D12_RESOURCE_STATE_DEPTH_READ"					  ), 
+		TEXT("D3D12_RESOURCE_STATE_NON_PIXEL_SHADER_RESOURCE"	  ), 
+		TEXT("D3D12_RESOURCE_STATE_PIXEL_SHADER_RESOURCE"		  ), 
+		TEXT("D3D12_RESOURCE_STATE_STREAM_OUT"					  ), 
+		TEXT("D3D12_RESOURCE_STATE_INDIRECT_ARGUMENT"			  ), 
+		TEXT("D3D12_RESOURCE_STATE_COPY_DEST"					  ), 
+		TEXT("D3D12_RESOURCE_STATE_COPY_SOURCE"					  ), 
+		TEXT("D3D12_RESOURCE_STATE_RESOLVE_DEST"				  ), 
+		TEXT("D3D12_RESOURCE_STATE_RESOLVE_SOURCE"				  ), 
+	};
+
+	FString ResourceStateString;
+	uint16 NumStates = 0;
+	for (uint16 i = 0; ResourceState && i < ARRAYSIZE(ResourceStateNames); i++)
+	{
+		if (ResourceState & 1)
+		{
+			if (NumStates > 0)
+			{
+				ResourceStateString += " | ";
+			}
+
+			ResourceStateString += ResourceStateNames[i];
+			NumStates++;
+		}
+		ResourceState = ResourceState >> 1;
+	}
+	return ResourceStateString;
+}
+
+void LogResourceBarriers(uint32 NumBarriers, D3D12_RESOURCE_BARRIER *pBarriers, ID3D12CommandList *const pCommandList)
+{
+	for (uint32 i = 0; i < NumBarriers; i++)
+	{
+		D3D12_RESOURCE_BARRIER &currentBarrier = pBarriers[i];
+		check(currentBarrier.Type == D3D12_RESOURCE_BARRIER_TYPE_TRANSITION);
+
+		FString StateBefore = ConvertToResourceStateString(static_cast<uint32>(currentBarrier.Transition.StateBefore));
+		FString StateAfter = ConvertToResourceStateString(static_cast<uint32>(currentBarrier.Transition.StateAfter));
+
+		UE_LOG(LogD3D12RHI, Log, TEXT("*** BARRIER (CmdList: %016llX) %u/%u: %016llX (Sub: %u), %s -> %s"), pCommandList, i + 1, NumBarriers,
+			currentBarrier.Transition.pResource,
+			currentBarrier.Transition.Subresource,
+			*StateBefore,
+			*StateAfter);
+	}
+}
+
+// Custom resource states
+// To Be Determined (TBD) means we need to fill out a resource barrier before the command list is executed.
+#define D3D12_RESOURCE_STATE_TBD (D3D12_RESOURCE_STATES)-1
+#define D3D12_RESOURCE_STATE_CORRUPT (D3D12_RESOURCE_STATES)-2
+
+//==================================================================================================================================
+// CResourceState
+// Tracking of per-resource or per-subresource state
+//==================================================================================================================================
+class CResourceState
+{
+public:
+	void Initialize(uint32 SubresourceCount);
+
+	bool AreAllSubresourcesSame() const;
+	bool CheckResourceState(D3D12_RESOURCE_STATES State) const;
+	bool CheckResourceStateInitalized() const;
+	D3D12_RESOURCE_STATES GetSubresourceState(uint32 SubresourceIndex) const;
+	void SetResourceState(D3D12_RESOURCE_STATES State);
+	void SetSubresourceState(uint32 SubresourceIndex, D3D12_RESOURCE_STATES State);
+
+private:
+	// Only used if m_AllSubresourcesSame is 1.
+	// Bits defining the state of the full resource, bits are from D3D12_RESOURCE_STATES
+	D3D12_RESOURCE_STATES m_ResourceState : 31;
+
+	// Set to 1 if m_ResourceState is valid.  In this case, all subresources have the same state
+	// Set to 0 if m_SubresourceState is valid.  In this case, each subresources may have a different state (or may be unknown)
+	uint32 m_AllSubresourcesSame : 1;
+
+	// Only used if m_AllSubresourcesSame is 0.
+	// The state of each subresources.  Bits are from D3D12_RESOURCE_STATES.
+	TArray<D3D12_RESOURCE_STATES> m_SubresourceState;
+};
+
+//==================================================================================================================================
+// FD3D12ShaderBytecode
+// Encapsulates D3D12 shader bytecode and creates a hash for the shader bytecode
+//==================================================================================================================================
+struct ShaderBytecodeHash
+{
+	// 160 bit strong SHA1 hash
+	uint32 SHA1Hash[5];
+
+	bool operator ==(const ShaderBytecodeHash &b) const
+	{
+		return (SHA1Hash[0] == b.SHA1Hash[0] &&
+			SHA1Hash[1] == b.SHA1Hash[1] &&
+			SHA1Hash[2] == b.SHA1Hash[2] &&
+			SHA1Hash[3] == b.SHA1Hash[3] &&
+			SHA1Hash[4] == b.SHA1Hash[4]);
+	}
+
+	bool operator !=(const ShaderBytecodeHash &b) const
+	{
+		return (SHA1Hash[0] != b.SHA1Hash[0] ||
+			SHA1Hash[1] != b.SHA1Hash[1] ||
+			SHA1Hash[2] != b.SHA1Hash[2] ||
+			SHA1Hash[3] != b.SHA1Hash[3] ||
+			SHA1Hash[4] != b.SHA1Hash[4]);
+	}
+};
+
+class FD3D12ShaderBytecode
+{
+public:
+	 FD3D12ShaderBytecode()
+	{
+		 FMemory::Memzero(&Shader, sizeof(Shader));
+		 FMemory::Memset(&Hash, 0, sizeof(Hash));
+	}
+
+	 FD3D12ShaderBytecode(const D3D12_SHADER_BYTECODE &InShader) :
+		Shader(InShader)
+	 {
+		 HashShader();
+	 }
+
+	void SetShaderBytecode(const D3D12_SHADER_BYTECODE &InShader)
+	{
+		Shader = InShader;
+		HashShader();
+	}
+
+	const D3D12_SHADER_BYTECODE& GetShaderBytecode() const { return Shader; }
+	ShaderBytecodeHash GetHash() const { return Hash; }
+
+private:
+	void HashShader()
+	{
+		FMemory::Memset(&Hash, 0, sizeof(Hash));
+		if (Shader.pShaderBytecode && Shader.BytecodeLength > 0)
+		{
+			FSHA1::HashBuffer(Shader.pShaderBytecode, Shader.BytecodeLength, (uint8*)Hash.SHA1Hash);
+		}
+	}
+
+private:
+	ShaderBytecodeHash Hash;
+	D3D12_SHADER_BYTECODE Shader;
+};
+
+struct FD3D12LowLevelGraphicsPipelineStateDesc
+{
+	D3D12_GRAPHICS_PIPELINE_STATE_DESC Desc;
+	ShaderBytecodeHash VSHash;
+	ShaderBytecodeHash HSHash;
+	ShaderBytecodeHash DSHash;
+	ShaderBytecodeHash GSHash;
+	ShaderBytecodeHash PSHash;
+
+	SIZE_T CombinedHash;
+};
+
+class FD3D12BoundShaderState; // forward-declare
+
+struct FD3D12HighLevelGraphicsPipelineStateDesc
+{
+	FD3D12BoundShaderState* BoundShaderState;
+	D3D12_BLEND_DESC* BlendState;
+	D3D12_DEPTH_STENCIL_DESC* DepthStencilState;
+	D3D12_RASTERIZER_DESC* RasterizerState;
+	// IBStripCutValue unused
+	uint32 SampleMask;
+	D3D12_PRIMITIVE_TOPOLOGY_TYPE PrimitiveTopologyType;
+	uint32 NumRenderTargets;
+	DXGI_FORMAT RTVFormats[ 8 ];
+	DXGI_FORMAT DSVFormat;
+	DXGI_SAMPLE_DESC SampleDesc;
+
+	SIZE_T CombinedHash; // Pre-computed hash
+
+	void GetLowLevelDesc(ID3D12RootSignature* pRS, FD3D12LowLevelGraphicsPipelineStateDesc& psoDesc);
+};
+
+struct FD3D12ComputePipelineStateDesc
+{
+	D3D12_COMPUTE_PIPELINE_STATE_DESC Desc;
+	ShaderBytecodeHash CSHash;
+
+	SIZE_T CombinedHash;
+};
+
+
+/**
+ * The base class of threadsafe reference counted objects.
+ */
+template <class Type>
+struct FThreadsafeQueue
+{
+private:
+	mutable FCriticalSection	SynchronizationObject; // made this mutable so this class can have const functions and still be thread safe
+	TQueue<Type>				Items;
+public:
+	void Enqueue(const Type& Item)
+	{
+		FScopeLock ScopeLock(&SynchronizationObject);
+		Items.Enqueue(Item);
+	}
+
+	bool Dequeue(Type& Result)
+	{
+		FScopeLock ScopeLock(&SynchronizationObject);
+		
+		return Items.Dequeue (Result);
+	}
+
+	template <typename CompareFunc>
+	bool Dequeue(Type& Result, CompareFunc Func)
+	{
+		FScopeLock ScopeLock(&SynchronizationObject);
+
+		if (Items.Peek (Result))
+		{
+			if (Func (Result))
+			{
+				Items.Dequeue (Result);
+				
+				return true;
+			}
+		}
+
+		return false;
+	}
+
+	bool Peek (Type& Result)
+	{
+		FScopeLock ScopeLock( &SynchronizationObject );
+		return Items.Peek (Result);
+	}
+
+	bool IsEmpty()
+	{
+		FScopeLock ScopeLock( &SynchronizationObject );
+		return Items.IsEmpty();
+	}
+
+	void Empty()
+	{
+		FScopeLock ScopeLock( &SynchronizationObject );
+
+		Type Result;
+		while (Items.Dequeue (Result)) {}
+	}
+};
\ No newline at end of file
diff --git a/Engine/Source/Runtime/Windows/D3D12RHI/Public/D3D12Viewport.h b/Engine/Source/Runtime/Windows/D3D12RHI/Public/D3D12Viewport.h
new file mode 100644
--- /dev/null
+++ b/Engine/Source/Runtime/Windows/D3D12RHI/Public/D3D12Viewport.h
@@ -0,0 +1,111 @@
+// Copyright 1998-2014 Epic Games, Inc. All Rights Reserved.
+
+/*=============================================================================
+	D3D12Viewport.h: D3D viewport RHI definitions.
+=============================================================================*/
+
+#pragma once
+
+static DXGI_FORMAT GetRenderTargetFormat(EPixelFormat PixelFormat)
+{
+	DXGI_FORMAT	DXFormat = (DXGI_FORMAT)GPixelFormats[PixelFormat].PlatformFormat;
+	switch (DXFormat)
+	{
+	case DXGI_FORMAT_B8G8R8A8_TYPELESS:		return DXGI_FORMAT_B8G8R8A8_UNORM;
+	case DXGI_FORMAT_BC1_TYPELESS:			return DXGI_FORMAT_BC1_UNORM;
+	case DXGI_FORMAT_BC2_TYPELESS:			return DXGI_FORMAT_BC2_UNORM;
+	case DXGI_FORMAT_BC3_TYPELESS:			return DXGI_FORMAT_BC3_UNORM;
+	case DXGI_FORMAT_R16_TYPELESS:			return DXGI_FORMAT_R16_UNORM;
+	case DXGI_FORMAT_R8G8B8A8_TYPELESS:		return DXGI_FORMAT_R8G8B8A8_UNORM;
+	default: 								return DXFormat;
+	}
+}
+
+class FD3D12Viewport : public FRHIViewport, public FD3D12DeviceChild
+{
+public:
+
+	FD3D12Viewport(class FD3D12Device* InParent, HWND InWindowHandle, uint32 InSizeX, uint32 InSizeY, bool bInIsFullscreen, EPixelFormat InPixelFormat);
+
+    void Init(IDXGIFactory4* Factory, bool AssociateWindow = true);
+
+	~FD3D12Viewport();
+
+	void Resize(uint32 InSizeX, uint32 InSizeY, bool bInIsFullscreen);
+
+	/**
+	 * If the swap chain has been invalidated by DXGI, resets the swap chain to the expected state; otherwise, does nothing.
+	 * Called once/frame by the game thread on all viewports.
+	 * @param bIgnoreFocus - Whether the reset should happen regardless of whether the window is focused.
+	 */
+	void ConditionalResetSwapChain(bool bIgnoreFocus);
+
+	/** Presents the swap chain.
+	 * Returns true if Present was done by Engine.
+	 */
+	bool Present(bool bLockToVsync);
+
+	// Accessors.
+	FIntPoint GetSizeXY() const { return FIntPoint(SizeX, SizeY); }
+	SIZE_T GetBackBufferIndex() const { return SwapChain3->GetCurrentBackBufferIndex(); }
+
+	FD3D12Texture2D* GetBackBuffer() const { return BackBuffers[GetBackBufferIndex()].GetReference(); }
+	FD3D12Texture2D* GetBackBuffer(uint32 Index) const { return BackBuffers[Index % NumBackBuffers].GetReference(); }
+
+	IDXGISwapChain3* GetSwapChain() const { return SwapChain3; }
+
+	virtual void* GetNativeSwapChain() const override { return GetSwapChain(); }
+	virtual void* GetNativeBackBufferTexture() const override { return GetBackBuffer()->GetResource(); }
+	virtual void* GetNativeBackBufferRT() const override { return GetBackBuffer()->GetRenderTargetView(0, 0); }
+
+	virtual void SetCustomPresent(FRHICustomPresent* InCustomPresent) override
+	{
+		CustomPresent = InCustomPresent;
+	}
+	virtual FRHICustomPresent* GetCustomPresent() const { return CustomPresent; }
+
+	virtual void* GetNativeWindow(void** AddParam = nullptr) const override { return (void*)WindowHandle; }
+
+    uint32 GetNumBackBuffers(){ return NumBackBuffers; }
+
+private:
+
+	/** Presents the frame synchronizing with DWM. */
+	void PresentWithVsyncDWM();
+
+	/**
+	 * Presents the swap chain checking the return result.
+	 * Returns true if Present was done by Engine.
+	 */
+	bool PresentChecked(int32 SyncInterval);
+
+	uint64 LastFlipTime;
+	uint64 LastFrameComplete;
+	uint64 LastCompleteTime;
+	int32 SyncCounter;
+	bool bSyncedLastFrame;
+	HWND WindowHandle;
+    uint32 MaximumFrameLatency;
+	uint32 SizeX;
+	uint32 SizeY;
+	bool bIsFullscreen;
+	EPixelFormat PixelFormat;
+	bool bIsValid;
+	TRefCountPtr<IDXGISwapChain3> SwapChain3;
+
+	static const uint32 MaxNumBackBuffers = 6;
+	static const uint32 DefaultNumBackBuffers = 3;
+	TRefCountPtr<FD3D12Texture2D> BackBuffers[MaxNumBackBuffers];
+	uint32 NumBackBuffers;
+	bool bIsBenchmarkMode;
+
+	FCustomPresentRHIRef CustomPresent;
+
+	DXGI_MODE_DESC SetupDXGI_MODE_DESC() const;
+};
+
+template<>
+struct TD3D12ResourceTraits<FRHIViewport>
+{
+	typedef FD3D12Viewport TConcreteType;
+};
