UE4 - WinXP compile fix (UE-18589)

diff --git a/Engine/Source/Runtime/RHI/RHI.Build.cs b/Engine/Source/Runtime/RHI/RHI.Build.cs
--- a/Engine/Source/Runtime/RHI/RHI.Build.cs
+++ b/Engine/Source/Runtime/RHI/RHI.Build.cs
@@ -16,12 +16,17 @@ public class RHI : ModuleRules
 			if ((Target.Platform == UnrealTargetPlatform.Win32) || (Target.Platform == UnrealTargetPlatform.Win64))
 			{
 				DynamicallyLoadedModuleNames.Add("D3D11RHI");
-				DynamicallyLoadedModuleNames.Add("D3D12RHI");
-//#todo-rco: Remove when public
-				string VulkanSDKPath = Environment.GetEnvironmentVariable("VulkanSDK");
-				if (!String.IsNullOrEmpty(VulkanSDKPath))
+				if (!WindowsPlatform.IsWindowsXPSupported())
 				{
-					DynamicallyLoadedModuleNames.Add("VulkanRHI");
+					//#todo-rco: D3D12 requires different SDK headers not compatible with WinXP
+					DynamicallyLoadedModuleNames.Add("D3D12RHI");
+
+//#todo-rco: Remove when public
+					string VulkanSDKPath = Environment.GetEnvironmentVariable("VulkanSDK");
+					if (!String.IsNullOrEmpty(VulkanSDKPath))
+					{
+						DynamicallyLoadedModuleNames.Add("VulkanRHI");
+					}
 				}
 			}
 
diff --git a/Engine/Source/Runtime/Windows/D3D11RHI/Private/D3D11Device.cpp b/Engine/Source/Runtime/Windows/D3D11RHI/Private/D3D11Device.cpp
--- a/Engine/Source/Runtime/Windows/D3D11RHI/Private/D3D11Device.cpp
+++ b/Engine/Source/Runtime/Windows/D3D11RHI/Private/D3D11Device.cpp
@@ -7,9 +7,6 @@
 #include "D3D11RHIPrivate.h"
 #include "AllowWindowsPlatformTypes.h"
 	#include <delayimp.h>
-	#ifndef DXGI_ERROR_NOT_FOUND
-	#define DXGI_ERROR_NOT_FOUND             _HRESULT_TYPEDEF_(0x887A0002L)
-	#endif
 #include "HideWindowsPlatformTypes.h"
 
 
diff --git a/Engine/Source/Runtime/Windows/D3D11RHI/Public/D3D11Util.h b/Engine/Source/Runtime/Windows/D3D11RHI/Public/D3D11Util.h
--- a/Engine/Source/Runtime/Windows/D3D11RHI/Public/D3D11Util.h
+++ b/Engine/Source/Runtime/Windows/D3D11RHI/Public/D3D11Util.h
@@ -6,6 +6,27 @@
 
 #pragma once
 
+#if WINVER == 0x0502
+// Windows XP uses Win7 sdk, and in that one winerror.h doesn't include them
+
+#define DXGI_ERROR_INVALID_CALL                 MAKE_DXGI_HRESULT(1)
+#define DXGI_ERROR_NOT_FOUND                    MAKE_DXGI_HRESULT(2)
+#define DXGI_ERROR_MORE_DATA                    MAKE_DXGI_HRESULT(3)
+#define DXGI_ERROR_UNSUPPORTED                  MAKE_DXGI_HRESULT(4)
+#define DXGI_ERROR_DEVICE_REMOVED               MAKE_DXGI_HRESULT(5)
+#define DXGI_ERROR_DEVICE_HUNG                  MAKE_DXGI_HRESULT(6)
+#define DXGI_ERROR_DEVICE_RESET                 MAKE_DXGI_HRESULT(7)
+#define DXGI_ERROR_WAS_STILL_DRAWING            MAKE_DXGI_HRESULT(10)
+#define DXGI_ERROR_FRAME_STATISTICS_DISJOINT    MAKE_DXGI_HRESULT(11)
+#define DXGI_ERROR_GRAPHICS_VIDPN_SOURCE_IN_USE MAKE_DXGI_HRESULT(12)
+#define DXGI_ERROR_DRIVER_INTERNAL_ERROR        MAKE_DXGI_HRESULT(32)
+#define DXGI_ERROR_NONEXCLUSIVE                 MAKE_DXGI_HRESULT(33)
+#define DXGI_ERROR_NOT_CURRENTLY_AVAILABLE      MAKE_DXGI_HRESULT(34)
+#define DXGI_ERROR_REMOTE_CLIENT_DISCONNECTED   MAKE_DXGI_HRESULT(35)
+#define DXGI_ERROR_REMOTE_OUTOFMEMORY           MAKE_DXGI_HRESULT(36)
+
+#endif
+
 /**
  * Checks that the given result isn't a failure.  If it is, the application exits with an appropriate error message.
  * @param	Result - The result code to check
