UE4 - Out of NotForLicensees

diff --git a/Engine/Source/Runtime/Windows/D3D12RHI/D3D12RHI.Build.cs b/Engine/Source/Runtime/Windows/D3D12RHI/D3D12RHI.Build.cs
--- a/Engine/Source/Runtime/Windows/D3D12RHI/D3D12RHI.Build.cs
+++ b/Engine/Source/Runtime/Windows/D3D12RHI/D3D12RHI.Build.cs
@@ -6,7 +6,6 @@ public class D3D12RHI : ModuleRules
 {
 	public D3D12RHI(TargetInfo Target)
 	{
-		BinariesSubFolder = "NotForLicensees";
 		PrivateIncludePaths.Add("Runtime/Windows/D3D12RHI/Private");
 
 		PrivateDependencyModuleNames.AddRange(
