UE4 - Move D3D12 PSO cache to Saved

diff --git a/Engine/Source/Runtime/Windows/D3D12RHI/Private/D3D12Device.cpp b/Engine/Source/Runtime/Windows/D3D12RHI/Private/D3D12Device.cpp
--- a/Engine/Source/Runtime/Windows/D3D12RHI/Private/D3D12Device.cpp
+++ b/Engine/Source/Runtime/Windows/D3D12RHI/Private/D3D12Device.cpp
@@ -718,8 +718,8 @@ void FD3D12Device::SetupAfterDeviceCreation()
 	CreateSignatures();
 
     PipelineStateCache = FD3D12PipelineStateCache(this);
-    FString GraphicsCacheFile = FPaths::GameContentDir() / TEXT("D3DGraphics.ushaderprecache");
-    FString ComputeCacheFile = FPaths::GameContentDir() / TEXT("D3DCompute.ushaderprecache");
+    FString GraphicsCacheFile = FPaths::GameSavedDir() / TEXT("D3DGraphics.ushaderprecache");
+    FString ComputeCacheFile = FPaths::GameSavedDir() / TEXT("D3DCompute.ushaderprecache");
 
 	PipelineStateCache.Init(GraphicsCacheFile, ComputeCacheFile);
 	PipelineStateCache.RebuildFromDiskCache(GraphicsRS, ComputeRS);
