UE4 - add some additional checks to pooled uniform buffers. 4.8 is showing some at a demo

diff --git a/Engine/Source/Runtime/Windows/D3D11RHI/Private/D3D11UniformBuffer.cpp b/Engine/Source/Runtime/Windows/D3D11RHI/Private/D3D11UniformBuffer.cpp
--- a/Engine/Source/Runtime/Windows/D3D11RHI/Private/D3D11UniformBuffer.cpp
+++ b/Engine/Source/Runtime/Windows/D3D11RHI/Private/D3D11UniformBuffer.cpp
@@ -80,7 +80,15 @@ void UniformBufferBeginFrame()
 	// Merge the bucket into the free pool array
 	for (int32 BucketIndex = 0; BucketIndex < NumPoolBuckets; BucketIndex++)
 	{
+		int32 LastNum = UniformBufferPool[BucketIndex].Num();
 		UniformBufferPool[BucketIndex].Append(SafeUniformBufferPools[SafeFrameIndex][BucketIndex]);
+#if DO_CHECK
+		while (LastNum < UniformBufferPool[BucketIndex].Num())
+		{
+			check(IsValidRef(UniformBufferPool[BucketIndex][LastNum].Buffer));
+			LastNum++;
+		}
+#endif
 		SafeUniformBufferPools[SafeFrameIndex][BucketIndex].Reset();
 	}
 }
@@ -215,6 +223,7 @@ FD3D11UniformBuffer::~FD3D11UniformBuffer()
 	// Do not return the allocation to the pool if it is in the dynamic constant buffer!
 	if (!RingAllocation.IsValid() && Resource != nullptr)
 	{
+		check(IsInRenderingThread());
 		D3D11_BUFFER_DESC Desc;
 		Resource->GetDesc(&Desc);
 
@@ -230,10 +239,14 @@ FD3D11UniformBuffer::~FD3D11UniformBuffer()
 			// Add to this frame's array of free uniform buffers
 			const int32 SafeFrameIndex = (GFrameNumberRenderThread - 1) % NumSafeFrames;
 			const uint32 BucketIndex = GetPoolBucketIndex(Desc.ByteWidth);
+			int32 LastNum = SafeUniformBufferPools[SafeFrameIndex][BucketIndex].Num();
 			check(Desc.ByteWidth <= GetPoolBucketSize(Desc.ByteWidth));
 			SafeUniformBufferPools[SafeFrameIndex][BucketIndex].Add(NewEntry);
 			INC_DWORD_STAT(STAT_D3D11NumFreeUniformBuffers);
 			INC_MEMORY_STAT_BY(STAT_D3D11FreeUniformBufferMemory, Desc.ByteWidth);
+
+			FPlatformMisc::MemoryBarrier(); // check for unwanted concurrency
+			check(SafeUniformBufferPools[SafeFrameIndex][BucketIndex].Num() == LastNum + 1);
 		}
 	}
 }
