UE4 - DX12 - Disable PSO cache, use r.RHI.EnablePSOCache=1 to enable

diff --git a/Engine/Source/Runtime/Windows/D3D12RHI/Private/D3D12StateCache.cpp b/Engine/Source/Runtime/Windows/D3D12RHI/Private/D3D12StateCache.cpp
--- a/Engine/Source/Runtime/Windows/D3D12RHI/Private/D3D12StateCache.cpp
+++ b/Engine/Source/Runtime/Windows/D3D12RHI/Private/D3D12StateCache.cpp
@@ -12,6 +12,14 @@
 
 #define IL_MAX_SEMANTIC_NAME 255
 
+static int32 GEnablePSOCache = 0;
+static FAutoConsoleVariableRef CVarEnablePSOCache(
+	TEXT("r.RHI.EnablePSOCache"),
+	GEnablePSOCache,
+	TEXT("Enables a disk cache for PipelineState Objects."),
+	ECVF_RenderThreadSafe | ECVF_ReadOnly
+	);
+
 extern bool D3D12RHI_ShouldCreateWithD3DDebug();
 
 inline bool operator!=(D3D12_CPU_DESCRIPTOR_HANDLE lhs, D3D12_CPU_DESCRIPTOR_HANDLE rhs)
@@ -694,9 +702,13 @@ void FD3D12StateCacheBase::Init(FD3D12Device* InParent, FD3D12CommandContext* In
 	DescriptorCache.Init(InParent, InCmdContext);
 
 	if (AncestralState)
+	{
 		InheritState(*AncestralState);
+	}
 	else
+	{
 		ClearState();
+	}
 
 	bAlwaysSetIndexBuffers = bInAlwaysSetIndexBuffers;
 }
@@ -712,7 +724,9 @@ void FD3D12StateCacheBase::Clear()
 void FD3D12StateCacheBase::ClearSRVs()
 {
 	if (bSRVSCleared)
+	{
 		return;
+	}
 
 	ZeroMemory(PipelineState.Common.CurrentShaderResourceViewsIntersectWithDepthRT, sizeof(PipelineState.Common.CurrentShaderResourceViewsIntersectWithDepthRT));
 	PipelineState.Common.ShaderResourceViewsIntersectWithDepthCount = 0;
@@ -891,7 +905,9 @@ void FD3D12HighLevelGraphicsPipelineStateDesc::GetLowLevelDesc(ID3D12RootSignatu
 	psoDesc.Desc.InputLayout = BoundShaderState->InputLayout;
 
 	if (BoundShaderState->GetGeometryShader())
+	{
 		psoDesc.Desc.StreamOutput = BoundShaderState->GetGeometryShader()->StreamOutput;
+	}
 
 #define COPY_SHADER(Initial, Name) \
 	if (FD3D12##Name##Shader* Shader = BoundShaderState->Get##Name##Shader()) \
@@ -1202,6 +1218,12 @@ void FD3D12StateCacheBase::ApplyState(bool IsCompute)
 #endif
 }
 
+inline bool FDiskCacheInterface::IsInErrorState() const
+{
+	return !GEnablePSOCache || mInErrorState;
+}
+
+
 void FDiskCacheInterface::Init(FString &filename)
 {
 	mFileStart = nullptr;
@@ -1215,12 +1237,19 @@ void FDiskCacheInterface::Init(FString &filename)
 
 	mFileName = filename;
 	mCacheExists = true;
-	WIN32_FIND_DATA fileData;
-	FindFirstFile(mFileName.GetCharArray().GetData(), &fileData);
-	if (GetLastError() == ERROR_FILE_NOT_FOUND)
+	if (!GEnablePSOCache)
 	{
 		mCacheExists = false;
 	}
+	else
+	{
+		WIN32_FIND_DATA fileData;
+		FindFirstFile(mFileName.GetCharArray().GetData(), &fileData);
+		if (GetLastError() == ERROR_FILE_NOT_FOUND)
+		{
+			mCacheExists = false;
+		}
+	}
 	bool fileFound = mCacheExists;
 	mCurrentFileMapSize = 1;
 	GrowMapping(64 * 1024, true);
@@ -1244,26 +1273,36 @@ void FDiskCacheInterface::Init(FString &filename)
 
 void FDiskCacheInterface::GrowMapping(SIZE_T size, bool firstrun)
 {
-	if (mInErrorState)
+	if (!GEnablePSOCache || mInErrorState)
 	{
 		return;
 	}
 
 	if (mCurrentFileMapSize - mCurrentOffset < size)
 	{
-		while ((mCurrentFileMapSize - mCurrentOffset) < size){ mCurrentFileMapSize += mFileGrowSize; }
+		while ((mCurrentFileMapSize - mCurrentOffset) < size)
+		{
+			mCurrentFileMapSize += mFileGrowSize;
+		}
 	}
 	else
 	{
 		return;
 	}
 
-	if (hMapAddress){
+	if (hMapAddress)
+	{
 		FlushViewOfFile(hMapAddress, mCurrentOffset);
 		UnmapViewOfFile(hMapAddress);
 	}
-	if (hMemoryMap){ CloseHandle(hMemoryMap); }
-	if (hFile){ CloseHandle(hFile); }
+	if (hMemoryMap)
+	{
+		CloseHandle(hMemoryMap);
+	}
+	if (hFile)
+	{
+		CloseHandle(hFile);
+	}
 
 	uint32 flag = (mCacheExists) ? OPEN_EXISTING : CREATE_NEW;
 	// open the shader cache file
@@ -1382,13 +1421,20 @@ void FDiskCacheInterface::Close(uint32 numberOfPSOs)
 
 	if (!IsInErrorState())
 	{
-		if (hMapAddress){
+		if (hMapAddress)
+		{
 			*(FDiskCacheHeader*)mFileStart = mHeader;
 			FlushViewOfFile(hMapAddress, mCurrentOffset);
 			UnmapViewOfFile(hMapAddress);
 		}
-		if (hMemoryMap){ CloseHandle(hMemoryMap); }
-		if (hFile){ CloseHandle(hFile); }
+		if (hMemoryMap)
+		{
+			CloseHandle(hMemoryMap);
+		}
+		if (hFile)
+		{
+			CloseHandle(hFile);
+		}
 	}
 }
 
@@ -1400,9 +1446,23 @@ void FDiskCacheInterface::ClearDiskCache()
 	mHeader.mNumPsos = 0;
 	mHeader.mUsesAPILibraries = FD3D12PipelineStateCache::bUseAPILibaries;
 
-	if (hMapAddress){ UnmapViewOfFile(hMapAddress); }
-	if (hMemoryMap){ CloseHandle(hMemoryMap); }
-	if (hFile){ CloseHandle(hFile); }
+	if (!GEnablePSOCache)
+	{
+		return;
+	}
+
+	if (hMapAddress)
+	{
+		UnmapViewOfFile(hMapAddress);
+	}
+	if (hMemoryMap)
+	{
+		CloseHandle(hMemoryMap);
+	}
+	if (hFile)
+	{
+		CloseHandle(hFile);
+	}
 #ifdef UNICODE
 	BOOL result = DeleteFileW(mFileName.GetCharArray().GetData());
 #else
@@ -1416,7 +1476,8 @@ void FDiskCacheInterface::Flush(uint32 numberOfPSOs)
 	mHeader.mNumPsos = numberOfPSOs;
 	mHeader.mUsesAPILibraries = FD3D12PipelineStateCache::bUseAPILibaries;
 
-	if (hMapAddress && !IsInErrorState()){
+	if (hMapAddress && !IsInErrorState())
+	{
 		*(FDiskCacheHeader*)mFileStart = mHeader;
 		FlushViewOfFile(hMapAddress, mCurrentOffset);
 	}
@@ -1502,7 +1563,10 @@ void FD3D12PipelineStateCache::RebuildFromDiskCache(ID3D12RootSignature* pGraphi
 		DiskCaches[PSO_CACHE_GRAPHICS].SetPointerAndAdvanceFilePosition((void**)&cachedBlobSize, sizeof(SIZE_T));
 
 		check(cachedBlobSize);
-		if (bUseAPILibaries && cachedBlobSize) { check(*cachedBlobSize); }
+		if (bUseAPILibaries && cachedBlobSize)
+		{
+			check(*cachedBlobSize);
+		}
 
 		if (bUseAPILibaries && cachedBlobSize &&*cachedBlobSize)
 		{
@@ -1550,7 +1614,10 @@ void FD3D12PipelineStateCache::RebuildFromDiskCache(ID3D12RootSignature* pGraphi
 		DiskCaches[PSO_CACHE_COMPUTE].SetPointerAndAdvanceFilePosition((void**)&cachedBlobSize, sizeof(SIZE_T));
 
 		check(cachedBlobSize);
-		if (bUseAPILibaries && cachedBlobSize) { check(*cachedBlobSize); }
+		if (bUseAPILibaries && cachedBlobSize)
+		{
+			check(*cachedBlobSize);
+		}
 
 		if (bUseAPILibaries && cachedBlobSize &&*cachedBlobSize)
 		{
@@ -1629,7 +1696,9 @@ ID3D12PipelineState* FD3D12PipelineStateCache::FindGraphicsLowLevel(FD3D12LowLev
 	graphicsPSODesc.CombinedHash = FD3D12PipelineStateCache::HashPSODesc(graphicsPSODesc);
 
 	if (TRefCountPtr<ID3D12PipelineState>* PSO = LowLevelGraphicsPipelineStateCache.Find(graphicsPSODesc))
+	{
 		return PSO->GetReference();
+	}
 
 	return Add(graphicsPSODesc, insertIntoDiskCache);
 }
@@ -1640,7 +1709,9 @@ ID3D12PipelineState* FD3D12PipelineStateCache::FindCompute(FD3D12ComputePipeline
 	computePSODesc.CombinedHash = FD3D12PipelineStateCache::HashPSODesc(computePSODesc);
 
 	if (TRefCountPtr<ID3D12PipelineState>* PSO = ComputePipelineStateCache.Find(computePSODesc))
+	{
 		return PSO->GetReference();
+	}
 
 	return Add(computePSODesc, insertIntoDiskCache);
 }
@@ -1843,7 +1914,7 @@ void FD3D12PipelineStateCache::Close()
 
 FORCEINLINE uint32 SSE4_CRC32(void* data, SIZE_T numBytes)
 {
-	check(bCPUSupportsSSE4);
+	check(GCPUSupportsSSE4);
 	uint32 hash = 0;
 #if _WIN64
 	static const SIZE_T alignment = 8;//64 Bit
@@ -1880,7 +1951,7 @@ FORCEINLINE uint32 SSE4_CRC32(void* data, SIZE_T numBytes)
 
 SIZE_T FD3D12PipelineStateCache::HashData(void* data, SIZE_T numBytes)
 {
-	if (bCPUSupportsSSE4)
+	if (GCPUSupportsSSE4)
 	{
 		return SIZE_T(SSE4_CRC32(data, numBytes));
 	}
@@ -1945,7 +2016,7 @@ void FD3D12PipelineStateCache::Init(FString &GraphicsCacheFilename, FString &Com
 	{
 		int32 cpui[4];
 		__cpuidex(cpui, 1, 0);
-		bCPUSupportsSSE4 = !!(cpui[SSE4_CPUID_ARRAY_INDEX] & SSE4_2);
+		GCPUSupportsSSE4 = !!(cpui[SSE4_CPUID_ARRAY_INDEX] & SSE4_2);
 	}
 
 }
@@ -2149,7 +2220,9 @@ bool FD3D12StateCacheBase::VerifyViewState(ID3D12DebugCommandList* pDebugCommand
 void FD3D12StateCacheBase::SetUAVs(EShaderFrequency ShaderStage, uint32 UAVStartSlot, uint32 NumSimultaneousUAVs, FD3D12UnorderedAccessView** UAVArray, uint32 *UAVInitialCountArray)
 {
 	if (NumSimultaneousUAVs == 0)
+	{
 		return;
+	}
 	
 	PipelineState.Common.CurrentUAVStartSlot = FMath::Min(UAVStartSlot, PipelineState.Common.CurrentUAVStartSlot);
 	PipelineState.Common.CurrentUAVStage = ShaderStage;
@@ -2386,7 +2459,9 @@ void FD3D12StateCacheBase::InternalSetShaderResourceView(FD3D12ShaderResourceVie
 		}
 
 		if (AlternatePathFunction != nullptr)
+		{
 			(*AlternatePathFunction)(this, SRV, ResourceIndex, SrvType);
+		}
 	}
 }
 
@@ -2713,7 +2788,9 @@ void FD3D12DescriptorCache::SetUAVs(EShaderFrequency ShaderStage, uint32 UAVStar
 	for (uint32 i = 0; i < SlotsNeeded; i++)
 	{
 		if (UnorderedAccessViewArray[i] != NULL)
+		{
 			UnorderedAccessViewArray[i]->GetResource()->UpdateResidency();
+		}
 	}
 #endif
 
@@ -2949,8 +3026,9 @@ void FD3D12DescriptorCache::SetSRVs(EShaderFrequency ShaderStage, TRefCountPtr<F
 	D3D12_GPU_DESCRIPTOR_HANDLE* FoundDescriptor = SRVMap.Find(Desc);
 
 	if (FoundDescriptor)
+	{
 		BindDescriptor = *FoundDescriptor;
-
+	}
 	else
 	{
 		// Reserve heap slots
@@ -3025,7 +3103,9 @@ void FD3D12DescriptorCache::SetSRVs(EShaderFrequency ShaderStage, TRefCountPtr<F
 	for (uint32 i = 0; i < SlotsNeeded; i++)
 	{
 		if (SRVs[i] != NULL)
+		{
 			SRVs[i]->GetResource()->UpdateResidency();
+		}
 	}
 #endif
 
diff --git a/Engine/Source/Runtime/Windows/D3D12RHI/Private/D3D12StateCachePrivate.h b/Engine/Source/Runtime/Windows/D3D12RHI/Private/D3D12StateCachePrivate.h
--- a/Engine/Source/Runtime/Windows/D3D12RHI/Private/D3D12StateCachePrivate.h
+++ b/Engine/Source/Runtime/Windows/D3D12RHI/Private/D3D12StateCachePrivate.h
@@ -491,8 +491,8 @@ public:
 		void RollOver();
 		uint32 ReserveSlots(uint32 NumSlotsRequested);
 		void SetNextSlot(uint32 NextSlot);
-		D3D12_CPU_DESCRIPTOR_HANDLE GetCPUSlotHandle(uint32 Slot) { return{ CPUBase.ptr + Slot * DescriptorSize }; }
-		D3D12_GPU_DESCRIPTOR_HANDLE GetGPUSlotHandle(uint32 Slot) { return{ GPUBase.ptr + Slot * DescriptorSize }; }
+		D3D12_CPU_DESCRIPTOR_HANDLE GetCPUSlotHandle(uint32 Slot) const { return{ CPUBase.ptr + Slot * DescriptorSize }; }
+		D3D12_GPU_DESCRIPTOR_HANDLE GetGPUSlotHandle(uint32 Slot) const { return{ GPUBase.ptr + Slot * DescriptorSize }; }
 	};
 
 	OnlineHeap ViewHeap;
@@ -618,12 +618,11 @@ public:
 	void Close(uint32 numberOfPSOs);
 	void Flush(uint32 numberOfPSOs);
 	void ClearDiskCache();
-	uint32 GetNumPSOs(){
+	uint32 GetNumPSOs() const
+	{
 		return mHeader.mNumPsos;
 	}
-	bool IsInErrorState(){
-		return mInErrorState;
-	}
+	bool IsInErrorState() const;
 
 	~FDiskCacheInterface()
 	{
@@ -637,7 +636,7 @@ public:
 	}
 };
 
-static bool bCPUSupportsSSE4;
+static bool GCPUSupportsSSE4;
 
 class FD3D12PipelineStateCache : public FD3D12DeviceChild
 {
@@ -738,20 +737,22 @@ public:
 	void Set (uint32 Bits)
 	{
 		for (uint32 i = 0; i < Bits; ++i)
-			SetIndex (i);
+		{
+			SetIndex(i);
+		}
 	}
 
-	bool AnySet()
+	bool AnySet() const
 	{
 		return (Array != 0);
 	}
 
-	uint32 LastSet()
+	uint32 LastSet() const
 	{
 		return 32 - FMath::CountLeadingZeros (Array);
 	}
 
-	uint32 GetValue()
+	uint32 GetValue() const
 	{
 		return Array;
 	}
@@ -777,7 +778,7 @@ public:
 		ChangedIndices.SetIndex (Index);
 	}
 
-	uint32 LastSet()
+	uint32 LastSet() const
 	{
 		FD3D12BitArray Merged (SetIndices.GetValue() | ChangedIndices.GetValue());
 
@@ -945,7 +946,9 @@ protected:
 			bNeedSetSamplersPerShaderStage[ShaderFrequency] = true;
 			bNeedSetSamplers = true;
 			if (AlternatePathFunction != nullptr)
+			{
 				(*AlternatePathFunction)(this, SamplerState, SamplerIndex);
+			}
 		}
 	}
 
@@ -1105,7 +1108,7 @@ public:
 	}
 
 	template <EShaderFrequency ShaderFrequency>
-	D3D12_STATE_CACHE_INLINE void GetSamplerState(uint32 StartSamplerIndex, uint32 NumSamplerIndexes, FD3D12SamplerState** SamplerStates)
+	D3D12_STATE_CACHE_INLINE void GetSamplerState(uint32 StartSamplerIndex, uint32 NumSamplerIndexes, FD3D12SamplerState** SamplerStates) const
 	{
 		check(StartSamplerIndex + NumSamplerIndexes <= D3D12_COMMONSHADER_INPUT_RESOURCE_SLOT_COUNT);
 		for (uint32 StateLoop = 0; StateLoop < NumSamplerIndexes; StateLoop++)
@@ -1153,14 +1156,14 @@ public:
 		}
 	}
 
-	D3D12_STATE_CACHE_INLINE void GetRasterizerState(D3D12_RASTERIZER_DESC** RasterizerState)
+	D3D12_STATE_CACHE_INLINE void GetRasterizerState(D3D12_RASTERIZER_DESC** RasterizerState) const
 	{
 		*RasterizerState = PipelineState.Graphics.HighLevelDesc.RasterizerState;
 	}
 
 	D3D12_STATE_CACHE_INLINE void SetBlendState(D3D12_BLEND_DESC* State, const float BlendFactor[4], uint32 SampleMask);
 
-	D3D12_STATE_CACHE_INLINE void GetBlendState(D3D12_BLEND_DESC** BlendState, float BlendFactor[4], uint32* SampleMask)
+	D3D12_STATE_CACHE_INLINE void GetBlendState(D3D12_BLEND_DESC** BlendState, float BlendFactor[4], uint32* SampleMask) const
 	{
 		*BlendState = PipelineState.Graphics.HighLevelDesc.BlendState;
 		*SampleMask = PipelineState.Graphics.HighLevelDesc.SampleMask;
@@ -1169,7 +1172,7 @@ public:
 
 	D3D12_STATE_CACHE_INLINE void SetDepthStencilState(D3D12_DEPTH_STENCIL_DESC* State, uint32 RefStencil);
 
-	D3D12_STATE_CACHE_INLINE void GetDepthStencilState(D3D12_DEPTH_STENCIL_DESC** DepthStencilState, uint32* StencilRef)
+	D3D12_STATE_CACHE_INLINE void GetDepthStencilState(D3D12_DEPTH_STENCIL_DESC** DepthStencilState, uint32* StencilRef) const
 	{
 		*DepthStencilState = PipelineState.Graphics.HighLevelDesc.DepthStencilState;
 		*StencilRef = PipelineState.Graphics.CurrentReferenceStencil;
@@ -1226,7 +1229,7 @@ public:
 		}
 	}
 
-	D3D12_STATE_CACHE_INLINE void GetBoundShaderState(FD3D12BoundShaderState** BoundShaderState)
+	D3D12_STATE_CACHE_INLINE void GetBoundShaderState(FD3D12BoundShaderState** BoundShaderState) const
 	{
 		*BoundShaderState = PipelineState.Graphics.HighLevelDesc.BoundShaderState;
 	}
@@ -1249,7 +1252,7 @@ public:
 		*ComputeShader = PipelineState.Compute.CurrentComputeShader;
 	}
 
-	D3D12_STATE_CACHE_INLINE void GetInputLayout(D3D12_INPUT_LAYOUT_DESC* InputLayout)
+	D3D12_STATE_CACHE_INLINE void GetInputLayout(D3D12_INPUT_LAYOUT_DESC* InputLayout) const
 	{
 		*InputLayout = PipelineState.Graphics.HighLevelDesc.BoundShaderState->InputLayout;
 	}
@@ -1337,7 +1340,7 @@ public:
 	D3D12_STATE_CACHE_INLINE void SetPrimitiveTopology(D3D12_PRIMITIVE_TOPOLOGY PrimitiveTopology);
 
 
-	D3D12_STATE_CACHE_INLINE void GetPrimitiveTopology(D3D12_PRIMITIVE_TOPOLOGY* PrimitiveTopology)
+	D3D12_STATE_CACHE_INLINE void GetPrimitiveTopology(D3D12_PRIMITIVE_TOPOLOGY* PrimitiveTopology) const
 	{
 		*PrimitiveTopology = PipelineState.Graphics.CurrentPrimitiveTopology;
 	}
@@ -1435,23 +1438,23 @@ public:
 	void ForceSetBlendFactor() { bNeedSetBlendFactor = true; }
 	void ForceSetStencilRef() { bNeedSetStencilRef = true; }
 
-	bool GetForceRebuildGraphicsPSO() { return PipelineState.Graphics.bNeedRebuildPSO; }
-	bool GetForceRebuildComputePSO() { return PipelineState.Compute.bNeedRebuildPSO; }
-	bool GetForceSetVB() { return bNeedSetVB; }
-	bool GetForceSetIB() { return bNeedSetIB; }
-	bool GetForceSetUAVs() { return bNeedSetUAVs; }
-	bool GetForceSetRTs() { return bNeedSetRTs; }
-	bool GetForceSetSOs() { return bNeedSetSOs; }
-	bool GetForceSetSamplersPerShaderStage(uint32 Frequency) { return bNeedSetSamplersPerShaderStage[Frequency]; }
-	bool GetForceSetSamplers() { return bNeedSetSamplers; }
-	bool GetForceSetSRVsPerShaderStage(uint32 Frequency) { return bNeedSetSRVsPerShaderStage[Frequency]; }
-	bool GetForceSetSRVs() { return bNeedSetSRVs; }
-	bool GetForceSetConstantBuffers() { return bNeedSetConstantBuffers; }
-	bool GetForceSetViewports() { return bNeedSetViewports; }
-	bool GetForceSetScissorRects() { return bNeedSetScissorRects; }
-	bool GetForceSetPrimitiveTopology() { return bNeedSetPrimitiveTopology; }
-	bool GetForceSetBlendFactor() { return bNeedSetBlendFactor; }
-	bool GetForceSetStencilRef() { return bNeedSetStencilRef; }
+	bool GetForceRebuildGraphicsPSO() const { return PipelineState.Graphics.bNeedRebuildPSO; }
+	bool GetForceRebuildComputePSO() const { return PipelineState.Compute.bNeedRebuildPSO; }
+	bool GetForceSetVB() const { return bNeedSetVB; }
+	bool GetForceSetIB() const { return bNeedSetIB; }
+	bool GetForceSetUAVs() const { return bNeedSetUAVs; }
+	bool GetForceSetRTs() const { return bNeedSetRTs; }
+	bool GetForceSetSOs() const { return bNeedSetSOs; }
+	bool GetForceSetSamplersPerShaderStage(uint32 Frequency) const { return bNeedSetSamplersPerShaderStage[Frequency]; }
+	bool GetForceSetSamplers() const { return bNeedSetSamplers; }
+	bool GetForceSetSRVsPerShaderStage(uint32 Frequency) const { return bNeedSetSRVsPerShaderStage[Frequency]; }
+	bool GetForceSetSRVs() const { return bNeedSetSRVs; }
+	bool GetForceSetConstantBuffers() const { return bNeedSetConstantBuffers; }
+	bool GetForceSetViewports() const { return bNeedSetViewports; }
+	bool GetForceSetScissorRects() const { return bNeedSetScissorRects; }
+	bool GetForceSetPrimitiveTopology() const { return bNeedSetPrimitiveTopology; }
+	bool GetForceSetBlendFactor() const { return bNeedSetBlendFactor; }
+	bool GetForceSetStencilRef() const { return bNeedSetStencilRef; }
 
 
 #if D3D11_STATE_CACHE_DEBUG
