UE4 - Fix mono build error

diff --git a/Engine/Source/Runtime/Windows/D3D12RHI/Private/D3D12RHIPrivate.h b/Engine/Source/Runtime/Windows/D3D12RHI/Private/D3D12RHIPrivate.h
--- a/Engine/Source/Runtime/Windows/D3D12RHI/Private/D3D12RHIPrivate.h
+++ b/Engine/Source/Runtime/Windows/D3D12RHI/Private/D3D12RHIPrivate.h
@@ -729,6 +729,9 @@ struct FD3D12Adapter
 
 class FD3D12DynamicRHI;
 
+bool operator==(const D3D12_SAMPLER_DESC& lhs, const D3D12_SAMPLER_DESC& rhs);
+uint32 GetTypeHash(const D3D12_SAMPLER_DESC& Desc);
+
 class FD3D12Device
 {
 
diff --git a/Engine/Source/Runtime/Windows/D3D12RHI/Private/D3D12State.cpp b/Engine/Source/Runtime/Windows/D3D12RHI/Private/D3D12State.cpp
--- a/Engine/Source/Runtime/Windows/D3D12RHI/Private/D3D12State.cpp
+++ b/Engine/Source/Runtime/Windows/D3D12RHI/Private/D3D12State.cpp
@@ -112,7 +112,7 @@ static D3D12_BLEND TranslateBlendFactor(EBlendFactor BlendFactor)
 	};
 }
 
-static bool operator==(const D3D12_SAMPLER_DESC& lhs, const D3D12_SAMPLER_DESC& rhs)
+bool operator==(const D3D12_SAMPLER_DESC& lhs, const D3D12_SAMPLER_DESC& rhs)
 {
 	return 0 == memcmp(&lhs, &rhs, sizeof(lhs));
 }
