D3D12: Fixed VS2012 build

diff --git a/Engine/Source/Runtime/Windows/D3D12RHI/Private/D3D12RHI.cpp b/Engine/Source/Runtime/Windows/D3D12RHI/Private/D3D12RHI.cpp
--- a/Engine/Source/Runtime/Windows/D3D12RHI/Private/D3D12RHI.cpp
+++ b/Engine/Source/Runtime/Windows/D3D12RHI/Private/D3D12RHI.cpp
@@ -217,23 +217,6 @@ void FD3D12CommandContext::CheckIfSRVIsResolved(FD3D12ShaderResourceView* SRV)
 #endif
 }
 
-template <EShaderFrequency ShaderFrequency>
-void FD3D12CommandContext::InternalSetShaderResourceView(FD3D12ResourceLocation* Resource, FD3D12ShaderResourceView* SRV, int32 ResourceIndex, FD3D12StateCache::ESRV_Type SrvType)
-{
-	// Check either both are set, or both are null.
-	check((Resource && SRV) || (!Resource && !SRV));
-	CheckIfSRVIsResolved(SRV);
-
-	// Set the SRV we have been given (or null).
-	StateCache.SetShaderResourceView<ShaderFrequency>(SRV, ResourceIndex, SrvType);
-}
-
-template <EShaderFrequency ShaderFrequency>
-void FD3D12CommandContext::ClearShaderResourceViews(FD3D12ResourceLocation* Resource)
-{
-	StateCache.ClearShaderResourceViews<ShaderFrequency>(Resource);
-}
-
 void FD3D12CommandContext::ConditionalClearShaderResource(FD3D12ResourceLocation* Resource)
 {
 	SCOPE_CYCLE_COUNTER(STAT_D3D12ClearShaderResourceTime);
@@ -739,7 +722,7 @@ void FD3D12ResourceLocation::UpdateStateCache(FD3D12StateCache& StateCache)
 	}
 }
 
-inline void FD3D12ResourceLocation::UpdateDefaultStateCache()
+void FD3D12ResourceLocation::UpdateDefaultStateCache()
 {
     UpdateStateCache(GetParentDevice()->GetDefaultCommandContext().StateCache);
 }
diff --git a/Engine/Source/Runtime/Windows/D3D12RHI/Private/D3D12RHIPrivate.h b/Engine/Source/Runtime/Windows/D3D12RHI/Private/D3D12RHIPrivate.h
--- a/Engine/Source/Runtime/Windows/D3D12RHI/Private/D3D12RHIPrivate.h
+++ b/Engine/Source/Runtime/Windows/D3D12RHI/Private/D3D12RHIPrivate.h
@@ -431,12 +431,23 @@ public:
 	}
 
 	template <EShaderFrequency ShaderFrequency>
-	void ClearShaderResourceViews(FD3D12ResourceLocation* Resource);
+	void ClearShaderResourceViews(FD3D12ResourceLocation* Resource)
+	{
+		StateCache.ClearShaderResourceViews<ShaderFrequency>(Resource);
+	}
 
 	void CheckIfSRVIsResolved(FD3D12ShaderResourceView* SRV);
 
 	template <EShaderFrequency ShaderFrequency>
-	void InternalSetShaderResourceView(FD3D12ResourceLocation* Resource, FD3D12ShaderResourceView* SRV, int32 ResourceIndex, FD3D12StateCache::ESRV_Type SrvType = FD3D12StateCache::SRV_Unknown);
+	void InternalSetShaderResourceView(FD3D12ResourceLocation* Resource, FD3D12ShaderResourceView* SRV, int32 ResourceIndex, FD3D12StateCache::ESRV_Type SrvType = FD3D12StateCache::SRV_Unknown)
+	{
+		// Check either both are set, or both are null.
+		check((Resource && SRV) || (!Resource && !SRV));
+		CheckIfSRVIsResolved(SRV);
+
+		// Set the SRV we have been given (or null).
+		StateCache.SetShaderResourceView<ShaderFrequency>(SRV, ResourceIndex, SrvType);
+	}
 
 	void SetCurrentComputeShader(FComputeShaderRHIParamRef ComputeShader)
 	{
diff --git a/Engine/Source/Runtime/Windows/D3D12RHI/Private/D3D12StateCache.cpp b/Engine/Source/Runtime/Windows/D3D12RHI/Private/D3D12StateCache.cpp
--- a/Engine/Source/Runtime/Windows/D3D12RHI/Private/D3D12StateCache.cpp
+++ b/Engine/Source/Runtime/Windows/D3D12RHI/Private/D3D12StateCache.cpp
@@ -744,6 +744,31 @@ void FD3D12StateCacheBase::ClearSRVs()
 	bSRVSCleared = true;
 }
 
+void FD3D12StateCacheBase::ClearSamplers()
+{
+	ZeroMemory(PipelineState.Common.CurrentSamplerStates, sizeof(PipelineState.Common.CurrentSamplerStates));
+	for (uint32 ShaderFrequency = 0; ShaderFrequency < _countof(bNeedSetSamplersPerShaderStage); ShaderFrequency++)
+	{
+		bNeedSetSamplersPerShaderStage[ShaderFrequency] = true;
+	}
+	bNeedSetSamplers = true;
+}
+
+void FD3D12StateCacheBase::FlushComputeShaderCache(bool bForce)
+{
+	if (bAutoFlushComputeShaderCache || bForce)
+	{
+		D3D12_RESOURCE_BARRIER desc = {};
+		desc.Type = D3D12_RESOURCE_BARRIER_TYPE_UAV;
+		desc.UAV.pResource = nullptr;
+
+		FD3D12CommandListHandle& CommandList = CmdContext->CommandListHandle;
+
+		CmdContext->numBarriers++;
+		CommandList->ResourceBarrier(1, &desc);
+	}
+}
+
 void FD3D12StateCacheBase::ClearState()
 {
 	// Shader Resource View State Cache
@@ -848,7 +873,7 @@ void FD3D12StateCacheBase::RestoreState()
 	}
 }
 
-D3D12_STATE_CACHE_INLINE void FD3D12StateCacheBase::SetViewport(D3D12_VIEWPORT Viewport)
+void FD3D12StateCacheBase::SetViewport(D3D12_VIEWPORT Viewport)
 {
 	if ((PipelineState.Graphics.CurrentNumberOfViewports != 1 || FMemory::Memcmp(&PipelineState.Graphics.CurrentViewport[0], &Viewport, sizeof(D3D12_VIEWPORT))) || GD3D12SkipStateCaching)
 	{
@@ -859,7 +884,7 @@ D3D12_STATE_CACHE_INLINE void FD3D12StateCacheBase::SetViewport(D3D12_VIEWPORT V
 	}
 }
 
-D3D12_STATE_CACHE_INLINE void FD3D12StateCacheBase::SetViewports(uint32 Count, D3D12_VIEWPORT* Viewports)
+void FD3D12StateCacheBase::SetViewports(uint32 Count, D3D12_VIEWPORT* Viewports)
 {
 	if ((PipelineState.Graphics.CurrentNumberOfViewports != Count || FMemory::Memcmp(&PipelineState.Graphics.CurrentViewport[0], Viewports, sizeof(D3D12_VIEWPORT)* Count)) || GD3D12SkipStateCaching)
 	{
@@ -870,7 +895,24 @@ D3D12_STATE_CACHE_INLINE void FD3D12StateCacheBase::SetViewports(uint32 Count, D
 	}
 }
 
-D3D12_STATE_CACHE_INLINE void FD3D12StateCacheBase::SetScissorRect(D3D12_RECT ScissorRect)
+void FD3D12StateCacheBase::UpdateViewportScissorRects()
+{
+	for (uint32 i = 0; i < PipelineState.Graphics.CurrentNumberOfScissorRects; ++i)
+	{
+		D3D12_VIEWPORT& Viewport		    = PipelineState.Graphics.CurrentViewport[FMath::Min (i, PipelineState.Graphics.CurrentNumberOfViewports)];
+		D3D12_RECT&     ScissorRect         = PipelineState.Graphics.CurrentScissorRects[i];
+		D3D12_RECT&     ViewportScissorRect = PipelineState.Graphics.CurrentViewportScissorRects[i];
+
+		ViewportScissorRect.top    = FMath::Max(ScissorRect.top,    (LONG)Viewport.TopLeftY);
+		ViewportScissorRect.left   = FMath::Max(ScissorRect.left,   (LONG)Viewport.TopLeftX);
+		ViewportScissorRect.bottom = FMath::Min(ScissorRect.bottom, (LONG)Viewport.TopLeftY + (LONG)Viewport.Height);
+		ViewportScissorRect.right  = FMath::Min(ScissorRect.right,  (LONG)Viewport.TopLeftX + (LONG)Viewport.Width);
+	}
+
+	bNeedSetScissorRects = true;
+}
+
+void FD3D12StateCacheBase::SetScissorRect(D3D12_RECT ScissorRect)
 {
 	if ((PipelineState.Graphics.CurrentNumberOfScissorRects != 1 || FMemory::Memcmp(&PipelineState.Graphics.CurrentScissorRects[0], &ScissorRect, sizeof(D3D12_RECT))) || GD3D12SkipStateCaching)
 	{
@@ -880,7 +922,7 @@ D3D12_STATE_CACHE_INLINE void FD3D12StateCacheBase::SetScissorRect(D3D12_RECT Sc
 	}
 }
 
-D3D12_STATE_CACHE_INLINE void FD3D12StateCacheBase::SetScissorRects(uint32 Count, D3D12_RECT* ScissorRects)
+void FD3D12StateCacheBase::SetScissorRects(uint32 Count, D3D12_RECT* ScissorRects)
 {
 	if ((PipelineState.Graphics.CurrentNumberOfScissorRects != Count || FMemory::Memcmp(&PipelineState.Graphics.CurrentScissorRects[0], ScissorRects, sizeof(D3D12_RECT)* Count)) || GD3D12SkipStateCaching)
 	{
@@ -2319,7 +2361,7 @@ void FD3D12StateCacheBase::SetBlendState(D3D12_BLEND_DESC* State, const float Bl
 	}
 }
 
-D3D12_STATE_CACHE_INLINE void FD3D12StateCacheBase::SetDepthStencilState(D3D12_DEPTH_STENCIL_DESC* State, uint32 RefStencil)
+void FD3D12StateCacheBase::SetDepthStencilState(D3D12_DEPTH_STENCIL_DESC* State, uint32 RefStencil)
 {
 	if (PipelineState.Graphics.HighLevelDesc.DepthStencilState != State || GD3D12SkipStateCaching)
 	{
@@ -2404,102 +2446,6 @@ void FD3D12StateCacheBase::SetStreamOutTargets(uint32 NumSimultaneousStreamOutTa
 	}
 }
 
-template <EShaderFrequency ShaderFrequency>
-void FD3D12StateCacheBase::InternalSetShaderResourceView(FD3D12ShaderResourceView*& SRV, uint32 ResourceIndex, ESRV_Type SrvType, TSetSRVAlternate AlternatePathFunction)
-{
-	check(ResourceIndex < ARRAYSIZE(PipelineState.Common.CurrentShaderResourceViews[ShaderFrequency]));
-	auto& CurrentShaderResourceViews = PipelineState.Common.CurrentShaderResourceViews[ShaderFrequency];
-	if ((CurrentShaderResourceViews[ResourceIndex] != SRV) || GD3D12SkipStateCaching)
-	{
-		// Keep track of the highest bound resource
-		int32& MaxResourceIndex = PipelineState.Common.MaxBoundShaderResourcesIndex[ShaderFrequency];
-		if (SRV != nullptr)
-		{
-			// Mark the SRVs as not cleared
-			bSRVSCleared = false;
-
-			check(ResourceIndex < MAX_SRVS);
-			// Update the max resource index to the highest bound resource index.
-			MaxResourceIndex = FMath::Max(MaxResourceIndex, static_cast<int32>(ResourceIndex));
-		}
-		else
-		{
-			// If this was the highest bound resource...
-			if (MaxResourceIndex == ResourceIndex)
-			{
-				// Adjust the max resource index downwards until we
-				// hit the next non-null slot, or we've run out of slots.
-				do
-				{
-					MaxResourceIndex--;
-				} while (MaxResourceIndex >= 0 && CurrentShaderResourceViews[MaxResourceIndex] == nullptr);
-			}
-		}
-
-		CurrentShaderResourceViews[ResourceIndex] = SRV;
-		bNeedSetSRVsPerShaderStage[ShaderFrequency] = true;
-		bNeedSetSRVs = true;
-
-		if (FD3D12DynamicRHI::ResourceViewsIntersect(PipelineState.Graphics.CurrentDepthStencilTarget, SRV))
-		{
-			check(PipelineState.Graphics.CurrentDepthStencilTarget->GetDesc().Flags & D3D12_DSV_FLAG_READ_ONLY_DEPTH);
-			if (!PipelineState.Common.CurrentShaderResourceViewsIntersectWithDepthRT[ShaderFrequency][ResourceIndex])
-			{
-				PipelineState.Common.CurrentShaderResourceViewsIntersectWithDepthRT[ShaderFrequency][ResourceIndex] = true;
-				PipelineState.Common.ShaderResourceViewsIntersectWithDepthCount++;
-			}
-		}
-		else
-		{
-			if (PipelineState.Common.CurrentShaderResourceViewsIntersectWithDepthRT[ShaderFrequency][ResourceIndex])
-			{
-				PipelineState.Common.CurrentShaderResourceViewsIntersectWithDepthRT[ShaderFrequency][ResourceIndex] = false;
-				PipelineState.Common.ShaderResourceViewsIntersectWithDepthCount--;
-			}
-		}
-
-		if (AlternatePathFunction != nullptr)
-		{
-			(*AlternatePathFunction)(this, SRV, ResourceIndex, SrvType);
-		}
-	}
-}
-
-template <EShaderFrequency ShaderFrequency>
-void FD3D12StateCacheBase::ClearShaderResourceViews(FD3D12ResourceLocation*& ResourceLocation)
-{
-	int32& MaxResourceIndex = PipelineState.Common.MaxBoundShaderResourcesIndex[ShaderFrequency];
-	auto& CurrentShaderResourceViews = PipelineState.Common.CurrentShaderResourceViews[ShaderFrequency];
-
-	const int32 OriginalMaxResourceIndex = MaxResourceIndex;
-	for (int32 i = 0; i <= OriginalMaxResourceIndex; ++i)
-	{
-		if (CurrentShaderResourceViews[i].GetReference() && CurrentShaderResourceViews[i]->GetResourceLocation() == ResourceLocation)
-		{
-			SetShaderResourceView<ShaderFrequency>(nullptr, i);
-		}
-	}
-
-	// Restore the max index to make sure we set null descriptors in the slots of any SRVs we unbind due to them overlapping with other
-	// resources like RTVs and DSVs.
-	MaxResourceIndex = OriginalMaxResourceIndex;
-}
-
-void FD3D12StateCacheBase::FlushComputeShaderCache(bool bForce)
-{
-	if (bAutoFlushComputeShaderCache || bForce)
-	{
-		D3D12_RESOURCE_BARRIER desc = {};
-		desc.Type = D3D12_RESOURCE_BARRIER_TYPE_UAV;
-		desc.UAV.pResource = nullptr;
-
-		FD3D12CommandListHandle& CommandList = CmdContext->CommandListHandle;
-
-		CmdContext->numBarriers++;
-		CommandList->ResourceBarrier(1, &desc);
-	}
-}
-
 void FD3D12DescriptorCache::HeapRolledOver(D3D12_DESCRIPTOR_HEAP_TYPE Type)
 {
 	TArray<ID3D12DescriptorHeap*> Heaps;
diff --git a/Engine/Source/Runtime/Windows/D3D12RHI/Private/D3D12StateCachePrivate.h b/Engine/Source/Runtime/Windows/D3D12RHI/Private/D3D12StateCachePrivate.h
--- a/Engine/Source/Runtime/Windows/D3D12RHI/Private/D3D12StateCachePrivate.h
+++ b/Engine/Source/Runtime/Windows/D3D12RHI/Private/D3D12StateCachePrivate.h
@@ -926,14 +926,72 @@ protected:
 	FD3D12DescriptorCache DescriptorCache;
 	bool bAlwaysSetIndexBuffers;
 
-	D3D12_STATE_CACHE_INLINE void InternalSetIndexBuffer(FD3D12ResourceLocation *IndexBufferLocation, DXGI_FORMAT Format, uint32 Offset);
+	void InternalSetIndexBuffer(FD3D12ResourceLocation *IndexBufferLocation, DXGI_FORMAT Format, uint32 Offset);
 
 	typedef void(*TSetSRVAlternate)(FD3D12StateCacheBase* StateCache, FD3D12ShaderResourceView* SRV, uint32 ResourceIndex, ESRV_Type SrvType);
 	template <EShaderFrequency ShaderFrequency>
-	void D3D12_STATE_CACHE_INLINE InternalSetShaderResourceView(FD3D12ShaderResourceView*& SRV, uint32 ResourceIndex, ESRV_Type SrvType, TSetSRVAlternate AlternatePathFunction);
+	void InternalSetShaderResourceView(FD3D12ShaderResourceView*& SRV, uint32 ResourceIndex, ESRV_Type SrvType, TSetSRVAlternate AlternatePathFunction)
+	{
+		check(ResourceIndex < ARRAYSIZE(PipelineState.Common.CurrentShaderResourceViews[ShaderFrequency]));
+		auto& CurrentShaderResourceViews = PipelineState.Common.CurrentShaderResourceViews[ShaderFrequency];
+		if ((CurrentShaderResourceViews[ResourceIndex] != SRV) || GD3D12SkipStateCaching)
+		{
+			// Keep track of the highest bound resource
+			int32& MaxResourceIndex = PipelineState.Common.MaxBoundShaderResourcesIndex[ShaderFrequency];
+			if (SRV != nullptr)
+			{
+				// Mark the SRVs as not cleared
+				bSRVSCleared = false;
+
+				check(ResourceIndex < MAX_SRVS);
+				// Update the max resource index to the highest bound resource index.
+				MaxResourceIndex = FMath::Max(MaxResourceIndex, static_cast<int32>(ResourceIndex));
+			}
+			else
+			{
+				// If this was the highest bound resource...
+				if (MaxResourceIndex == ResourceIndex)
+				{
+					// Adjust the max resource index downwards until we
+					// hit the next non-null slot, or we've run out of slots.
+					do
+					{
+						MaxResourceIndex--;
+					} while (MaxResourceIndex >= 0 && CurrentShaderResourceViews[MaxResourceIndex] == nullptr);
+				}
+			}
+
+			CurrentShaderResourceViews[ResourceIndex] = SRV;
+			bNeedSetSRVsPerShaderStage[ShaderFrequency] = true;
+			bNeedSetSRVs = true;
+
+			if (FD3D12DynamicRHI::ResourceViewsIntersect(PipelineState.Graphics.CurrentDepthStencilTarget, SRV))
+			{
+				check(PipelineState.Graphics.CurrentDepthStencilTarget->GetDesc().Flags & D3D12_DSV_FLAG_READ_ONLY_DEPTH);
+				if (!PipelineState.Common.CurrentShaderResourceViewsIntersectWithDepthRT[ShaderFrequency][ResourceIndex])
+				{
+					PipelineState.Common.CurrentShaderResourceViewsIntersectWithDepthRT[ShaderFrequency][ResourceIndex] = true;
+					PipelineState.Common.ShaderResourceViewsIntersectWithDepthCount++;
+				}
+			}
+			else
+			{
+				if (PipelineState.Common.CurrentShaderResourceViewsIntersectWithDepthRT[ShaderFrequency][ResourceIndex])
+				{
+					PipelineState.Common.CurrentShaderResourceViewsIntersectWithDepthRT[ShaderFrequency][ResourceIndex] = false;
+					PipelineState.Common.ShaderResourceViewsIntersectWithDepthCount--;
+				}
+			}
+
+			if (AlternatePathFunction != nullptr)
+			{
+				(*AlternatePathFunction)(this, SRV, ResourceIndex, SrvType);
+			}
+		}
+	}
 
 	typedef void(*TSetStreamSourceAlternate)(FD3D12StateCacheBase* StateCache, FD3D12Resource* VertexBuffer, uint32 StreamIndex, uint32 Stride, uint32 Offset);
-	D3D12_STATE_CACHE_INLINE void InternalSetStreamSource(FD3D12ResourceLocation* VertexBufferLocation, uint32 StreamIndex, uint32 Stride, uint32 Offset, TSetStreamSourceAlternate AlternatePathFunction);
+	void InternalSetStreamSource(FD3D12ResourceLocation* VertexBufferLocation, uint32 StreamIndex, uint32 Stride, uint32 Offset, TSetStreamSourceAlternate AlternatePathFunction);
 
 	typedef void(*TSetSamplerStateAlternate)(FD3D12StateCacheBase* StateCache, FD3D12SamplerState* SamplerState, uint32 SamplerIndex);
 	template <EShaderFrequency ShaderFrequency>
@@ -998,27 +1056,40 @@ public:
 		RestoreState();
 	}
 
-	FD3D12DescriptorCache* GetDescriptorCache() { return &DescriptorCache; }
+	FD3D12DescriptorCache* GetDescriptorCache()
+	{
+		return &DescriptorCache;
+	}
 
 	ID3D12PipelineState* GetPipelineStateObject()
 	{
 		return PipelineState.Common.CurrentPipelineStateObject;
 	}
 
-	void ClearSamplers()
+	void ClearSamplers();
+	void ClearSRVs();
+
+	template <EShaderFrequency ShaderFrequency>
+	void ClearShaderResourceViews(FD3D12ResourceLocation*& ResourceLocation)
 	{
-		ZeroMemory(PipelineState.Common.CurrentSamplerStates, sizeof(PipelineState.Common.CurrentSamplerStates));
-		for (uint32 ShaderFrequency = 0; ShaderFrequency < _countof(bNeedSetSamplersPerShaderStage); ShaderFrequency++)
+		int32& MaxResourceIndex = PipelineState.Common.MaxBoundShaderResourcesIndex[ShaderFrequency];
+		auto& CurrentShaderResourceViews = PipelineState.Common.CurrentShaderResourceViews[ShaderFrequency];
+
+		const int32 OriginalMaxResourceIndex = MaxResourceIndex;
+		for (int32 i = 0; i <= OriginalMaxResourceIndex; ++i)
 		{
-			bNeedSetSamplersPerShaderStage[ShaderFrequency] = true;
+			if (CurrentShaderResourceViews[i].GetReference() && CurrentShaderResourceViews[i]->GetResourceLocation() == ResourceLocation)
+			{
+				SetShaderResourceView<ShaderFrequency>(nullptr, i);
+			}
 		}
-		bNeedSetSamplers = true;
-	}
 
-	void ClearSRVs();
+		// Restore the max index to make sure we set null descriptors in the slots of any SRVs we unbind due to them overlapping with other
+		// resources like RTVs and DSVs.
+		MaxResourceIndex = OriginalMaxResourceIndex;
+	}
 
-	template <EShaderFrequency ShaderFrequency>
-	void ClearShaderResourceViews(FD3D12ResourceLocation*& ResourceLocation);
+	void FD3D12StateCacheBase::FlushComputeShaderCache(bool bForce = false);
 
 	template <EShaderFrequency ShaderFrequency>
 	D3D12_STATE_CACHE_INLINE void SetShaderResourceView(FD3D12ShaderResourceView* SRV, uint32 ResourceIndex, ESRV_Type SrvType = SRV_Unknown)
@@ -1044,25 +1115,9 @@ public:
 		}
 	}
 
-	void UpdateViewportScissorRects()
-	{
-		for (uint32 i = 0; i < PipelineState.Graphics.CurrentNumberOfScissorRects; ++i)
-		{
-			D3D12_VIEWPORT& Viewport		    = PipelineState.Graphics.CurrentViewport[FMath::Min (i, PipelineState.Graphics.CurrentNumberOfViewports)];
-			D3D12_RECT&     ScissorRect         = PipelineState.Graphics.CurrentScissorRects[i];
-			D3D12_RECT&     ViewportScissorRect = PipelineState.Graphics.CurrentViewportScissorRects[i];
-
-			ViewportScissorRect.top    = FMath::Max(ScissorRect.top,    (LONG)Viewport.TopLeftY);
-			ViewportScissorRect.left   = FMath::Max(ScissorRect.left,   (LONG)Viewport.TopLeftX);
-			ViewportScissorRect.bottom = FMath::Min(ScissorRect.bottom, (LONG)Viewport.TopLeftY + (LONG)Viewport.Height);
-			ViewportScissorRect.right  = FMath::Min(ScissorRect.right,  (LONG)Viewport.TopLeftX + (LONG)Viewport.Width);
-		}
-
-		bNeedSetScissorRects = true;
-	}
-
-	D3D12_STATE_CACHE_INLINE void SetScissorRects(uint32 Count, D3D12_RECT* ScissorRects);
-	D3D12_STATE_CACHE_INLINE void SetScissorRect(D3D12_RECT ScissorRect);
+	void UpdateViewportScissorRects();
+	void SetScissorRects(uint32 Count, D3D12_RECT* ScissorRects);
+	void SetScissorRect(D3D12_RECT ScissorRect);
 
 	D3D12_STATE_CACHE_INLINE void GetScissorRect(D3D12_RECT *ScissorRect) const
 	{
@@ -1071,8 +1126,8 @@ public:
 	}
 
 
-	D3D12_STATE_CACHE_INLINE void SetViewport(D3D12_VIEWPORT Viewport);
-	D3D12_STATE_CACHE_INLINE void SetViewports(uint32 Count, D3D12_VIEWPORT* Viewports);
+	void SetViewport(D3D12_VIEWPORT Viewport);
+	void SetViewports(uint32 Count, D3D12_VIEWPORT* Viewports);
 
 	D3D12_STATE_CACHE_INLINE void GetViewport(D3D12_VIEWPORT *Viewport) const
 	{
@@ -1161,7 +1216,7 @@ public:
 		*RasterizerState = PipelineState.Graphics.HighLevelDesc.RasterizerState;
 	}
 
-	D3D12_STATE_CACHE_INLINE void SetBlendState(D3D12_BLEND_DESC* State, const float BlendFactor[4], uint32 SampleMask);
+	void SetBlendState(D3D12_BLEND_DESC* State, const float BlendFactor[4], uint32 SampleMask);
 
 	D3D12_STATE_CACHE_INLINE void GetBlendState(D3D12_BLEND_DESC** BlendState, float BlendFactor[4], uint32* SampleMask) const
 	{
@@ -1170,7 +1225,7 @@ public:
 		FMemory::Memcpy(BlendFactor, PipelineState.Graphics.CurrentBlendFactor, sizeof(PipelineState.Graphics.CurrentBlendFactor));
 	}
 
-	D3D12_STATE_CACHE_INLINE void SetDepthStencilState(D3D12_DEPTH_STENCIL_DESC* State, uint32 RefStencil);
+	void SetDepthStencilState(D3D12_DEPTH_STENCIL_DESC* State, uint32 RefStencil);
 
 	D3D12_STATE_CACHE_INLINE void GetDepthStencilState(D3D12_DEPTH_STENCIL_DESC** DepthStencilState, uint32* StencilRef) const
 	{
@@ -1337,7 +1392,7 @@ public:
 			&& PipelineState.Graphics.CurrentIndexBufferLocation->GetOffset() == ResourceLocation->GetOffset();
 	}
 
-	D3D12_STATE_CACHE_INLINE void SetPrimitiveTopology(D3D12_PRIMITIVE_TOPOLOGY PrimitiveTopology);
+	void SetPrimitiveTopology(D3D12_PRIMITIVE_TOPOLOGY PrimitiveTopology);
 
 
 	D3D12_STATE_CACHE_INLINE void GetPrimitiveTopology(D3D12_PRIMITIVE_TOPOLOGY* PrimitiveTopology) const
@@ -1360,13 +1415,13 @@ public:
 	{
 	}
 	void ApplyState(bool IsCompute = false);
-	void D3D12_STATE_CACHE_INLINE RestoreState();
+	void RestoreState();
 	bool VerifyResourceStates(const bool IsCompute);
 
 	template <class ViewT>
 	bool VerifyViewState(ID3D12DebugCommandList* pDebugCommandList, FD3D12View<ViewT> *pView, D3D12_RESOURCE_STATES State);
 
-	void D3D12_STATE_CACHE_INLINE SetRenderTargets(uint32 NumSimultaneousRenderTargets, FD3D12RenderTargetView** RTArray, FD3D12DepthStencilView* DSTarget);
+	void SetRenderTargets(uint32 NumSimultaneousRenderTargets, FD3D12RenderTargetView** RTArray, FD3D12DepthStencilView* DSTarget);
 	D3D12_STATE_CACHE_INLINE void GetRenderTargets(FD3D12RenderTargetView **RTArray, uint32* NumSimultaneousRTs, FD3D12DepthStencilView** DepthStencilTarget)
 	{
 		if (RTArray) //NULL is legal
@@ -1381,9 +1436,9 @@ public:
 		}
 	}
 
-	void D3D12_STATE_CACHE_INLINE SetStreamOutTargets(uint32 NumSimultaneousStreamOutTargets, FD3D12Resource** SOArray, const uint32* SOOffsets);
+	void SetStreamOutTargets(uint32 NumSimultaneousStreamOutTargets, FD3D12Resource** SOArray, const uint32* SOOffsets);
 
-	void D3D12_STATE_CACHE_INLINE SetUAVs(EShaderFrequency ShaderStage, uint32 UAVStartSlot, uint32 NumSimultaneousUAVs, FD3D12UnorderedAccessView** UAVArray, uint32 *UAVInitialCountArray);
+	void SetUAVs(EShaderFrequency ShaderStage, uint32 UAVStartSlot, uint32 NumSimultaneousUAVs, FD3D12UnorderedAccessView** UAVArray, uint32 *UAVInitialCountArray);
 
 	D3D12_STATE_CACHE_INLINE void GetUAVs(FD3D12UnorderedAccessView** UAVArray, uint32* UAVStartSlot, uint32* NumSimultaneousUAVs)
 	{
@@ -1406,8 +1461,6 @@ public:
 		bAutoFlushComputeShaderCache = bEnable;
 	}
 
-	void FlushComputeShaderCache(bool bForce = false);
-
 	/**
 	 * Clears all D3D11 State, setting all input/output resource slots, shaders, input layouts,
 	 * predications, scissor rectangles, depth-stencil state, rasterizer state, blend state,
