D3D12: Fixing build

diff --git a/Engine/Source/Runtime/Windows/D3D11RHI/Private/D3D11RHIPrivate.h b/Engine/Source/Runtime/Windows/D3D11RHI/Private/D3D11RHIPrivate.h
--- a/Engine/Source/Runtime/Windows/D3D11RHI/Private/D3D11RHIPrivate.h
+++ b/Engine/Source/Runtime/Windows/D3D11RHI/Private/D3D11RHIPrivate.h
@@ -58,6 +58,22 @@ DECLARE_DWORD_COUNTER_STAT_EXTERN(TEXT("Textures Allocated"),STAT_D3D11TexturesA
 DECLARE_DWORD_COUNTER_STAT_EXTERN(TEXT("Textures Released"),STAT_D3D11TexturesReleased,STATGROUP_D3D11RHI, );
 DECLARE_MEMORY_STAT_EXTERN(TEXT("Texture object pool memory"),STAT_D3D11TexturePoolMemory,STATGROUP_D3D11RHI, );
 
+struct FD3D11GlobalStats
+{
+	// in bytes, never change after RHI, needed to scale game features
+	static int64 GDedicatedVideoMemory;
+	
+	// in bytes, never change after RHI, needed to scale game features
+	static int64 GDedicatedSystemMemory;
+	
+	// in bytes, never change after RHI, needed to scale game features
+	static int64 GSharedSystemMemory;
+	
+	// In bytes. Never changed after RHI init. Our estimate of the amount of memory that we can use for graphics resources in total.
+	static int64 GTotalGraphicsMemory;
+};
+
+
 // This class has multiple inheritance but really FGPUTiming is a static class
 class FD3D11BufferedGPUTiming : public FRenderResource, public FGPUTiming
 {
diff --git a/Engine/Source/Runtime/Windows/D3D11RHI/Private/D3D11Texture.cpp b/Engine/Source/Runtime/Windows/D3D11RHI/Private/D3D11Texture.cpp
--- a/Engine/Source/Runtime/Windows/D3D11RHI/Private/D3D11Texture.cpp
+++ b/Engine/Source/Runtime/Windows/D3D11RHI/Private/D3D11Texture.cpp
@@ -6,19 +6,17 @@
 
 #include "D3D11RHIPrivate.h"
 
+
+int64 FD3D11GlobalStats::GDedicatedVideoMemory = 0;
+int64 FD3D11GlobalStats::GDedicatedSystemMemory = 0;
+int64 FD3D11GlobalStats::GSharedSystemMemory = 0;
+int64 FD3D11GlobalStats::GTotalGraphicsMemory = 0;
+
+
 /*-----------------------------------------------------------------------------
 	Texture allocator support.
 -----------------------------------------------------------------------------*/
 
-// in bytes, never change after RHI, needed to scale game features
-int64 GDedicatedVideoMemory = 0;
-// in bytes, never change after RHI, needed to scale game features
-int64 GDedicatedSystemMemory = 0;
-// in bytes, never change after RHI, needed to scale game features
-int64 GSharedSystemMemory = 0;
-// In bytes. Never changed after RHI init. Our estimate of the amount of memory that we can use for graphics resources in total.
-int64 GTotalGraphicsMemory = 0;
-
 static bool ShouldCountAsTextureMemory(uint32 BindFlags)
 {
 	return (BindFlags & (D3D11_BIND_DEPTH_STENCIL | D3D11_BIND_RENDER_TARGET | D3D11_BIND_UNORDERED_ACCESS)) == 0;
@@ -232,10 +230,10 @@ uint64 FD3D11DynamicRHI::RHICalcTextureCubePlatformSize(uint32 Size, uint8 Forma
  */
 void FD3D11DynamicRHI::RHIGetTextureMemoryStats(FTextureMemoryStats& OutStats)
 {
-	OutStats.DedicatedVideoMemory = GDedicatedVideoMemory;
-    OutStats.DedicatedSystemMemory = GDedicatedSystemMemory;
-    OutStats.SharedSystemMemory = GSharedSystemMemory;
-	OutStats.TotalGraphicsMemory = GTotalGraphicsMemory ? GTotalGraphicsMemory : -1;
+	OutStats.DedicatedVideoMemory = FD3D11GlobalStats::GDedicatedVideoMemory;
+    OutStats.DedicatedSystemMemory = FD3D11GlobalStats::GDedicatedSystemMemory;
+    OutStats.SharedSystemMemory = FD3D11GlobalStats::GSharedSystemMemory;
+	OutStats.TotalGraphicsMemory = FD3D11GlobalStats::GTotalGraphicsMemory ? FD3D11GlobalStats::GTotalGraphicsMemory : -1;
 
 	OutStats.AllocatedMemorySize = int64(GCurrentTextureMemorySize) * 1024;
 	OutStats.TexturePoolSize = GTexturePoolSize;
diff --git a/Engine/Source/Runtime/Windows/D3D11RHI/Private/Windows/WindowsD3D11Device.cpp b/Engine/Source/Runtime/Windows/D3D11RHI/Private/Windows/WindowsD3D11Device.cpp
--- a/Engine/Source/Runtime/Windows/D3D11RHI/Private/Windows/WindowsD3D11Device.cpp
+++ b/Engine/Source/Runtime/Windows/D3D11RHI/Private/Windows/WindowsD3D11Device.cpp
@@ -11,6 +11,7 @@
 
 #include "HardwareInfo.h"
 
+
 extern bool D3D11RHI_ShouldCreateWithD3DDebug();
 extern bool D3D11RHI_ShouldAllowAsyncResourceCreation();
 
@@ -446,65 +447,60 @@ void FD3D11DynamicRHI::InitD3DDevice()
 					GRHIAdapterName = AdapterDesc.Description;
 					GRHIVendorId = AdapterDesc.VendorId;
 
-					extern int64 GDedicatedVideoMemory;
-					extern int64 GDedicatedSystemMemory;
-					extern int64 GSharedSystemMemory;
-					extern int64 GTotalGraphicsMemory;
-
 					// Issue: 32bit windows doesn't report 64bit value, we take what we get.
-					GDedicatedVideoMemory = int64(AdapterDesc.DedicatedVideoMemory);
-					GDedicatedSystemMemory = int64(AdapterDesc.DedicatedSystemMemory);
-					GSharedSystemMemory = int64(AdapterDesc.SharedSystemMemory);
+					FD3D11GlobalStats::GDedicatedVideoMemory = int64(AdapterDesc.DedicatedVideoMemory);
+					FD3D11GlobalStats::GDedicatedSystemMemory = int64(AdapterDesc.DedicatedSystemMemory);
+					FD3D11GlobalStats::GSharedSystemMemory = int64(AdapterDesc.SharedSystemMemory);
 
 					// Total amount of system memory, clamped to 8 GB
 					int64 TotalPhysicalMemory = FMath::Min(int64(FPlatformMemory::GetConstants().TotalPhysicalGB), 8ll) * (1024ll * 1024ll * 1024ll);
 
 					// Consider 50% of the shared memory but max 25% of total system memory.
-					int64 ConsideredSharedSystemMemory = FMath::Min( GSharedSystemMemory / 2ll, TotalPhysicalMemory / 4ll );
+					int64 ConsideredSharedSystemMemory = FMath::Min( FD3D11GlobalStats::GSharedSystemMemory / 2ll, TotalPhysicalMemory / 4ll );
 
-					GTotalGraphicsMemory = 0;
+					FD3D11GlobalStats::GTotalGraphicsMemory = 0;
 					if ( IsRHIDeviceIntel() )
 					{
 						// It's all system memory.
-						GTotalGraphicsMemory = GDedicatedVideoMemory;
-						GTotalGraphicsMemory += GDedicatedSystemMemory;
-						GTotalGraphicsMemory += ConsideredSharedSystemMemory;
+						FD3D11GlobalStats::GTotalGraphicsMemory = FD3D11GlobalStats::GDedicatedVideoMemory;
+						FD3D11GlobalStats::GTotalGraphicsMemory += FD3D11GlobalStats::GDedicatedSystemMemory;
+						FD3D11GlobalStats::GTotalGraphicsMemory += ConsideredSharedSystemMemory;
 					}
-					else if ( GDedicatedVideoMemory >= 200*1024*1024 )
+					else if ( FD3D11GlobalStats::GDedicatedVideoMemory >= 200*1024*1024 )
 					{
 						// Use dedicated video memory, if it's more than 200 MB
-						GTotalGraphicsMemory = GDedicatedVideoMemory;
+						FD3D11GlobalStats::GTotalGraphicsMemory = FD3D11GlobalStats::GDedicatedVideoMemory;
 					}
-					else if ( GDedicatedSystemMemory >= 200*1024*1024 )
+					else if ( FD3D11GlobalStats::GDedicatedSystemMemory >= 200*1024*1024 )
 					{
 						// Use dedicated system memory, if it's more than 200 MB
-						GTotalGraphicsMemory = GDedicatedSystemMemory;
+						FD3D11GlobalStats::GTotalGraphicsMemory = FD3D11GlobalStats::GDedicatedSystemMemory;
 					}
-					else if ( GSharedSystemMemory >= 400*1024*1024 )
+					else if ( FD3D11GlobalStats::GSharedSystemMemory >= 400*1024*1024 )
 					{
 						// Use some shared system memory, if it's more than 400 MB
-						GTotalGraphicsMemory = ConsideredSharedSystemMemory;
+						FD3D11GlobalStats::GTotalGraphicsMemory = ConsideredSharedSystemMemory;
 					}
 					else
 					{
 						// Otherwise consider 25% of total system memory for graphics.
-						GTotalGraphicsMemory = TotalPhysicalMemory / 4ll;
+						FD3D11GlobalStats::GTotalGraphicsMemory = TotalPhysicalMemory / 4ll;
 					}
 
 					if ( sizeof(SIZE_T) < 8 )
 					{
 						// Clamp to 1 GB if we're less than 64-bit
-						GTotalGraphicsMemory = FMath::Min( GTotalGraphicsMemory, 1024ll * 1024ll * 1024ll );
+						FD3D11GlobalStats::GTotalGraphicsMemory = FMath::Min( FD3D11GlobalStats::GTotalGraphicsMemory, 1024ll * 1024ll * 1024ll );
 					}
 					else
 					{
 						// Clamp to 1.9 GB if we're 64-bit
-						GTotalGraphicsMemory = FMath::Min( GTotalGraphicsMemory, 1945ll * 1024ll * 1024ll );
+						FD3D11GlobalStats::GTotalGraphicsMemory = FMath::Min( FD3D11GlobalStats::GTotalGraphicsMemory, 1945ll * 1024ll * 1024ll );
 					}
 
 					if ( GPoolSizeVRAMPercentage > 0 )
 					{
-						float PoolSize = float(GPoolSizeVRAMPercentage) * 0.01f * float(GTotalGraphicsMemory);
+						float PoolSize = float(GPoolSizeVRAMPercentage) * 0.01f * float(FD3D11GlobalStats::GTotalGraphicsMemory);
 
 						// Truncate GTexturePoolSize to MB (but still counted in bytes)
 						GTexturePoolSize = int64(FGenericPlatformMath::TruncToFloat(PoolSize / 1024.0f / 1024.0f)) * 1024 * 1024;
@@ -512,7 +508,7 @@ void FD3D11DynamicRHI::InitD3DDevice()
 						UE_LOG(LogRHI,Log,TEXT("Texture pool is %llu MB (%d%% of %llu MB)"),
 							GTexturePoolSize / 1024 / 1024,
 							GPoolSizeVRAMPercentage,
-							GTotalGraphicsMemory / 1024 / 1024);
+							FD3D11GlobalStats::GTotalGraphicsMemory / 1024 / 1024);
 					}
 
 					const bool bIsPerfHUD = !FCString::Stricmp(AdapterDesc.Description,TEXT("NVIDIA PerfHUD"));
diff --git a/Engine/Source/Runtime/Windows/D3D12RHI/Private/D3D12RHIPrivate.h b/Engine/Source/Runtime/Windows/D3D12RHI/Private/D3D12RHIPrivate.h
--- a/Engine/Source/Runtime/Windows/D3D12RHI/Private/D3D12RHIPrivate.h
+++ b/Engine/Source/Runtime/Windows/D3D12RHI/Private/D3D12RHIPrivate.h
@@ -113,6 +113,22 @@ DECLARE_MEMORY_STAT_EXTERN(TEXT("Used Video Memory"), STAT_D3D12UsedVideoMemory,
 DECLARE_MEMORY_STAT_EXTERN(TEXT("Available Video Memory"), STAT_D3D12AvailableVideoMemory, STATGROUP_D3D12RHI, );
 DECLARE_MEMORY_STAT_EXTERN(TEXT("Total Video Memory"), STAT_D3D12TotalVideoMemory, STATGROUP_D3D12RHI, );
 
+struct FD3D12GlobalStats
+{
+	// in bytes, never change after RHI, needed to scale game features
+	static int64 GDedicatedVideoMemory;
+	
+	// in bytes, never change after RHI, needed to scale game features
+	static int64 GDedicatedSystemMemory;
+	
+	// in bytes, never change after RHI, needed to scale game features
+	static int64 GSharedSystemMemory;
+	
+	// In bytes. Never changed after RHI init. Our estimate of the amount of memory that we can use for graphics resources in total.
+	static int64 GTotalGraphicsMemory;
+};
+
+
 // This class handles query heaps
 class FD3D12QueryHeap : public FD3D12DeviceChild
 {
diff --git a/Engine/Source/Runtime/Windows/D3D12RHI/Private/D3D12Texture.cpp b/Engine/Source/Runtime/Windows/D3D12RHI/Private/D3D12Texture.cpp
--- a/Engine/Source/Runtime/Windows/D3D12RHI/Private/D3D12Texture.cpp
+++ b/Engine/Source/Runtime/Windows/D3D12RHI/Private/D3D12Texture.cpp
@@ -6,20 +6,17 @@
 
 #include "D3D12RHIPrivate.h"
 
+
+int64 FD3D12GlobalStats::GDedicatedVideoMemory = 0;
+int64 FD3D12GlobalStats::GDedicatedSystemMemory = 0;
+int64 FD3D12GlobalStats::GSharedSystemMemory = 0;
+int64 FD3D12GlobalStats::GTotalGraphicsMemory = 0;
+
+
 /*-----------------------------------------------------------------------------
 	Texture allocator support.
 	-----------------------------------------------------------------------------*/
 
-// in bytes, never change after RHI, needed to scale game features
-int64 GDedicatedVideoMemory = 0;
-// in bytes, never change after RHI, needed to scale game features
-int64 GDedicatedSystemMemory = 0;
-// in bytes, never change after RHI, needed to scale game features
-int64 GSharedSystemMemory = 0;
-// In bytes. Never changed after RHI init. Our estimate of the amount of memory that we can use for graphics resources in total.
-int64 GTotalGraphicsMemory = 0;
-
-
 namespace D3D12RHI
 {
 	static bool ShouldCountAsTextureMemory(uint32 MiscFlags)
@@ -236,9 +233,9 @@ uint64 FD3D12DynamicRHI::RHICalcTextureCubePlatformSize(uint32 Size, uint8 Forma
  */
 void FD3D12DynamicRHI::RHIGetTextureMemoryStats(FTextureMemoryStats& OutStats)
 {
-	OutStats.DedicatedVideoMemory = GDedicatedVideoMemory;
-	OutStats.DedicatedSystemMemory = GDedicatedSystemMemory;
-	OutStats.SharedSystemMemory = GSharedSystemMemory;
+	OutStats.DedicatedVideoMemory = FD3D12GlobalStats::GDedicatedVideoMemory;
+	OutStats.DedicatedSystemMemory = FD3D12GlobalStats::GDedicatedSystemMemory;
+	OutStats.SharedSystemMemory = FD3D12GlobalStats::GSharedSystemMemory;
 	bool bIsWARP = FParse::Param(FCommandLine::Get(), TEXT("warp"));
 
 	// Hack to handle a kernel bug where QueryVideomMemoryInfo can't be called on WARP
@@ -292,7 +289,7 @@ void FD3D12DynamicRHI::RHIGetTextureMemoryStats(FTextureMemoryStats& OutStats)
 	}
 	else
 	{
-		OutStats.TotalGraphicsMemory = GTotalGraphicsMemory;
+		OutStats.TotalGraphicsMemory = FD3D12GlobalStats::GTotalGraphicsMemory;
 	}
 
 
diff --git a/Engine/Source/Runtime/Windows/D3D12RHI/Private/Windows/WindowsD3D12Device.cpp b/Engine/Source/Runtime/Windows/D3D12RHI/Private/Windows/WindowsD3D12Device.cpp
--- a/Engine/Source/Runtime/Windows/D3D12RHI/Private/Windows/WindowsD3D12Device.cpp
+++ b/Engine/Source/Runtime/Windows/D3D12RHI/Private/Windows/WindowsD3D12Device.cpp
@@ -402,21 +402,16 @@ void FD3D12DynamicRHI::PerRHISetup(FD3D12Device* MainDevice)
     GRHIAdapterName = AdapterDesc->Description;
     GRHIVendorId = AdapterDesc->VendorId;
 
-    extern int64 GDedicatedVideoMemory;
-    extern int64 GDedicatedSystemMemory;
-    extern int64 GSharedSystemMemory;
-    extern int64 GTotalGraphicsMemory;
-
     // Issue: 32bit windows doesn't report 64bit value, we take what we get.
-    GDedicatedVideoMemory = int64(AdapterDesc->DedicatedVideoMemory);
-    GDedicatedSystemMemory = int64(AdapterDesc->DedicatedSystemMemory);
-    GSharedSystemMemory = int64(AdapterDesc->SharedSystemMemory);
+    FD3D12GlobalStats::GDedicatedVideoMemory = int64(AdapterDesc->DedicatedVideoMemory);
+    FD3D12GlobalStats::GDedicatedSystemMemory = int64(AdapterDesc->DedicatedSystemMemory);
+    FD3D12GlobalStats::GSharedSystemMemory = int64(AdapterDesc->SharedSystemMemory);
 
     // Total amount of system memory, clamped to 8 GB
     int64 TotalPhysicalMemory = FMath::Min(int64(FPlatformMemory::GetConstants().TotalPhysicalGB), 8ll) * (1024ll * 1024ll * 1024ll);
 
     // Consider 50% of the shared memory but max 25% of total system memory.
-    int64 ConsideredSharedSystemMemory = FMath::Min(GSharedSystemMemory / 2ll, TotalPhysicalMemory / 4ll);
+    int64 ConsideredSharedSystemMemory = FMath::Min(FD3D12GlobalStats::GSharedSystemMemory / 2ll, TotalPhysicalMemory / 4ll);
 
     IDXGIAdapter3* DxgiAdapter3 = MainDevice->GetAdapter3();
     TRefCountPtr<IDXGIAdapter> EnumAdapter;
@@ -425,20 +420,20 @@ void FD3D12DynamicRHI::PerRHISetup(FD3D12Device* MainDevice)
 	DXGI_QUERY_VIDEO_MEMORY_INFO LocalVideoMemoryInfo = {};
 	if (!bIsWARP)
 	{
-		GTotalGraphicsMemory = 0;
+		FD3D12GlobalStats::GTotalGraphicsMemory = 0;
 		if (DXGIFactory->EnumAdapters(MainDevice->GetAdapterIndex(), EnumAdapter.GetInitReference()) != DXGI_ERROR_NOT_FOUND)
 		{
 			if (EnumAdapter)
 			{
 				VERIFYD3D11RESULT(DxgiAdapter3->QueryVideoMemoryInfo(0, DXGI_MEMORY_SEGMENT_GROUP_LOCAL, &LocalVideoMemoryInfo));
-				GTotalGraphicsMemory = LocalVideoMemoryInfo.Budget;
+				FD3D12GlobalStats::GTotalGraphicsMemory = LocalVideoMemoryInfo.Budget;
 			}
 		}
 
 	}
 	else
 	{
-		GTotalGraphicsMemory = GDedicatedSystemMemory;
+		FD3D12GlobalStats::GTotalGraphicsMemory = FD3D12GlobalStats::GDedicatedSystemMemory;
 	}
 
 	GMaxTextureMipCount = FMath::CeilLogTwo(GMaxTextureDimensions) + 1;
@@ -455,24 +450,24 @@ void FD3D12DynamicRHI::PerRHISetup(FD3D12Device* MainDevice)
 	if (sizeof(SIZE_T) < 8)
 	{
 		// Clamp to 1 GB if we're less than 64-bit
-		GTotalGraphicsMemory = FMath::Min(GTotalGraphicsMemory, 1024ll * 1024ll * 1024ll);
+		FD3D12GlobalStats::GTotalGraphicsMemory = FMath::Min(FD3D12GlobalStats::GTotalGraphicsMemory, 1024ll * 1024ll * 1024ll);
 	}
 	else
 	{
 		// Clamp to 1.9 GB if we're 64-bit
-		GTotalGraphicsMemory = FMath::Min(GTotalGraphicsMemory, 1945ll * 1024ll * 1024ll);
+		FD3D12GlobalStats::GTotalGraphicsMemory = FMath::Min(FD3D12GlobalStats::GTotalGraphicsMemory, 1945ll * 1024ll * 1024ll);
 	}
 
 	if (!bIsWARP)
 	{
-		VERIFYD3D11RESULT(DxgiAdapter3->SetVideoMemoryReservation(0, DXGI_MEMORY_SEGMENT_GROUP_LOCAL, FMath::Min((int64)LocalVideoMemoryInfo.AvailableForReservation, GTotalGraphicsMemory)));
+		VERIFYD3D11RESULT(DxgiAdapter3->SetVideoMemoryReservation(0, DXGI_MEMORY_SEGMENT_GROUP_LOCAL, FMath::Min((int64)LocalVideoMemoryInfo.AvailableForReservation, FD3D12GlobalStats::GTotalGraphicsMemory)));
 	}
 
 
 
     if (GPoolSizeVRAMPercentage > 0)
     {
-        float PoolSize = float(GPoolSizeVRAMPercentage) * 0.01f * float(GTotalGraphicsMemory);
+        float PoolSize = float(GPoolSizeVRAMPercentage) * 0.01f * float(FD3D12GlobalStats::GTotalGraphicsMemory);
 
         // Truncate GTexturePoolSize to MB (but still counted in bytes)
         GTexturePoolSize = int64(FGenericPlatformMath::TruncToFloat(PoolSize / 1024.0f / 1024.0f)) * 1024 * 1024;
@@ -480,7 +475,7 @@ void FD3D12DynamicRHI::PerRHISetup(FD3D12Device* MainDevice)
         UE_LOG(LogRHI, Log, TEXT("Texture pool is %llu MB (%d%% of %llu MB)"),
             GTexturePoolSize / 1024 / 1024,
             GPoolSizeVRAMPercentage,
-            GTotalGraphicsMemory / 1024 / 1024);
+            FD3D12GlobalStats::GTotalGraphicsMemory / 1024 / 1024);
     }
 
 #if (UE_BUILD_SHIPPING && WITH_EDITOR) && PLATFORM_WINDOWS && !PLATFORM_64BITS
