UE4 - DX12 - Cvars D3D12.* instead of r.RHI.* or D3D12RHI.*

diff --git a/Engine/Source/Runtime/Windows/D3D12RHI/Private/D3D12Query.cpp b/Engine/Source/Runtime/Windows/D3D12RHI/Private/D3D12Query.cpp
--- a/Engine/Source/Runtime/Windows/D3D12RHI/Private/D3D12Query.cpp
+++ b/Engine/Source/Runtime/Windows/D3D12RHI/Private/D3D12Query.cpp
@@ -15,7 +15,7 @@ namespace RHIConsoleVariables
 {
     int32 bStablePowerState = 0;
     static FAutoConsoleVariableRef CVarStablePowerState(
-        TEXT("D3D12RHI.StablePowerState"),
+        TEXT("D3D12.StablePowerState"),
         bStablePowerState,
         TEXT("If true, enable stable power state. This increases GPU timing measurement accuracy but may decrease overall GPU clock rate."),
         ECVF_Default
diff --git a/Engine/Source/Runtime/Windows/D3D12RHI/Private/D3D12RHI.cpp b/Engine/Source/Runtime/Windows/D3D12RHI/Private/D3D12RHI.cpp
--- a/Engine/Source/Runtime/Windows/D3D12RHI/Private/D3D12RHI.cpp
+++ b/Engine/Source/Runtime/Windows/D3D12RHI/Private/D3D12RHI.cpp
@@ -761,7 +761,7 @@ void FD3D12DynamicRHI::GetLocalVideoMemoryInfo (DXGI_QUERY_VIDEO_MEMORY_INFO* Lo
 #if SUPPORTS_MEMORY_RESIDENCY
 int32 bLogMemoryResidency = 1;
 static FAutoConsoleVariableRef CVarLogMemoryResidency(
-	TEXT("r.D3D12LogMemoryResidency"),
+	TEXT("D3D12.LogMemoryResidency"),
 	bLogMemoryResidency,
 	TEXT("Print out a log of the memory residency stats.")
 	TEXT("  0: Off completely\n")
diff --git a/Engine/Source/Runtime/Windows/D3D12RHI/Private/D3D12StateCache.cpp b/Engine/Source/Runtime/Windows/D3D12RHI/Private/D3D12StateCache.cpp
--- a/Engine/Source/Runtime/Windows/D3D12RHI/Private/D3D12StateCache.cpp
+++ b/Engine/Source/Runtime/Windows/D3D12RHI/Private/D3D12StateCache.cpp
@@ -14,7 +14,7 @@
 
 static int32 GEnablePSOCache = 0;
 static FAutoConsoleVariableRef CVarEnablePSOCache(
-	TEXT("r.RHI.EnablePSOCache"),
+	TEXT("D3D12.EnablePSOCache"),
 	GEnablePSOCache,
 	TEXT("Enables a disk cache for PipelineState Objects."),
 	ECVF_RenderThreadSafe | ECVF_ReadOnly
diff --git a/Engine/Source/Runtime/Windows/D3D12RHI/Private/D3D12Viewport.cpp b/Engine/Source/Runtime/Windows/D3D12RHI/Private/D3D12Viewport.cpp
--- a/Engine/Source/Runtime/Windows/D3D12RHI/Private/D3D12Viewport.cpp
+++ b/Engine/Source/Runtime/Windows/D3D12RHI/Private/D3D12Viewport.cpp
@@ -29,7 +29,7 @@ namespace D3D12RHI
 	{
 		int32 bSyncWithDWM = 0;
 		static FAutoConsoleVariableRef CVarSyncWithDWM(
-		TEXT("D3D12RHI.SyncWithDWM"),
+		TEXT("D3D12.SyncWithDWM"),
 			bSyncWithDWM,
 			TEXT("If true, synchronize with the desktop window manager for vblank."),
 			ECVF_RenderThreadSafe
@@ -37,7 +37,7 @@ namespace D3D12RHI
 
 		float RefreshPercentageBeforePresent = 1.0f;
 		static FAutoConsoleVariableRef CVarRefreshPercentageBeforePresent(
-		TEXT("D3D12RHI.RefreshPercentageBeforePresent"),
+		TEXT("D3D12.RefreshPercentageBeforePresent"),
 			RefreshPercentageBeforePresent,
 			TEXT("The percentage of the refresh period to wait before presenting."),
 			ECVF_RenderThreadSafe
@@ -45,7 +45,7 @@ namespace D3D12RHI
 
 		int32 bForceThirtyHz = 1;
 		static FAutoConsoleVariableRef CVarForceThirtyHz(
-		TEXT("D3D12RHI.ForceThirtyHz"),
+		TEXT("D3D12.ForceThirtyHz"),
 			bForceThirtyHz,
 			TEXT("If true, the display will never update more often than 30Hz."),
 			ECVF_RenderThreadSafe
@@ -53,7 +53,7 @@ namespace D3D12RHI
 
 		int32 SyncInterval = 1;
 		static FAutoConsoleVariableRef CVarSyncInterval(
-		TEXT("D3D12RHI.SyncInterval"),
+		TEXT("D3D12.SyncInterval"),
 			SyncInterval,
 			TEXT("When synchronizing with D3D, specifies the interval at which to refresh."),
 			ECVF_RenderThreadSafe
@@ -61,7 +61,7 @@ namespace D3D12RHI
 
 		float SyncRefreshThreshold = 1.05f;
 		static FAutoConsoleVariableRef CVarSyncRefreshThreshold(
-		TEXT("D3D12RHI.SyncRefreshThreshold"),
+		TEXT("D3D12.SyncRefreshThreshold"),
 			SyncRefreshThreshold,
 			TEXT("Threshold for time above which vsync will be disabled as a percentage of the refresh rate."),
 			ECVF_RenderThreadSafe
@@ -69,7 +69,7 @@ namespace D3D12RHI
 
 		int32 MaxSyncCounter = 8;
 		static FAutoConsoleVariableRef CVarMaxSyncCounter(
-			TEXT("D3D12RHI.MaxSyncCounter"),
+			TEXT("D3D12.MaxSyncCounter"),
 			MaxSyncCounter,
 			TEXT("Maximum sync counter to smooth out vsync transitions."),
 			ECVF_RenderThreadSafe
@@ -77,7 +77,7 @@ namespace D3D12RHI
 
 		int32 SyncThreshold = 7;
 		static FAutoConsoleVariableRef CVarSyncThreshold(
-			TEXT("D3D12RHI.SyncThreshold"),
+			TEXT("D3D12.SyncThreshold"),
 			SyncThreshold,
 			TEXT("Number of consecutive 'fast' frames before vsync is enabled."),
 			ECVF_RenderThreadSafe
@@ -85,7 +85,7 @@ namespace D3D12RHI
 
 		int32 MaximumFrameLatency = 3;
 		static FAutoConsoleVariableRef CVarMaximumFrameLatency(
-			TEXT("D3D12RHI.MaximumFrameLatency"),
+			TEXT("D3D12.MaximumFrameLatency"),
 			MaximumFrameLatency,
 			TEXT("Number of frames that can be queued for render."),
 			ECVF_RenderThreadSafe
