UE4 - Integrate MS D3D12 RHI; added RHIPostInit()

diff --git a/Engine/Source/Runtime/Launch/Private/LaunchEngineLoop.cpp b/Engine/Source/Runtime/Launch/Private/LaunchEngineLoop.cpp
--- a/Engine/Source/Runtime/Launch/Private/LaunchEngineLoop.cpp
+++ b/Engine/Source/Runtime/Launch/Private/LaunchEngineLoop.cpp
@@ -1451,6 +1451,8 @@ int32 FEngineLoop::PreInit( const TCHAR* CmdLine )
 	}
 	SlowTask.EnterProgressFrame(5);
 
+	RHIPostInit();
+
 	if (GUseThreadedRendering)
 	{
 		if (GRHISupportsRHIThread)
diff --git a/Engine/Source/Runtime/RHI/Private/DynamicRHI.cpp b/Engine/Source/Runtime/RHI/Private/DynamicRHI.cpp
--- a/Engine/Source/Runtime/RHI/Private/DynamicRHI.cpp
+++ b/Engine/Source/Runtime/RHI/Private/DynamicRHI.cpp
@@ -63,6 +63,12 @@ void RHIInit(bool bHasEditorToken)
 	}
 }
 
+void RHIPostInit()
+{
+	check(GDynamicRHI);
+	GDynamicRHI->PostInit();
+}
+
 void RHIExit()
 {
 	if ( !GUsingNullRHI && GDynamicRHI != NULL )
diff --git a/Engine/Source/Runtime/RHI/Public/DynamicRHI.h b/Engine/Source/Runtime/RHI/Public/DynamicRHI.h
--- a/Engine/Source/Runtime/RHI/Public/DynamicRHI.h
+++ b/Engine/Source/Runtime/RHI/Public/DynamicRHI.h
@@ -23,6 +23,9 @@ public:
 	/** Initializes the RHI; separate from IDynamicRHIModule::CreateRHI so that GDynamicRHI is set when it is called. */
 	virtual void Init() = 0;
 
+	/** Called after the RHI is initialized; before the render thread is started. */
+	virtual void PostInit() {}
+
 	/** Shutdown the RHI; handle shutdown and resource destruction before the RHI's actual destructor is called (so that all resources of the RHI are still available for shutdown). */
 	virtual void Shutdown() = 0;
 
diff --git a/Engine/Source/Runtime/RHI/Public/RHI.h b/Engine/Source/Runtime/RHI/Public/RHI.h
--- a/Engine/Source/Runtime/RHI/Public/RHI.h
+++ b/Engine/Source/Runtime/RHI/Public/RHI.h
@@ -1166,6 +1166,9 @@ DECLARE_MEMORY_STAT_POOL_EXTERN(TEXT("Pixel buffer memory"),STAT_PixelBufferMemo
 /** Initializes the RHI. */
 extern RHI_API void RHIInit(bool bHasEditorToken);
 
+/** Performs additional RHI initialization before the render thread starts. */
+extern RHI_API void RHIPostInit();
+
 /** Shuts down the RHI. */
 extern RHI_API void RHIExit();
 
diff --git a/Engine/Source/Runtime/Windows/D3D12RHI/Private/D3D12Commands.cpp b/Engine/Source/Runtime/Windows/D3D12RHI/Private/D3D12Commands.cpp
--- a/Engine/Source/Runtime/Windows/D3D12RHI/Private/D3D12Commands.cpp
+++ b/Engine/Source/Runtime/Windows/D3D12RHI/Private/D3D12Commands.cpp
@@ -23,8 +23,8 @@ D3D12Commands.cpp: D3D RHI commands implementation.
 // MSFT: Seb: Fix up these D3D11 names and remove namespace
 namespace D3D12RHI
 {
-	FGlobalBoundShaderState GD3D11ClearMRTBoundShaderState[8];
-	TGlobalResource<FVector4VertexDeclaration> GD3D11Vector4VertexDeclaration;
+	FGlobalBoundShaderState GD3D12ClearMRTBoundShaderState[8];
+	TGlobalResource<FVector4VertexDeclaration> GD3D12Vector4VertexDeclaration;
 
 	// MSFT: Seb: Do we need this here?
 	//FGlobalBoundShaderState ResolveBoundShaderState;
@@ -66,7 +66,7 @@ void FD3D12DynamicRHI::SetupRecursiveResources()
 	FRHICommandList_RecursiveHazardous RHICmdList(RHIGetDefaultContext());
 	auto ShaderMap = GetGlobalShaderMap(GMaxRHIFeatureLevel);
 	TShaderMapRef<TOneColorVS<true> > VertexShader(ShaderMap);
-	GD3D11Vector4VertexDeclaration.InitRHI();
+	GD3D12Vector4VertexDeclaration.InitRHI();
 
 	for (int32 NumBuffers = 1; NumBuffers <= MaxSimultaneousRenderTargets; NumBuffers++)
 	{
@@ -113,7 +113,7 @@ void FD3D12DynamicRHI::SetupRecursiveResources()
 			PixelShader = *MRTPixelShader;
 		}
 
-		SetGlobalBoundShaderState(RHICmdList, GMaxRHIFeatureLevel, GD3D11ClearMRTBoundShaderState[NumBuffers - 1], GD3D11Vector4VertexDeclaration.VertexDeclarationRHI, *VertexShader, PixelShader);
+		SetGlobalBoundShaderState(RHICmdList, GMaxRHIFeatureLevel, GD3D12ClearMRTBoundShaderState[NumBuffers - 1], GD3D12Vector4VertexDeclaration.VertexDeclarationRHI, *VertexShader, PixelShader);
 	}
 
 	// MSFT: Seb: Is this needed?
@@ -1780,36 +1780,36 @@ void FD3D12CommandContext::RHIClearMRTImpl(bool bClearColor, int32 NumClearColor
 		/** The global D3D device's immediate context */
 		FDeviceStateHelper() {}
 
-		void CaptureDeviceState(FD3D12StateCache& StateCache)
+		void CaptureDeviceState(FD3D12StateCache& StateCacheRef)
 		{
-			StateCache.GetBoundShaderState(&OldShaderState);
-			StateCache.GetShaderResourceViews<SF_Vertex>(0, NumVertResources, &VertResources[0]);
-			StateCache.GetDepthStencilState(&pOldDepthStencilState, &StencilRef);
-			StateCache.GetBlendState(&pOldBlendState, BlendFactor, &SampleMask);
-			StateCache.GetRasterizerState(&pOldRasterizerState);
+			StateCacheRef.GetBoundShaderState(&OldShaderState);
+			StateCacheRef.GetShaderResourceViews<SF_Vertex>(0, NumVertResources, &VertResources[0]);
+			StateCacheRef.GetDepthStencilState(&pOldDepthStencilState, &StencilRef);
+			StateCacheRef.GetBlendState(&pOldBlendState, BlendFactor, &SampleMask);
+			StateCacheRef.GetRasterizerState(&pOldRasterizerState);
 		}
 
-		void ClearCurrentVertexResources(FD3D12StateCache& StateCache)
+		void ClearCurrentVertexResources(FD3D12StateCache& StateCacheRef)
 		{
 			static FD3D12ShaderResourceView* NullResources[ResourceCount] = {};
 			for (uint32 ResourceLoop = 0 ; ResourceLoop < NumVertResources; ResourceLoop++)
 			{
-				StateCache.SetShaderResourceView<SF_Vertex>(NullResources[0],0);
+				StateCacheRef.SetShaderResourceView<SF_Vertex>(NullResources[0], 0);
 			}
 		}
 
-		void RestoreDeviceState(FD3D12StateCache& StateCache) 
+		void RestoreDeviceState(FD3D12StateCache& StateCacheRef)
 		{
 
 			// Restore the old shaders
-			StateCache.SetBoundShaderState(OldShaderState);
+			StateCacheRef.SetBoundShaderState(OldShaderState);
 			for (uint32 ResourceLoop = 0 ; ResourceLoop < NumVertResources; ResourceLoop++)
 			{
-				StateCache.SetShaderResourceView<SF_Vertex>(VertResources[ResourceLoop],ResourceLoop);
+				StateCacheRef.SetShaderResourceView<SF_Vertex>(VertResources[ResourceLoop], ResourceLoop);
 			}
-			StateCache.SetDepthStencilState(pOldDepthStencilState, StencilRef);
-			StateCache.SetBlendState(pOldBlendState, BlendFactor, SampleMask);
-			StateCache.SetRasterizerState(pOldRasterizerState);
+			StateCacheRef.SetDepthStencilState(pOldDepthStencilState, StencilRef);
+			StateCacheRef.SetBlendState(pOldBlendState, BlendFactor, SampleMask);
+			StateCacheRef.SetRasterizerState(pOldRasterizerState);
 
 			ReleaseResources();
 		}
@@ -1913,6 +1913,19 @@ void FD3D12CommandContext::RHIClearMRTImpl(bool bClearColor, int32 NumClearColor
 			D3D12_RESOURCE_DESC const& Desc = BaseTexture->GetDesc();
 			Width = (uint32)Desc.Width;
 			Height = Desc.Height;
+
+			// Adjust dimensions for the mip level we're clearing.
+			D3D12_DEPTH_STENCIL_VIEW_DESC DSVDesc = DepthStencilView->GetDesc();
+			if (DSVDesc.ViewDimension == D3D12_DSV_DIMENSION_TEXTURE1D ||
+				DSVDesc.ViewDimension == D3D12_DSV_DIMENSION_TEXTURE1DARRAY ||
+				DSVDesc.ViewDimension == D3D12_DSV_DIMENSION_TEXTURE2D ||
+				DSVDesc.ViewDimension == D3D12_DSV_DIMENSION_TEXTURE2DARRAY)
+			{
+				// All the non-multisampled texture types have their mip-slice in the same position.
+				uint32 MipIndex = DSVDesc.Texture2D.MipSlice;
+				Width >>= MipIndex;
+				Height >>= MipIndex;
+			}
 		}
 
 		if ((Viewport.Width < Width || Viewport.Height < Height) 
@@ -2045,7 +2058,7 @@ void FD3D12CommandContext::RHIClearMRTImpl(bool bClearColor, int32 NumClearColor
 
 		{
 			FRHICommandList_RecursiveHazardous RHICmdList(this);
-			SetGlobalBoundShaderState(RHICmdList, GMaxRHIFeatureLevel, GD3D11ClearMRTBoundShaderState[FMath::Max(BoundRenderTargets.GetNumActiveTargets() - 1, 0)], GD3D11Vector4VertexDeclaration.VertexDeclarationRHI, *VertexShader, PixelShader);
+			SetGlobalBoundShaderState(RHICmdList, GMaxRHIFeatureLevel, GD3D12ClearMRTBoundShaderState[FMath::Max(BoundRenderTargets.GetNumActiveTargets() - 1, 0)], GD3D12Vector4VertexDeclaration.VertexDeclarationRHI, *VertexShader, PixelShader);
 			PixelShader->SetColors(RHICmdList, ClearColorArray, NumClearColors);
 
 			{
diff --git a/Engine/Source/Runtime/Windows/D3D12RHI/Private/D3D12Device.cpp b/Engine/Source/Runtime/Windows/D3D12RHI/Private/D3D12Device.cpp
--- a/Engine/Source/Runtime/Windows/D3D12RHI/Private/D3D12Device.cpp
+++ b/Engine/Source/Runtime/Windows/D3D12RHI/Private/D3D12Device.cpp
@@ -427,6 +427,13 @@ void FD3D12Device::CreateCommandContexts()
 	CommandContextArray[0]->OpenCommandList();
 }
 
+FD3D12DynamicRHI::~FD3D12DynamicRHI()
+{
+	UE_LOG(LogD3D12RHI, Log, TEXT("~FD3D12DynamicRHI"));
+
+	check(MainDevice == nullptr);
+}
+
 void FD3D12DynamicRHI::Shutdown()
 {
 	check(IsInGameThread() && IsInRenderingThread());  // require that the render thread has been shut down
@@ -434,6 +441,7 @@ void FD3D12DynamicRHI::Shutdown()
 	// Cleanup the D3D devices.
 	MainDevice->CleanupD3DDevice();
 	delete(MainDevice);
+	MainDevice = nullptr;
 
 	// Release buffered timestamp queries
 	GPUProfilingData.FrameTiming.ReleaseResource();
@@ -466,14 +474,14 @@ void FD3D12DynamicRHI::RHIGetSupportedResolution( uint32 &Width, uint32 &Height
 	BestMode.Height = 0;
 
 	{
-		HRESULT hr = S_OK;
+		HRESULT HResult = S_OK;
 		TRefCountPtr<IDXGIAdapter> Adapter;
-        hr = DXGIFactory->EnumAdapters(GetRHIDevice()->GetAdapterIndex(), Adapter.GetInitReference());
-		if( DXGI_ERROR_NOT_FOUND == hr )
+		HResult = DXGIFactory->EnumAdapters(GetRHIDevice()->GetAdapterIndex(), Adapter.GetInitReference());
+		if (DXGI_ERROR_NOT_FOUND == HResult)
 		{
 			return;
 		}
-		if( FAILED(hr) )
+		if (FAILED(HResult))
 		{
 			return;
 		}
@@ -488,22 +496,22 @@ void FD3D12DynamicRHI::RHIGetSupportedResolution( uint32 &Width, uint32 &Height
 		for(uint32 o = 0;o < 1; o++)
 		{
 			TRefCountPtr<IDXGIOutput> Output;
-			hr = Adapter->EnumOutputs(o,Output.GetInitReference());
-			if(DXGI_ERROR_NOT_FOUND == hr)
+			HResult = Adapter->EnumOutputs(o,Output.GetInitReference());
+			if(DXGI_ERROR_NOT_FOUND == HResult)
 				break;
-			if(FAILED(hr))
+			if(FAILED(HResult))
 				return;
 
 			// TODO: GetDisplayModeList is a terribly SLOW call.  It can take up to a second per invocation.
 			//  We might want to work around some DXGI badness here.
 			DXGI_FORMAT Format = DXGI_FORMAT_R8G8B8A8_UNORM;
 			uint32 NumModes = 0;
-			hr = Output->GetDisplayModeList(Format,0,&NumModes,NULL);
-			if(hr == DXGI_ERROR_NOT_FOUND)
+			HResult = Output->GetDisplayModeList(Format,0,&NumModes,NULL);
+			if(HResult == DXGI_ERROR_NOT_FOUND)
 			{
 				return;
 			}
-			else if(hr == DXGI_ERROR_NOT_CURRENTLY_AVAILABLE)
+			else if(HResult == DXGI_ERROR_NOT_CURRENTLY_AVAILABLE)
 			{
 				UE_LOG(LogD3D12RHI, Fatal,
 					TEXT("This application cannot be run over a remote desktop configuration")
@@ -848,12 +856,6 @@ void FD3D12DynamicRHI::RHIFlushResources()
 
 void FD3D12DynamicRHI::RHIAcquireThreadOwnership()
 {
-#if WITH_EDITOR
-	if (GRHISupportsRHIThread)
-	{
-		SetupRecursiveResources();
-	}
-#endif
 }
 
 void FD3D12DynamicRHI::RHIReleaseThreadOwnership()
diff --git a/Engine/Source/Runtime/Windows/D3D12RHI/Private/D3D12RHI.cpp b/Engine/Source/Runtime/Windows/D3D12RHI/Private/D3D12RHI.cpp
--- a/Engine/Source/Runtime/Windows/D3D12RHI/Private/D3D12RHI.cpp
+++ b/Engine/Source/Runtime/Windows/D3D12RHI/Private/D3D12RHI.cpp
@@ -267,7 +267,7 @@ void FD3D12DynamicRHI::IssueLongGPUTask()
 			TShaderMapRef<TOneColorVS<true> > VertexShader(ShaderMap);
 			TShaderMapRef<FLongGPUTaskPS> PixelShader(ShaderMap);
 
-			RHICmdList.SetLocalBoundShaderState(RHICmdList.BuildLocalBoundShaderState(GD3D11Vector4VertexDeclaration.VertexDeclarationRHI, VertexShader->GetVertexShader(), FHullShaderRHIRef(), FDomainShaderRHIRef(), PixelShader->GetPixelShader(), FGeometryShaderRHIRef()));
+			RHICmdList.SetLocalBoundShaderState(RHICmdList.BuildLocalBoundShaderState(GD3D12Vector4VertexDeclaration.VertexDeclarationRHI, VertexShader->GetVertexShader(), FHullShaderRHIRef(), FDomainShaderRHIRef(), PixelShader->GetPixelShader(), FGeometryShaderRHIRef()));
 
 			// Draw a fullscreen quad
 			FVector4 Vertices[4];
diff --git a/Engine/Source/Runtime/Windows/D3D12RHI/Private/D3D12RHIPrivate.h b/Engine/Source/Runtime/Windows/D3D12RHI/Private/D3D12RHIPrivate.h
--- a/Engine/Source/Runtime/Windows/D3D12RHI/Private/D3D12RHIPrivate.h
+++ b/Engine/Source/Runtime/Windows/D3D12RHI/Private/D3D12RHIPrivate.h
@@ -946,10 +946,11 @@ public:
     FD3D12DynamicRHI(IDXGIFactory4* InDXGIFactory, FD3D12Adapter& InAdapter);
 
 	/** Destructor */
-	virtual ~FD3D12DynamicRHI() {}
+	virtual ~FD3D12DynamicRHI();
 
 	// FDynamicRHI interface.
 	virtual void Init() override;
+	virtual void PostInit() override;
 	virtual void Shutdown() override;
 
     template<typename TRHIType>
@@ -1900,7 +1901,7 @@ public:
 
 namespace D3D12RHI
 {
-	extern TGlobalResource<FVector4VertexDeclaration> GD3D11Vector4VertexDeclaration;
+	extern TGlobalResource<FVector4VertexDeclaration> GD3D12Vector4VertexDeclaration;
 
 	/**
 	 *	Default 'Fast VRAM' allocator...
diff --git a/Engine/Source/Runtime/Windows/D3D12RHI/Private/D3D12StateCache.cpp b/Engine/Source/Runtime/Windows/D3D12RHI/Private/D3D12StateCache.cpp
--- a/Engine/Source/Runtime/Windows/D3D12RHI/Private/D3D12StateCache.cpp
+++ b/Engine/Source/Runtime/Windows/D3D12RHI/Private/D3D12StateCache.cpp
@@ -1149,7 +1149,9 @@ void FD3D12StateCacheBase::ApplyState(bool IsCompute)
 		{
 			if ((bNeedSetSRVs && bNeedSetSRVsPerShaderStage[Stage]) || bForceState)
 			{
-				NumSRVs[Stage] = PipelineState.Common.MaxBoundShaderResourcesIndex[Stage] + 1;
+				// Disable SRV slot optimization for now. This will causes us to use a lot more descriptor heap slots.
+				//NumSRVs[Stage] = PipelineState.Common.MaxBoundShaderResourcesIndex[Stage] + 1;
+				NumSRVs[Stage] = MAX_SRVS;
 			}
 			if ((bNeedSetConstantBuffers && (bNeedSetConstantBuffersPerShaderStage[Stage])) || bForceState)
 			{
diff --git a/Engine/Source/Runtime/Windows/D3D12RHI/Private/D3D12Texture.cpp b/Engine/Source/Runtime/Windows/D3D12RHI/Private/D3D12Texture.cpp
--- a/Engine/Source/Runtime/Windows/D3D12RHI/Private/D3D12Texture.cpp
+++ b/Engine/Source/Runtime/Windows/D3D12RHI/Private/D3D12Texture.cpp
@@ -1374,6 +1374,9 @@ FShaderResourceViewRHIRef FD3D12DynamicRHI::RHICreateShaderResourceView(FTexture
 	uint64 Offset = 0;
     const D3D12_RESOURCE_DESC& TextureDesc = Texture2D->GetResource()->GetDesc();
 
+	bool bSRGB = (Texture2D->GetFlags() & TexCreate_SRGB) != 0;
+	const DXGI_FORMAT PlatformShaderResourceFormat = FindShaderResourceDXGIFormat(TextureDesc.Format, bSRGB);
+
 	// Create a Shader Resource View
     D3D12_SHADER_RESOURCE_VIEW_DESC SRVDesc = {};
     SRVDesc.Shader4ComponentMapping = D3D12_DEFAULT_SHADER_4_COMPONENT_MAPPING;
@@ -1381,7 +1384,7 @@ FShaderResourceViewRHIRef FD3D12DynamicRHI::RHICreateShaderResourceView(FTexture
 	SRVDesc.Texture2D.MostDetailedMip = MipLevel;
 	SRVDesc.Texture2D.MipLevels = 1;
 
-	SRVDesc.Format = TextureDesc.Format;
+	SRVDesc.Format = PlatformShaderResourceFormat;
 
 	SRVDesc.Texture2D.PlaneSlice = GetPlaneSliceFromViewFormat(TextureDesc.Format, SRVDesc.Format);
 
diff --git a/Engine/Source/Runtime/Windows/D3D12RHI/Private/D3D12Util.cpp b/Engine/Source/Runtime/Windows/D3D12RHI/Private/D3D12Util.cpp
--- a/Engine/Source/Runtime/Windows/D3D12RHI/Private/D3D12Util.cpp
+++ b/Engine/Source/Runtime/Windows/D3D12RHI/Private/D3D12Util.cpp
@@ -163,6 +163,26 @@ static void TerminateOnDeviceRemoved(HRESULT D3DResult)
 	}
 }
 
+
+static void TerminateOnOutOfMemory(HRESULT D3DResult, bool bCreatingTextures)
+{
+	if (D3DResult == E_OUTOFMEMORY)
+	{
+		if (bCreatingTextures)
+		{
+			FPlatformMisc::MessageBoxExt(EAppMsgType::Ok, *LOCTEXT("OutOfVideoMemoryTextures", "Out of video memory trying to allocate a texture! Make sure your video card has the minimum required memory, try lowering the resolution and/or closing other applications that are running. Exiting...").ToString(), TEXT("Error"));
+		}
+		else
+		{
+			FPlatformMisc::MessageBoxExt(EAppMsgType::Ok, *NSLOCTEXT("D3D12RHI", "OutOfMemory", "Out of video memory trying to allocate a rendering resource. Make sure your video card has the minimum required memory, try lowering the resolution and/or closing other applications that are running. Exiting...").ToString(), TEXT("Error"));
+		}
+#if STATS
+		GetRendererModule().DebugLogOnCrash();
+#endif
+		FPlatformMisc::RequestExit(true);
+	}
+}
+
 #ifndef MAKE_D3DHRESULT
 #define _FACD3D						0x876
 #define MAKE_D3DHRESULT( code)		MAKE_HRESULT( 1, _FACD3D, code )
@@ -179,16 +199,7 @@ namespace D3D12RHI
 		UE_LOG(LogD3D12RHI, Error, TEXT("%s failed \n at %s:%u \n with error %s"), ANSI_TO_TCHAR(Code), ANSI_TO_TCHAR(Filename), Line, *ErrorString);
 
 		TerminateOnDeviceRemoved(D3DResult);
-
-		// this is to track down a rarely happening crash
-		if (D3DResult == E_OUTOFMEMORY)
-		{
-			FPlatformMisc::MessageBoxExt(EAppMsgType::Ok, *NSLOCTEXT("D3D12RHI", "OutOfMemory", "Out of video memory trying to allocate a rendering resource.").ToString(), TEXT("Error"));
-#if STATS
-			GetRendererModule().DebugLogOnCrash();
-#endif
-			FPlatformMisc::RequestExit(/*bForce=*/ true);
-		}
+		TerminateOnOutOfMemory(D3DResult, false);
 
 		UE_LOG(LogD3D12RHI, Fatal, TEXT("%s failed \n at %s:%u \n with error %s"), ANSI_TO_TCHAR(Code), ANSI_TO_TCHAR(Filename), Line, *ErrorString);
 	}
@@ -200,9 +211,22 @@ namespace D3D12RHI
 		const FString ErrorString = GetD3D11ErrorString(D3DResult, 0);
 		const TCHAR* D3DFormatString = GetD3D11TextureFormatString((DXGI_FORMAT)Format);
 
-		UE_LOG(LogD3D12RHI, Error, TEXT("%s failed \n at %s:%u \n with error %s"), ANSI_TO_TCHAR(Code), ANSI_TO_TCHAR(Filename), Line, *ErrorString);
+		UE_LOG(LogD3D12RHI, Error,
+			TEXT("%s failed \n at %s:%u \n with error %s, \n Size=%ix%ix%i Format=%s(0x%08X), NumMips=%i, Flags=%s"),
+			ANSI_TO_TCHAR(Code),
+			ANSI_TO_TCHAR(Filename),
+			Line,
+			*ErrorString,
+			SizeX,
+			SizeY,
+			SizeZ,
+			D3DFormatString,
+			Format,
+			NumMips,
+			*GetD3D12TextureFlagString(Flags));
 
 		TerminateOnDeviceRemoved(D3DResult);
+		TerminateOnOutOfMemory(D3DResult, true);
 
 		// this is to track down a rarely happening crash
 		if (D3DResult == E_OUTOFMEMORY)
diff --git a/Engine/Source/Runtime/Windows/D3D12RHI/Private/Windows/WindowsD3D12Device.cpp b/Engine/Source/Runtime/Windows/D3D12RHI/Private/Windows/WindowsD3D12Device.cpp
--- a/Engine/Source/Runtime/Windows/D3D12RHI/Private/Windows/WindowsD3D12Device.cpp
+++ b/Engine/Source/Runtime/Windows/D3D12RHI/Private/Windows/WindowsD3D12Device.cpp
@@ -379,13 +379,14 @@ void FD3D12DynamicRHI::Init()
     PresentCount = 0;
 #endif
 	InitD3DDevices();
+}
 
-#if !WITH_EDITOR
+void FD3D12DynamicRHI::PostInit()
+{
 	if (GRHISupportsRHIThread)
 	{
 		SetupRecursiveResources();
 	}
-#endif
 }
 
 void FD3D12DynamicRHI::PerRHISetup(FD3D12Device* MainDevice)
@@ -789,14 +790,14 @@ bool FD3D12DynamicRHI::RHIGetAvailableResolutions(FScreenResolutionArray& Resolu
 		MaxAllowableRefreshRate = 10480;
 	}
 
-	HRESULT hr = S_OK;
+	HRESULT HResult = S_OK;
 	TRefCountPtr<IDXGIAdapter> Adapter;
     //TODO: should only be called on display out device
-    hr = DXGIFactory->EnumAdapters(GetRHIDevice()->GetAdapterIndex(), Adapter.GetInitReference());
+	HResult = DXGIFactory->EnumAdapters(GetRHIDevice()->GetAdapterIndex(), Adapter.GetInitReference());
 
-	if (DXGI_ERROR_NOT_FOUND == hr)
+	if (DXGI_ERROR_NOT_FOUND == HResult)
 		return false;
-	if (FAILED(hr))
+	if (FAILED(HResult))
 		return false;
 
 	// get the description of the adapter
@@ -810,22 +811,22 @@ bool FD3D12DynamicRHI::RHIGetAvailableResolutions(FScreenResolutionArray& Resolu
 	do 
 	{
 		TRefCountPtr<IDXGIOutput> Output;
-		hr = Adapter->EnumOutputs(CurrentOutput, Output.GetInitReference());
-		if (DXGI_ERROR_NOT_FOUND == hr)
+		HResult = Adapter->EnumOutputs(CurrentOutput, Output.GetInitReference());
+		if (DXGI_ERROR_NOT_FOUND == HResult)
 			break;
-		if (FAILED(hr))
+		if (FAILED(HResult))
 			return false;
 
 		// TODO: GetDisplayModeList is a terribly SLOW call.  It can take up to a second per invocation.
 		//  We might want to work around some DXGI badness here.
 		DXGI_FORMAT Format = DXGI_FORMAT_R8G8B8A8_UNORM;
 		uint32 NumModes = 0;
-		hr = Output->GetDisplayModeList(Format, 0, &NumModes, NULL);
-		if (hr == DXGI_ERROR_NOT_FOUND)
+		HResult = Output->GetDisplayModeList(Format, 0, &NumModes, NULL);
+		if (HResult == DXGI_ERROR_NOT_FOUND)
 		{
 			continue;
 		}
-		else if (hr == DXGI_ERROR_NOT_CURRENTLY_AVAILABLE)
+		else if (HResult == DXGI_ERROR_NOT_CURRENTLY_AVAILABLE)
 		{
 			UE_LOG(LogD3D12RHI, Fatal,
 				TEXT("This application cannot be run over a remote desktop configuration")
