UE4 - D3D12 minor fixes

diff --git a/Engine/Source/Runtime/Windows/D3D12RHI/Private/D3D12Commands.cpp b/Engine/Source/Runtime/Windows/D3D12RHI/Private/D3D12Commands.cpp
--- a/Engine/Source/Runtime/Windows/D3D12RHI/Private/D3D12Commands.cpp
+++ b/Engine/Source/Runtime/Windows/D3D12RHI/Private/D3D12Commands.cpp
@@ -265,7 +265,7 @@ void FD3D12CommandContext::RHISetViewport(uint32 MinX, uint32 MinY, float MinZ,
 
 void FD3D12CommandContext::RHISetScissorRect(bool bEnable, uint32 MinX, uint32 MinY, uint32 MaxX, uint32 MaxY)
 {
-	if(bEnable)
+	if (bEnable)
 	{
 		D3D12_RECT ScissorRect;
 		ScissorRect.left = MinX;
@@ -296,7 +296,7 @@ void FD3D12CommandContext::RHISetBoundShaderState(FBoundShaderStateRHIParamRef B
 
 	StateCache.SetBoundShaderState(BoundShaderState);
 
-	if(BoundShaderState->GetHullShader() && BoundShaderState->GetDomainShader())
+	if (BoundShaderState->GetHullShader() && BoundShaderState->GetDomainShader())
 	{
 		bUsingTessellation = true;
 	}
@@ -434,9 +434,13 @@ void FD3D12CommandContext::RHISetShaderTexture(FComputeShaderRHIParamRef Compute
 	FD3D12ResourceLocation* ResourceLocation = NewTexture ? NewTexture->ResourceLocation : nullptr;
 
 	if ( ( NewTexture == NULL) || ( NewTexture->GetRenderTargetView( 0, 0 ) !=NULL) || ( NewTexture->HasDepthStencilView()) )
+	{
 		SetShaderResourceView<SF_Compute>(ResourceLocation, ShaderResourceView, TextureIndex, FD3D12StateCache::SRV_Dynamic);
+	}
 	else
+	{
 		SetShaderResourceView<SF_Compute>(ResourceLocation, ShaderResourceView, TextureIndex, FD3D12StateCache::SRV_Static);
+	}
 
 	OwningRHI.IncrementSetShaderTextureCycles(FPlatformTime::Cycles() - Start);
 	OwningRHI.IncrementSetShaderTextureCalls();
@@ -448,7 +452,7 @@ void FD3D12CommandContext::RHISetUAVParameter(FComputeShaderRHIParamRef ComputeS
 
 	FD3D12UnorderedAccessView* UAV = FD3D12DynamicRHI::ResourceCast(UAVRHI);
 
-	if(UAV)
+	if (UAV)
 	{
 		ConditionalClearShaderResource(UAV->GetResourceLocation());
 	}
@@ -466,7 +470,7 @@ void FD3D12CommandContext::RHISetUAVParameter(FComputeShaderRHIParamRef ComputeS
 
 	FD3D12UnorderedAccessView* UAV = FD3D12DynamicRHI::ResourceCast(UAVRHI);
 
-	if(UAV)
+	if (UAV)
 	{
 		ConditionalClearShaderResource(UAV->GetResourceLocation());
 	}
@@ -952,7 +956,7 @@ void FD3D12CommandContext::RHISetRenderTargets(
 
 	// Set the appropriate depth stencil view depending on whether depth writes are enabled or not
 	FD3D12DepthStencilView* DepthStencilView = NULL;
-	if(NewDepthStencilTarget)
+	if (NewDepthStencilTarget)
 	{
 		CurrentDSVAccessType = NewDepthStencilTargetRHI->GetDepthStencilAccess();
 		DepthStencilView = NewDepthStencilTarget->GetDepthStencilView(CurrentDSVAccessType);
@@ -962,7 +966,7 @@ void FD3D12CommandContext::RHISetRenderTargets(
 	}
 
 	// Check if the depth stencil target is different from the old state.
-	if(CurrentDepthStencilTarget != DepthStencilView)
+	if (CurrentDepthStencilTarget != DepthStencilView)
 	{
 		if (DepthStencilView)
 		{
@@ -1006,7 +1010,7 @@ void FD3D12CommandContext::RHISetRenderTargets(
 			// A check to allow you to pinpoint what is using mismatching targets
 			// We filter our d3ddebug spew that checks for this as the d3d runtime's check is wrong.
 			// For filter code, see D3D12Device.cpp look for "OMSETRENDERTARGETS_INVALIDVIEW"
-			if(RenderTargetView && DepthStencilView)
+			if (RenderTargetView && DepthStencilView)
 			{
 				FRTVDesc RTTDesc = GetRenderTargetViewDesc(RenderTargetView);
 
@@ -1015,7 +1019,7 @@ void FD3D12CommandContext::RHISetRenderTargets(
 				D3D12_RESOURCE_DESC DTTDesc = DepthTargetTexture->GetDesc();
 
 				// enforce color target is <= depth and MSAA settings match
-				if(RTTDesc.Width > DTTDesc.Width || RTTDesc.Height > DTTDesc.Height || 
+				if (RTTDesc.Width > DTTDesc.Width || RTTDesc.Height > DTTDesc.Height || 
 					RTTDesc.SampleDesc.Count != DTTDesc.SampleDesc.Count || 
 					RTTDesc.SampleDesc.Quality != DTTDesc.SampleDesc.Quality)
 				{
@@ -1030,13 +1034,13 @@ void FD3D12CommandContext::RHISetRenderTargets(
 		NewRenderTargetViews[RenderTargetIndex] = RenderTargetView;
 
 		// Check if the render target is different from the old state.
-		if(CurrentRenderTargets[RenderTargetIndex] != RenderTargetView)
+		if (CurrentRenderTargets[RenderTargetIndex] != RenderTargetView)
 		{
 			CurrentRenderTargets[RenderTargetIndex] = RenderTargetView;
 			bTargetChanged = true;
 		}
 	}
-	if(NumSimultaneousRenderTargets != NewNumSimultaneousRenderTargets)
+	if (NumSimultaneousRenderTargets != NewNumSimultaneousRenderTargets)
 	{
 		NumSimultaneousRenderTargets = NewNumSimultaneousRenderTargets;
 		bTargetChanged = true;
@@ -1046,7 +1050,7 @@ void FD3D12CommandContext::RHISetRenderTargets(
 	for(uint32 UAVIndex = 0;UAVIndex < MaxSimultaneousUAVs;++UAVIndex)
 	{
 		FD3D12UnorderedAccessView* RHIUAV = NULL;
-		if(UAVIndex < NewNumUAVs && UAVs[UAVIndex] != NULL)
+		if (UAVIndex < NewNumUAVs && UAVs[UAVIndex] != NULL)
 		{
 			RHIUAV = (FD3D12UnorderedAccessView*)UAVs[UAVIndex];
 			FD3D12DynamicRHI::TransitionResource(CommandListHandle, RHIUAV, D3D12_RESOURCE_STATE_UNORDERED_ACCESS);
@@ -1055,20 +1059,20 @@ void FD3D12CommandContext::RHISetRenderTargets(
 			ConditionalClearShaderResource(RHIUAV->GetResourceLocation());
 		}
 
-		if(CurrentUAVs[UAVIndex] != RHIUAV)
+		if (CurrentUAVs[UAVIndex] != RHIUAV)
 		{
 			CurrentUAVs[UAVIndex] = RHIUAV;
 			bTargetChanged = true;
 		}
 	}
-	if(NumUAVs != NewNumUAVs)
+	if (NumUAVs != NewNumUAVs)
 	{
 		NumUAVs = NewNumUAVs;
 		bTargetChanged = true;
 	}
 
 	// Only make the D3D call to change render targets if something actually changed.
-	if(bTargetChanged)
+	if (bTargetChanged)
 	{
 		CommitRenderTargetsAndUAVs();
 	}
@@ -1129,7 +1133,7 @@ void FD3D12CommandContext::RHIBeginRenderQuery(FRenderQueryRHIParamRef QueryRHI)
 {
 	FD3D12OcclusionQuery* Query = FD3D12DynamicRHI::ResourceCast(QueryRHI);
 
-	if(Query->Type == RQT_Occlusion)
+	if (Query->Type == RQT_Occlusion)
 	{
 		Query->bResultIsCached = false;
 		Query->HeapIndex = GetParentDevice()->GetQueryHeap()->BeginQuery(*this, D3D12_QUERY_TYPE_OCCLUSION);
@@ -1170,7 +1174,7 @@ void FD3D12CommandContext::RHIEndRenderQuery(FRenderQueryRHIParamRef QueryRHI)
 
 static D3D_PRIMITIVE_TOPOLOGY GetD3D11PrimitiveType(uint32 PrimitiveType, bool bUsingTessellation)
 {
-	if(bUsingTessellation)
+	if (bUsingTessellation)
 	{
 		switch(PrimitiveType)
 		{
@@ -1268,7 +1272,7 @@ void FD3D12CommandContext::CommitNonComputeShaderConstants()
 	// Note that this is *potentially* unsafe because bDiscardSharedConstants is cleared at the
 	// end of the function, however we're OK for now because bDiscardSharedConstants
 	// is always reset whenever bUsingTessellation changes in SetBoundShaderState()
-	if(bUsingTessellation)
+	if (bUsingTessellation)
 	{
 		if (CurrentBoundShaderState->bShaderNeedsGlobalConstantBuffer[SF_Hull])
 		{
@@ -1824,19 +1828,19 @@ void FD3D12CommandContext::RHIClearMRTImpl(bool bClearColor, int32 NumClearColor
 		ClearWithExcludeRects = ExcludeRectCVar->GetValueOnRenderThread();
 #endif
 
-		if(ClearWithExcludeRects >= 2)
+		if (ClearWithExcludeRects >= 2)
 		{
 			// by default use the exclude rect
 			ClearWithExcludeRects = 1;
 
-			if(IsRHIDeviceIntel())
+			if (IsRHIDeviceIntel())
 			{
 				// Disable exclude rect (Intel has fast clear so better we disable)
 				ClearWithExcludeRects = 0;
 			}
 		}
 
-		if(!ClearWithExcludeRects)
+		if (!ClearWithExcludeRects)
 		{
 			// Disable exclude rect
 			ExcludeRect = FIntRect();
@@ -1875,12 +1879,12 @@ void FD3D12CommandContext::RHIClearMRTImpl(bool bClearColor, int32 NumClearColor
 	}
 
 	/*	// possible optimization
-	if(ExcludeRect.Width() > 0 && ExcludeRect.Height() > 0 && HardwareHasLinearClearPerformance) 
+	if (ExcludeRect.Width() > 0 && ExcludeRect.Height() > 0 && HardwareHasLinearClearPerformance) 
 	{
 	UseDrawClear = true;
 	}
 	*/
-	if(ExcludeRect.Min.X == 0 && ExcludeRect.Width() == Viewport.Width && ExcludeRect.Min.Y == 0 && ExcludeRect.Height() == Viewport.Height)
+	if (ExcludeRect.Min.X == 0 && ExcludeRect.Width() == Viewport.Width && ExcludeRect.Min.Y == 0 && ExcludeRect.Height() == Viewport.Height)
 	{
 		// no need to do anything
 		return;
@@ -2063,7 +2067,7 @@ void FD3D12CommandContext::RHIClearMRTImpl(bool bClearColor, int32 NumClearColor
 
 			{
 				// Draw a fullscreen quad
-				if(ExcludeRect.Width() > 0 && ExcludeRect.Height() > 0)
+				if (ExcludeRect.Width() > 0 && ExcludeRect.Height() > 0)
 				{
 					// with a hole in it (optimization in case the hardware has non constant clear performance)
 					FVector4 OuterVertices[4];
diff --git a/Engine/Source/Runtime/Windows/D3D12RHI/Private/D3D12Device.cpp b/Engine/Source/Runtime/Windows/D3D12RHI/Private/D3D12Device.cpp
--- a/Engine/Source/Runtime/Windows/D3D12RHI/Private/D3D12Device.cpp
+++ b/Engine/Source/Runtime/Windows/D3D12RHI/Private/D3D12Device.cpp
@@ -32,7 +32,7 @@ bool D3D12RHI_ShouldAllowAsyncResourceCreation()
 
 IMPLEMENT_MODULE(FD3D12DynamicRHIModule, D3D12RHI);
 
-FD3D12DynamicRHI *FD3D12DynamicRHI::SingleD3DRHI = nullptr;
+FD3D12DynamicRHI* FD3D12DynamicRHI::SingleD3DRHI = nullptr;
 
 IRHICommandContext* FD3D12DynamicRHI::RHIGetDefaultContext()
 {
@@ -54,7 +54,7 @@ public:
 
 	/** Custom new/delete with recycling */
 	void* operator new(size_t Size);
-	void operator delete(void *RawMemory);
+	void operator delete(void* RawMemory);
 
 	FD3D12CommandContextContainer(FD3D12Device* InOwningDevice)
 		: OwningDevice(InOwningDevice), CmdContext(nullptr)
@@ -73,7 +73,9 @@ public:
 		check(CmdContext != nullptr);
 
 		if (CmdContext->CommandListHandle == nullptr)
+		{
 			CmdContext->OpenCommandList();
+		}
 
 		CmdContext->ClearState();
 
@@ -286,7 +288,7 @@ FD3D12DynamicRHI::FD3D12DynamicRHI(IDXGIFactory4* InDXGIFactory, FD3D12Adapter&
 	// Initialize the RHI capabilities.
 	check(FeatureLevel == D3D_FEATURE_LEVEL_11_0 || FeatureLevel == D3D_FEATURE_LEVEL_10_0 );
 
-	if(FeatureLevel == D3D_FEATURE_LEVEL_10_0)
+	if (FeatureLevel == D3D_FEATURE_LEVEL_10_0)
 	{
 		GSupportsDepthFetchDuringDepthTest = false;
 	}
@@ -421,7 +423,9 @@ void FD3D12Device::CreateCommandContexts()
 
 		// Make available all but the first command context for parallel threads
 		if (i > 0)
+		{
 			FreeCommandContexts.Add(CommandContextArray[i]);
+		}
 	}
 
 	CommandContextArray[0]->OpenCommandList();
@@ -456,6 +460,7 @@ void FD3D12CommandContext::RHIPushEvent(const TCHAR* Name)
 {
     OwningRHI.PushGPUEvent(Name);
 }
+
 void FD3D12CommandContext::RHIPopEvent()
 {
     OwningRHI.PopGPUEvent();
@@ -497,21 +502,25 @@ void FD3D12DynamicRHI::RHIGetSupportedResolution( uint32 &Width, uint32 &Height
 		{
 			TRefCountPtr<IDXGIOutput> Output;
 			HResult = Adapter->EnumOutputs(o,Output.GetInitReference());
-			if(DXGI_ERROR_NOT_FOUND == HResult)
+			if (DXGI_ERROR_NOT_FOUND == HResult)
+			{
 				break;
-			if(FAILED(HResult))
+			}
+			if (FAILED(HResult))
+			{
 				return;
+			}
 
 			// TODO: GetDisplayModeList is a terribly SLOW call.  It can take up to a second per invocation.
 			//  We might want to work around some DXGI badness here.
 			DXGI_FORMAT Format = DXGI_FORMAT_R8G8B8A8_UNORM;
 			uint32 NumModes = 0;
 			HResult = Output->GetDisplayModeList(Format,0,&NumModes,NULL);
-			if(HResult == DXGI_ERROR_NOT_FOUND)
+			if (HResult == DXGI_ERROR_NOT_FOUND)
 			{
 				return;
 			}
-			else if(HResult == DXGI_ERROR_NOT_CURRENTLY_AVAILABLE)
+			else if (HResult == DXGI_ERROR_NOT_CURRENTLY_AVAILABLE)
 			{
 				UE_LOG(LogD3D12RHI, Fatal,
 					TEXT("This application cannot be run over a remote desktop configuration")
@@ -529,14 +538,14 @@ void FD3D12DynamicRHI::RHIGetSupportedResolution( uint32 &Width, uint32 &Height
 				CA_SUPPRESS( 6385 );
 				bool IsEqualOrBetterWidth = FMath::Abs((int32)ModeList[m].Width - (int32)Width) <= FMath::Abs((int32)BestMode.Width - (int32)Width);
 				bool IsEqualOrBetterHeight = FMath::Abs((int32)ModeList[m].Height - (int32)Height) <= FMath::Abs((int32)BestMode.Height - (int32)Height);
-				if(!InitializedMode || (IsEqualOrBetterWidth && IsEqualOrBetterHeight))
+				if (!InitializedMode || (IsEqualOrBetterWidth && IsEqualOrBetterHeight))
 				{
 					BestMode = ModeList[m];
 					InitializedMode = true;
 				}
 			}
 
-			delete[] ModeList;
+			delete [] ModeList;
 		}
 #endif // PLATFORM_XBOXONE
 	}
@@ -587,7 +596,7 @@ MSVC_PRAGMA(warning(push))
 
 uint32 FD3D12DynamicRHI::GetMaxMSAAQuality(uint32 SampleCount)
 {
-	if(SampleCount <= DX_MAX_MSAA_COUNT)
+	if (SampleCount <= DX_MAX_MSAA_COUNT)
 	{
 		// 0 has better quality (a more even distribution)
 		// higher quality levels might be useful for non box filtered AA or when using weighted samples 
@@ -767,7 +776,7 @@ void FD3D12Device::UpdateMSAASettings()
 
 void FD3D12Device::CleanupD3DDevice()
 {
-	if(GIsRHIInitialized)
+	if (GIsRHIInitialized)
 	{
 		// Ensure any pending rendering is finished
 		GetCommandListManager().SignalFrameComplete(true);
@@ -826,7 +835,7 @@ void FD3D12Device::CleanupD3DDevice()
 		// Cleanup thread resources
         for (uint32 index; (index = InterlockedDecrement(&NumThreadDynamicHeapAllocators)) != (uint32)-1;)
 		{
-            FD3D12DynamicHeapAllocator *pHeapAllocator = ThreadDynamicHeapAllocatorArray[index];
+            FD3D12DynamicHeapAllocator* pHeapAllocator = ThreadDynamicHeapAllocatorArray[index];
 			pHeapAllocator->ReleaseAllResources();
 			delete(pHeapAllocator);
 		}
diff --git a/Engine/Source/Runtime/Windows/D3D12RHI/Private/D3D12RHI.cpp b/Engine/Source/Runtime/Windows/D3D12RHI/Private/D3D12RHI.cpp
--- a/Engine/Source/Runtime/Windows/D3D12RHI/Private/D3D12RHI.cpp
+++ b/Engine/Source/Runtime/Windows/D3D12RHI/Private/D3D12RHI.cpp
@@ -971,7 +971,9 @@ void FD3D12ResourceResidencyManager::FreeMemory (uint32 BytesToFree)
 	int64 AvailableSpace = budget - int64(LocalVideoMemoryInfo.CurrentUsage);
 
 	if (AvailableSpace < BytesToFree)
+	{
 		Process (0, BytesToFree - AvailableSpace);
+	}
 }
 
 void FD3D12ResourceResidencyManager::MakeResident()
@@ -997,7 +999,9 @@ void FD3D12ResourceResidencyManager::MakeResident()
 				hresult = GetParentDevice()->GetDevice()->MakeResident (PageableResources.Num(), PageableResources.GetData());
 
 				if (hresult != E_OUTOFMEMORY)
+				{
 					break;
+				}
 			}
 		}
 
diff --git a/Engine/Source/Runtime/Windows/D3D12RHI/Private/Windows/WindowsD3D12Device.cpp b/Engine/Source/Runtime/Windows/D3D12RHI/Private/Windows/WindowsD3D12Device.cpp
--- a/Engine/Source/Runtime/Windows/D3D12RHI/Private/Windows/WindowsD3D12Device.cpp
+++ b/Engine/Source/Runtime/Windows/D3D12RHI/Private/Windows/WindowsD3D12Device.cpp
@@ -443,7 +443,7 @@ void FD3D12DynamicRHI::PerRHISetup(FD3D12Device* MainDevice)
 	// "GMaxTextureMipCount = 9" allows Intel to run the 4.8 Elemental demo, ideally this handling 
 	// is unnecessary as the upper engine should be able to look at GTotalGraphicsMemory to determine 
 	// the appropriate mip level but that doesn't appear to happen correctly on start-up.
-	if (LocalVideoMemoryInfo.Budget <= 2ll * 1024ll * 1024ll * 1024ll)
+	if (IsRHIDeviceIntel() && LocalVideoMemoryInfo.Budget <= 2ll * 1024ll * 1024ll * 1024ll)
 	{
 		GMaxTextureMipCount = FMath::Min<int32>(GMaxTextureMipCount, 9);
 	}
@@ -796,9 +796,13 @@ bool FD3D12DynamicRHI::RHIGetAvailableResolutions(FScreenResolutionArray& Resolu
 	HResult = DXGIFactory->EnumAdapters(GetRHIDevice()->GetAdapterIndex(), Adapter.GetInitReference());
 
 	if (DXGI_ERROR_NOT_FOUND == HResult)
+	{
 		return false;
+	}
 	if (FAILED(HResult))
+	{
 		return false;
+	}
 
 	// get the description of the adapter
 	DXGI_ADAPTER_DESC AdapterDesc;
@@ -813,9 +817,13 @@ bool FD3D12DynamicRHI::RHIGetAvailableResolutions(FScreenResolutionArray& Resolu
 		TRefCountPtr<IDXGIOutput> Output;
 		HResult = Adapter->EnumOutputs(CurrentOutput, Output.GetInitReference());
 		if (DXGI_ERROR_NOT_FOUND == HResult)
+		{
 			break;
+		}
 		if (FAILED(HResult))
+		{
 			return false;
+		}
 
 		// TODO: GetDisplayModeList is a terribly SLOW call.  It can take up to a second per invocation.
 		//  We might want to work around some DXGI badness here.
