UE4 - Integrate changes from DX12

diff --git a/Engine/Source/Runtime/Core/Private/Windows/WindowsPlatformMisc.cpp b/Engine/Source/Runtime/Core/Private/Windows/WindowsPlatformMisc.cpp
--- a/Engine/Source/Runtime/Core/Private/Windows/WindowsPlatformMisc.cpp
+++ b/Engine/Source/Runtime/Core/Private/Windows/WindowsPlatformMisc.cpp
@@ -1789,6 +1789,7 @@ void FWindowsPlatformMisc::LoadPreInitModules()
 	// D3D11 is not supported on WinXP, so in this case we use the OpenGL RHI
 	if(FWindowsPlatformMisc::VerifyWindowsVersion(6, 0))
 	{
+		FModuleManager::Get().LoadModule(TEXT("D3D12RHI"));
 		FModuleManager::Get().LoadModule(TEXT("D3D11RHI"));
 	}
 	FModuleManager::Get().LoadModule(TEXT("OpenGLDrv"));
diff --git a/Engine/Source/Runtime/Core/Public/Windows/WindowsCriticalSection.h b/Engine/Source/Runtime/Core/Public/Windows/WindowsCriticalSection.h
--- a/Engine/Source/Runtime/Core/Public/Windows/WindowsCriticalSection.h
+++ b/Engine/Source/Runtime/Core/Public/Windows/WindowsCriticalSection.h
@@ -47,6 +47,19 @@ public:
 		}
 	}
 
+	/**
+	* Quick test for seeing if the lock is already being used.
+	*/
+	FORCEINLINE bool TryLock()
+	{
+		if (TryEnterCriticalSection(&CriticalSection))
+		{
+			LeaveCriticalSection(&CriticalSection);
+			return true;
+		};
+		return false;
+	}
+
 	/**
 	 * Releases the lock on the critical seciton
 	 */
diff --git a/Engine/Source/Runtime/Engine/Private/Texture2D.cpp b/Engine/Source/Runtime/Engine/Private/Texture2D.cpp
--- a/Engine/Source/Runtime/Engine/Private/Texture2D.cpp
+++ b/Engine/Source/Runtime/Engine/Private/Texture2D.cpp
@@ -751,6 +751,9 @@ void UTexture2D::CalcAllowedMips( int32 MipCount, int32 NumNonStreamingMips, int
 	int32 MaxAllowedMips = FMath::Max( MipCount - LODBias, MinAllowedMips );
 	MaxAllowedMips = FMath::Min( MaxAllowedMips, GMaxTextureMipCount );
 
+	// Make sure min <= max
+	MinAllowedMips = FMath::Min(MinAllowedMips, MaxAllowedMips);
+
 	// Return results.
 	OutMinAllowedMips = MinAllowedMips;
 	OutMaxAllowedMips = MaxAllowedMips;
diff --git a/Engine/Source/Runtime/Launch/Launch.Build.cs b/Engine/Source/Runtime/Launch/Launch.Build.cs
--- a/Engine/Source/Runtime/Launch/Launch.Build.cs
+++ b/Engine/Source/Runtime/Launch/Launch.Build.cs
@@ -51,6 +51,7 @@ public class Launch : ModuleRules
 				(Target.Platform == UnrealTargetPlatform.Win64))
 			{
 				DynamicallyLoadedModuleNames.Add("D3D11RHI");
+				DynamicallyLoadedModuleNames.Add("D3D12RHI");
 				DynamicallyLoadedModuleNames.Add("XAudio2");
 			}
 			else if (Target.Platform == UnrealTargetPlatform.Mac)
diff --git a/Engine/Source/Runtime/Windows/D3D12RHI/Private/D3D12Viewport.cpp b/Engine/Source/Runtime/Windows/D3D12RHI/Private/D3D12Viewport.cpp
--- a/Engine/Source/Runtime/Windows/D3D12RHI/Private/D3D12Viewport.cpp
+++ b/Engine/Source/Runtime/Windows/D3D12RHI/Private/D3D12Viewport.cpp
@@ -306,7 +306,7 @@ bool FD3D12Viewport::PresentChecked(int32 SyncInterval)
 		// Signal the frame is complete.
 		GetParentDevice()->GetCommandListManager().SignalFrameComplete();
 
-#if UE_BUILD_DEBUG
+#if 0//UE_BUILD_DEBUG
 		UE_LOG(LogD3D12RHI, Log, TEXT("*** PRESENT ***"), this);
 		GetParentDevice()->GetOwningRHI()->DrawCount = 0;
 		GetParentDevice()->GetOwningRHI()->PresentCount++;
