UE4 - D3D12 - Remove stats dump & add back stalls test

diff --git a/Engine/Source/Runtime/Windows/D3D12RHI/Private/D3D12DirectCommandListManager.cpp b/Engine/Source/Runtime/Windows/D3D12RHI/Private/D3D12DirectCommandListManager.cpp
--- a/Engine/Source/Runtime/Windows/D3D12RHI/Private/D3D12DirectCommandListManager.cpp
+++ b/Engine/Source/Runtime/Windows/D3D12RHI/Private/D3D12DirectCommandListManager.cpp
@@ -245,7 +245,7 @@ void FD3D12CommandListManager::ExecuteCommandLists(TArray<FD3D12CommandListHandl
 	FD3D12CommandListHandle BarrierCommandList[128];
 	if (NeedsResourceBarriers)
 	{
-#if 0//UE_BUILD_DEBUG	
+#if UE_BUILD_DEBUG	
 		if (!ResourceStateCS.TryLock())
 		{
 			FD3D12DynamicRHI::GetD3DRHI()->SubmissionLockStalls++;
diff --git a/Engine/Source/Runtime/Windows/D3D12RHI/Private/D3D12Viewport.cpp b/Engine/Source/Runtime/Windows/D3D12RHI/Private/D3D12Viewport.cpp
--- a/Engine/Source/Runtime/Windows/D3D12RHI/Private/D3D12Viewport.cpp
+++ b/Engine/Source/Runtime/Windows/D3D12RHI/Private/D3D12Viewport.cpp
@@ -91,6 +91,15 @@ namespace D3D12RHI
 			ECVF_RenderThreadSafe
 			);
 
+#if UE_BUILD_DEBUG
+		int32 DumpStatsEveryNFrames = 0;
+		static FAutoConsoleVariableRef CVarDumpStatsNFrames(
+			TEXT("D3D12.DumpStatsEveryNFrames"),
+			DumpStatsEveryNFrames,
+			TEXT("Dumps D3D12 stats every N frames on Present; 0 means no information (default)."),
+			ECVF_RenderThreadSafe
+			);
+#endif
 	};
 
 	extern void D3D11TextureAllocated2D(FD3D12Texture2D& Texture);
@@ -306,16 +315,18 @@ bool FD3D12Viewport::PresentChecked(int32 SyncInterval)
 		// Signal the frame is complete.
 		GetParentDevice()->GetCommandListManager().SignalFrameComplete();
 
-#if 0//UE_BUILD_DEBUG
-		UE_LOG(LogD3D12RHI, Log, TEXT("*** PRESENT ***"), this);
-		GetParentDevice()->GetOwningRHI()->DrawCount = 0;
-		GetParentDevice()->GetOwningRHI()->PresentCount++;
-
-		// Dump stats every 240 frames
-		if (GetParentDevice()->GetOwningRHI()->PresentCount % 240 == 0)
+#if UE_BUILD_DEBUG
+		if (RHIConsoleVariables::DumpStatsEveryNFrames > 0)
 		{
-			FOutputDeviceRedirector* pOutputDevice = FOutputDeviceRedirector::Get();
-			GetParentDevice()->GetDefaultUploadHeapAllocator().DumpAllocatorStats(*pOutputDevice);
+			GetParentDevice()->GetOwningRHI()->DrawCount = 0;
+			GetParentDevice()->GetOwningRHI()->PresentCount++;
+
+			if (GetParentDevice()->GetOwningRHI()->PresentCount % RHIConsoleVariables::DumpStatsEveryNFrames == 0)
+			{
+				UE_LOG(LogD3D12RHI, Log, TEXT("*** PRESENT ***"), this);
+				FOutputDeviceRedirector* pOutputDevice = FOutputDeviceRedirector::Get();
+				GetParentDevice()->GetDefaultUploadHeapAllocator().DumpAllocatorStats(*pOutputDevice);
+			}
 		}
 #endif
 	}
