UE4 - WinXP compile fixes

diff --git a/Engine/Source/Runtime/Windows/D3D12RHI/Private/D3D12DirectCommandListManager.h b/Engine/Source/Runtime/Windows/D3D12RHI/Private/D3D12DirectCommandListManager.h
--- a/Engine/Source/Runtime/Windows/D3D12RHI/Private/D3D12DirectCommandListManager.h
+++ b/Engine/Source/Runtime/Windows/D3D12RHI/Private/D3D12DirectCommandListManager.h
@@ -99,7 +99,11 @@ private:
 
 	FD3D12Fence Fences[FT_NumTypes];
 	
+#if WINVER >= 0x0600 // Interlock...64 functions are only available from Vista onwards
 	int64									NumCommandListsAllocated;
+#else
+	int32									NumCommandListsAllocated;
+#endif
 	D3D12_COMMAND_LIST_TYPE					CommandListType;
 	ID3D12Device*							Direct3DDevice;
 	FCriticalSection						PeekReadyListCS;
diff --git a/Engine/Source/Runtime/Windows/D3D12RHI/Private/D3D12RHIPrivate.h b/Engine/Source/Runtime/Windows/D3D12RHI/Private/D3D12RHIPrivate.h
--- a/Engine/Source/Runtime/Windows/D3D12RHI/Private/D3D12RHIPrivate.h
+++ b/Engine/Source/Runtime/Windows/D3D12RHI/Private/D3D12RHIPrivate.h
@@ -142,7 +142,11 @@ private:
 
         int64 GenerateID()
         {
+#if WINVER >= 0x0600 // Interlock...64 functions are only available from Vista onwards
             static int64 ID = 0;
+#else
+			static int32 ID = 0;
+#endif
             return FPlatformAtomics::_InterlockedIncrement(&ID);
         }
 
diff --git a/Engine/Source/Runtime/Windows/D3D12RHI/Public/D3D12Util.h b/Engine/Source/Runtime/Windows/D3D12RHI/Public/D3D12Util.h
--- a/Engine/Source/Runtime/Windows/D3D12RHI/Public/D3D12Util.h
+++ b/Engine/Source/Runtime/Windows/D3D12RHI/Public/D3D12Util.h
@@ -6,6 +6,28 @@
 
 #pragma once
 
+#if WINVER == 0x0502
+// Windows XP uses Win7 sdk, and in that one winerror.h doesn't include them
+#ifndef DXGI_ERROR_INVALID_CALL
+#define DXGI_ERROR_INVALID_CALL                 MAKE_DXGI_HRESULT(1)
+#define DXGI_ERROR_NOT_FOUND                    MAKE_DXGI_HRESULT(2)
+#define DXGI_ERROR_MORE_DATA                    MAKE_DXGI_HRESULT(3)
+#define DXGI_ERROR_UNSUPPORTED                  MAKE_DXGI_HRESULT(4)
+#define DXGI_ERROR_DEVICE_REMOVED               MAKE_DXGI_HRESULT(5)
+#define DXGI_ERROR_DEVICE_HUNG                  MAKE_DXGI_HRESULT(6)
+#define DXGI_ERROR_DEVICE_RESET                 MAKE_DXGI_HRESULT(7)
+#define DXGI_ERROR_WAS_STILL_DRAWING            MAKE_DXGI_HRESULT(10)
+#define DXGI_ERROR_FRAME_STATISTICS_DISJOINT    MAKE_DXGI_HRESULT(11)
+#define DXGI_ERROR_GRAPHICS_VIDPN_SOURCE_IN_USE MAKE_DXGI_HRESULT(12)
+#define DXGI_ERROR_DRIVER_INTERNAL_ERROR        MAKE_DXGI_HRESULT(32)
+#define DXGI_ERROR_NONEXCLUSIVE                 MAKE_DXGI_HRESULT(33)
+#define DXGI_ERROR_NOT_CURRENTLY_AVAILABLE      MAKE_DXGI_HRESULT(34)
+#define DXGI_ERROR_REMOTE_CLIENT_DISCONNECTED   MAKE_DXGI_HRESULT(35)
+#define DXGI_ERROR_REMOTE_OUTOFMEMORY           MAKE_DXGI_HRESULT(36)
+#endif
+
+#endif
+
 namespace D3D12RHI
 {
 	/**
