UE4 - D3D12 - Integration from MS

diff --git a/Engine/Source/Runtime/Windows/D3D12RHI/Private/D3D12Query.cpp b/Engine/Source/Runtime/Windows/D3D12RHI/Private/D3D12Query.cpp
--- a/Engine/Source/Runtime/Windows/D3D12RHI/Private/D3D12Query.cpp
+++ b/Engine/Source/Runtime/Windows/D3D12RHI/Private/D3D12Query.cpp
@@ -362,7 +362,7 @@ void FD3D12QueryHeap::EndQueryBatchAndResolveQueryData(FD3D12CommandContext& Cmd
 
     // Increment the active element count
     ActiveAllocatedElementCount += CurrentQueryBatch.ElementCount;
-    check(ActiveAllocatedElementCount <= GetQueryHeapCount());
+	checkf(ActiveAllocatedElementCount <= GetQueryHeapCount(), TEXT("The query heap is too small. Either increase the heap count (larger resource) or decrease MAX_ACTIVE_BATCHES."));
 
     // Track the current active batches (application is using the data)
     LastBatch = GetNextBatchElement(LastBatch);
diff --git a/Engine/Source/Runtime/Windows/D3D12RHI/Private/D3D12RHIPrivate.h b/Engine/Source/Runtime/Windows/D3D12RHI/Private/D3D12RHIPrivate.h
--- a/Engine/Source/Runtime/Windows/D3D12RHI/Private/D3D12RHIPrivate.h
+++ b/Engine/Source/Runtime/Windows/D3D12RHI/Private/D3D12RHIPrivate.h
@@ -227,7 +227,7 @@ private:
 
     TArray<QueryBatch> ActiveQueryBatches;              // List of active query batches. The data for these is in use.
 
-    static const uint32 MAX_ACTIVE_BATCHES = 10;        // The max number of query batches that will be held.
+    static const uint32 MAX_ACTIVE_BATCHES = 5;         // The max number of query batches that will be held.
     uint32 LastBatch;                                   // The index of the newest batch.
 
     uint32 HeadActiveElement;                   // The oldest element that is in use (Active). The data for this element is being used.
diff --git a/Engine/Source/Runtime/Windows/D3D12RHI/Private/D3D12Texture.cpp b/Engine/Source/Runtime/Windows/D3D12RHI/Private/D3D12Texture.cpp
--- a/Engine/Source/Runtime/Windows/D3D12RHI/Private/D3D12Texture.cpp
+++ b/Engine/Source/Runtime/Windows/D3D12RHI/Private/D3D12Texture.cpp
@@ -1374,15 +1374,15 @@ FShaderResourceViewRHIRef FD3D12DynamicRHI::RHICreateShaderResourceView(FTexture
 	uint64 Offset = 0;
     const D3D12_RESOURCE_DESC& TextureDesc = Texture2D->GetResource()->GetDesc();
 
-	bool bSRGB = (Texture2D->GetFlags() & TexCreate_SRGB) != 0;
+	const bool bSRGB = (Texture2D->GetFlags() & TexCreate_SRGB) != 0;
 	const DXGI_FORMAT PlatformShaderResourceFormat = FindShaderResourceDXGIFormat(TextureDesc.Format, bSRGB);
 
 	// Create a Shader Resource View
     D3D12_SHADER_RESOURCE_VIEW_DESC SRVDesc = {};
     SRVDesc.Shader4ComponentMapping = D3D12_DEFAULT_SHADER_4_COMPONENT_MAPPING;
 	SRVDesc.ViewDimension = D3D12_SRV_DIMENSION_TEXTURE2D;
-	SRVDesc.Texture2D.MostDetailedMip = MipLevel;
 	SRVDesc.Texture2D.MipLevels = 1;
+	SRVDesc.Texture2D.MostDetailedMip = MipLevel;
 
 	SRVDesc.Format = PlatformShaderResourceFormat;
 
@@ -1393,54 +1393,50 @@ FShaderResourceViewRHIRef FD3D12DynamicRHI::RHICreateShaderResourceView(FTexture
 
 FShaderResourceViewRHIRef FD3D12DynamicRHI::RHICreateShaderResourceView(FTexture2DArrayRHIParamRef Texture2DRHI, uint8 MipLevel)
 {
-	check(0);
-#if 0
-	FD3D12Texture*  Texture2D = FD3D12DynamicRHI::ResourceCast(Texture2DRHI);
+	FD3D12Texture2DArray*  Texture2DArray = FD3D12DynamicRHI::ResourceCast(Texture2DRHI);
 
 	uint64 Offset = 0;
-	const D3D12_RESOURCE_DESC& TextureDesc = Texture2D->GetResource()->GetDesc();
+	const D3D12_RESOURCE_DESC& TextureDesc = Texture2DArray->GetResource()->GetDesc();
+
+	const bool bSRGB = (Texture2DArray->GetFlags() & TexCreate_SRGB) != 0;
+	const DXGI_FORMAT PlatformShaderResourceFormat = FindShaderResourceDXGIFormat(TextureDesc.Format, bSRGB);
 
 	// Create a Shader Resource View
 	D3D12_SHADER_RESOURCE_VIEW_DESC SRVDesc = {};
 	SRVDesc.Shader4ComponentMapping = D3D12_DEFAULT_SHADER_4_COMPONENT_MAPPING;
-	SRVDesc.ViewDimension = D3D12_SRV_DIMENSION_TEXTURE2D;
-	SRVDesc.Texture2D.MostDetailedMip = MipLevel;
-	SRVDesc.Texture2D.MipLevels = 1;
+	SRVDesc.ViewDimension = D3D12_SRV_DIMENSION_TEXTURE2DARRAY;
+	SRVDesc.Texture2DArray.ArraySize = TextureDesc.DepthOrArraySize;
+	SRVDesc.Texture2DArray.FirstArraySlice = 0;
+	SRVDesc.Texture2DArray.MipLevels = 1;
+	SRVDesc.Texture2DArray.MostDetailedMip = MipLevel;
 
-	SRVDesc.Format = TextureDesc.Format;
+	SRVDesc.Format = PlatformShaderResourceFormat;
 
-	SRVDesc.Texture2D.PlaneSlice = GetPlaneSliceFromViewFormat(TextureDesc.Format, SRVDesc.Format);
+	SRVDesc.Texture2DArray.PlaneSlice = GetPlaneSliceFromViewFormat(TextureDesc.Format, SRVDesc.Format);
 
-	return new FD3D12ShaderResourceView(GetRHIDevice(), &SRVDesc, Texture2D->ResourceLocation);
-#else
-	return nullptr;
-#endif
+	return new FD3D12ShaderResourceView(GetRHIDevice(), &SRVDesc, Texture2DArray->ResourceLocation);
 }
 
 FShaderResourceViewRHIRef FD3D12DynamicRHI::RHICreateShaderResourceView(FTextureCubeRHIParamRef TextureCubeRHI, uint8 MipLevel)
 {
-	check(0);
-#if 0
-	FD3D12Texture*  Texture2D = FD3D12DynamicRHI::ResourceCast(Texture2DRHI);
+	FD3D12TextureCube*  TextureCube = FD3D12DynamicRHI::ResourceCast(TextureCubeRHI);
 
 	uint64 Offset = 0;
-	const D3D12_RESOURCE_DESC& TextureDesc = Texture2D->GetResource()->GetDesc();
+	const D3D12_RESOURCE_DESC& TextureDesc = TextureCube->GetResource()->GetDesc();
+
+	const bool bSRGB = (TextureCube->GetFlags() & TexCreate_SRGB) != 0;
+	const DXGI_FORMAT PlatformShaderResourceFormat = FindShaderResourceDXGIFormat(TextureDesc.Format, bSRGB);
 
 	// Create a Shader Resource View
 	D3D12_SHADER_RESOURCE_VIEW_DESC SRVDesc = {};
 	SRVDesc.Shader4ComponentMapping = D3D12_DEFAULT_SHADER_4_COMPONENT_MAPPING;
-	SRVDesc.ViewDimension = D3D12_SRV_DIMENSION_TEXTURE2D;
-	SRVDesc.Texture2D.MostDetailedMip = MipLevel;
-	SRVDesc.Texture2D.MipLevels = 1;
-
-	SRVDesc.Format = TextureDesc.Format;
+	SRVDesc.ViewDimension = D3D12_SRV_DIMENSION_TEXTURECUBE;
+	SRVDesc.TextureCube.MipLevels = 1;
+	SRVDesc.TextureCube.MostDetailedMip = MipLevel;
 
-	SRVDesc.Texture2D.PlaneSlice = GetPlaneSliceFromViewFormat(TextureDesc.Format, SRVDesc.Format);
+	SRVDesc.Format = PlatformShaderResourceFormat;
 
-	return new FD3D12ShaderResourceView(GetRHIDevice(), &SRVDesc, Texture2D->ResourceLocation);
-#else
-	return nullptr;
-#endif
+	return new FD3D12ShaderResourceView(GetRHIDevice(), &SRVDesc, TextureCube->ResourceLocation);
 }
 
 FShaderResourceViewRHIRef FD3D12DynamicRHI::RHICreateShaderResourceView(FTexture2DRHIParamRef Texture2DRHI, uint8 MipLevel, uint8 NumMipLevels, uint8 Format)
@@ -1451,7 +1447,7 @@ FShaderResourceViewRHIRef FD3D12DynamicRHI::RHICreateShaderResourceView(FTexture
 
 	const DXGI_FORMAT PlatformResourceFormat = FD3D12DynamicRHI::GetPlatformTextureResourceFormat((DXGI_FORMAT)GPixelFormats[Format].PlatformFormat, Texture2D->GetFlags());
 
-	bool bSRGB = (Texture2D->GetFlags() & TexCreate_SRGB) != 0;
+	const bool bSRGB = (Texture2D->GetFlags() & TexCreate_SRGB) != 0;
 	const DXGI_FORMAT PlatformShaderResourceFormat = FindShaderResourceDXGIFormat(PlatformResourceFormat, bSRGB);
 
 	// Create a Shader Resource View
@@ -1479,26 +1475,23 @@ FShaderResourceViewRHIRef FD3D12DynamicRHI::RHICreateShaderResourceView(FTexture
 FShaderResourceViewRHIRef FD3D12DynamicRHI::RHICreateShaderResourceView(FTexture3DRHIParamRef Texture3DRHI, uint8 MipLevel)
 {
 	FD3D12Texture3D*  Texture3D = FD3D12DynamicRHI::ResourceCast(Texture3DRHI);
-check(0);
-#if 0
+
 	uint64 Offset = 0;
-	const D3D12_RESOURCE_DESC& TextureDesc = Texture2D->GetResource()->GetDesc();
+	const D3D12_RESOURCE_DESC& TextureDesc = Texture3D->GetResource()->GetDesc();
+
+	const bool bSRGB = (Texture3D->GetFlags() & TexCreate_SRGB) != 0;
+	const DXGI_FORMAT PlatformShaderResourceFormat = FindShaderResourceDXGIFormat(TextureDesc.Format, bSRGB);
 
 	// Create a Shader Resource View
 	D3D12_SHADER_RESOURCE_VIEW_DESC SRVDesc = {};
 	SRVDesc.Shader4ComponentMapping = D3D12_DEFAULT_SHADER_4_COMPONENT_MAPPING;
 	SRVDesc.ViewDimension = D3D12_SRV_DIMENSION_TEXTURE3D;
-	SRVDesc.Texture3D.MostDetailedMip = MipLevel;
 	SRVDesc.Texture3D.MipLevels = 1;
+	SRVDesc.Texture3D.MostDetailedMip = MipLevel;
 
-	SRVDesc.Format = TextureDesc.Format;
-
-	//SRVDesc.Texture3D.PlaneSlice = GetPlaneSliceFromViewFormat(TextureDesc.Format, SRVDesc.Format);
+	SRVDesc.Format = PlatformShaderResourceFormat;
 
 	return new FD3D12ShaderResourceView(GetRHIDevice(), &SRVDesc, Texture3D->ResourceLocation);
-#else
-	return nullptr;
-#endif
 }
 
 /** Generates mip maps for the surface. */
