UE4 - D3D12 - Remove some typedefs

diff --git a/Engine/Source/Runtime/Windows/D3D12RHI/Private/D3D12DirectCommandListManager.cpp b/Engine/Source/Runtime/Windows/D3D12RHI/Private/D3D12DirectCommandListManager.cpp
--- a/Engine/Source/Runtime/Windows/D3D12RHI/Private/D3D12DirectCommandListManager.cpp
+++ b/Engine/Source/Runtime/Windows/D3D12RHI/Private/D3D12DirectCommandListManager.cpp
@@ -34,7 +34,7 @@ void FD3D12Fence::CreateFence(ID3D12Device* pDirect3DDevice, uint64 InitialValue
 	CurrentFence = LastCompletedFence + 1;
 }
 
-uint64 FD3D12Fence::Signal(FD3D12CommandQueue* pCommandQueue)
+uint64 FD3D12Fence::Signal(ID3D12CommandQueue* pCommandQueue)
 {
 	check(pCommandQueue != nullptr);
 
diff --git a/Engine/Source/Runtime/Windows/D3D12RHI/Private/D3D12DirectCommandListManager.h b/Engine/Source/Runtime/Windows/D3D12RHI/Private/D3D12DirectCommandListManager.h
--- a/Engine/Source/Runtime/Windows/D3D12RHI/Private/D3D12DirectCommandListManager.h
+++ b/Engine/Source/Runtime/Windows/D3D12RHI/Private/D3D12DirectCommandListManager.h
@@ -27,7 +27,7 @@ public:
 	~FD3D12Fence();
 
 	void CreateFence(ID3D12Device* pDirect3DDevice, uint64 InitialValue = 0);
-	uint64 Signal(FD3D12CommandQueue* pCommandQueue);
+	uint64 Signal(ID3D12CommandQueue* pCommandQueue);
 	bool IsFenceFinished(uint64 FenceValue);
 	void WaitForFence(uint64 FenceValue);
 
@@ -75,7 +75,7 @@ public:
 		return D3DCommandQueue->GetTimestampFrequency(Frequency);
 	}
 
-    inline FD3D12CommandQueue* GetD3DCommandQueue()
+    inline ID3D12CommandQueue* GetD3DCommandQueue()
     {
         return D3DCommandQueue.GetReference();
     }
@@ -93,7 +93,7 @@ private:
 	uint64 ExecuteAndIncrementFence(ID3D12CommandList* pD3DCommandLists[], uint64 NumCommandLists, FD3D12Fence &Fence);
 	FD3D12CommandListHandle CreateCommandListHandle();
 
-	TRefCountPtr<FD3D12CommandQueue>		D3DCommandQueue;
+	TRefCountPtr<ID3D12CommandQueue>		D3DCommandQueue;
 
 	FThreadsafeQueue<FD3D12CommandListHandle> ReadyLists;
 
diff --git a/Engine/Source/Runtime/Windows/D3D12RHI/Private/Windows/D3D12RHIBasePrivate.h b/Engine/Source/Runtime/Windows/D3D12RHI/Private/Windows/D3D12RHIBasePrivate.h
--- a/Engine/Source/Runtime/Windows/D3D12RHI/Private/Windows/D3D12RHIBasePrivate.h
+++ b/Engine/Source/Runtime/Windows/D3D12RHI/Private/Windows/D3D12RHIBasePrivate.h
@@ -35,10 +35,6 @@
 #pragma pack(pop)
 #pragma warning(pop)
 
-typedef ID3D12GraphicsCommandList FD3D12CommandList;
-typedef ID3D12CommandQueue FD3D12CommandQueue;
-typedef ID3D12CommandAllocator FD3D12CommandAllocator;
-
 
 class FD3D12Device;
 
diff --git a/Engine/Source/Runtime/Windows/D3D12RHI/Public/D3D12Resources.h b/Engine/Source/Runtime/Windows/D3D12RHI/Public/D3D12Resources.h
--- a/Engine/Source/Runtime/Windows/D3D12RHI/Public/D3D12Resources.h
+++ b/Engine/Source/Runtime/Windows/D3D12RHI/Public/D3D12Resources.h
@@ -95,8 +95,8 @@ private:
 		mutable FThreadSafeCounter				NumRefs;
 		FD3D12CommandListManager*               CommandListManager;
 		FD3D12CommandContext*					CurrentOwningContext;
-		TRefCountPtr<FD3D12CommandList>			CommandList;		// Raw D3D command list pointer
-		TRefCountPtr<FD3D12CommandAllocator>	CommandAllocator;	// Command allocator paired with the command list
+		TRefCountPtr<ID3D12GraphicsCommandList>	CommandList;		// Raw D3D command list pointer
+		TRefCountPtr<ID3D12CommandAllocator>	CommandAllocator;	// Command allocator paired with the command list
 		uint64									FenceIndex;			// -1 when active; positive when submitted for execution
 		uint64									CurrentGeneration;
 		bool                                    IsClosed;
@@ -216,7 +216,7 @@ public:
 		return lhs != rhs.CommandListData;
 	}
 
-	FD3D12CommandList* operator->() const
+	ID3D12GraphicsCommandList* operator->() const
 	{
 		check (CommandListData && !CommandListData->IsClosed);
 
@@ -247,7 +247,7 @@ public:
 		CommandListData->Reset();
 	}
 
-	FD3D12CommandList* CommandList() const
+	ID3D12GraphicsCommandList* CommandList() const
 	{
 		check (CommandListData);
 		return CommandListData->CommandList.GetReference();
