UE4 - D3D12RHI - Fixes for Kite

diff --git a/Engine/Source/Runtime/Windows/D3D12RHI/Private/D3D12StateCache.cpp b/Engine/Source/Runtime/Windows/D3D12RHI/Private/D3D12StateCache.cpp
--- a/Engine/Source/Runtime/Windows/D3D12RHI/Private/D3D12StateCache.cpp
+++ b/Engine/Source/Runtime/Windows/D3D12RHI/Private/D3D12StateCache.cpp
@@ -2778,7 +2778,6 @@ void FD3D12DescriptorCache::SetRenderTargets(FD3D12RenderTargetView **RenderTarg
 		{
 			// Depth is also bound as an SRV. Need to add some SRV state to the depth plane, so transition the depth and stencil separately.
 			const D3D12_DEPTH_STENCIL_VIEW_DESC &desc = DepthStencilTarget->GetDesc();
-			check(desc.Flags & D3D12_DSV_FLAG_READ_ONLY_DEPTH);
 			const D3D12_RESOURCE_STATES depthState = D3D12_RESOURCE_STATE_DEPTH_READ | D3D12_RESOURCE_STATE_PIXEL_SHADER_RESOURCE | D3D12_RESOURCE_STATE_NON_PIXEL_SHADER_RESOURCE;
 			FD3D12DynamicRHI::TransitionResource(CommandList, DepthStencilTarget->GetResource(), depthState, DepthStencilTarget->GetDepthOnlyViewSubresourceSubset());
 
diff --git a/Engine/Source/Runtime/Windows/D3D12RHI/Private/D3D12StateCachePrivate.h b/Engine/Source/Runtime/Windows/D3D12RHI/Private/D3D12StateCachePrivate.h
--- a/Engine/Source/Runtime/Windows/D3D12RHI/Private/D3D12StateCachePrivate.h
+++ b/Engine/Source/Runtime/Windows/D3D12RHI/Private/D3D12StateCachePrivate.h
@@ -22,7 +22,7 @@
 #define D3D11_STATE_CACHE_DEBUG 0
 #endif
 
-#define MAX_SRVS 16
+#define MAX_SRVS 22
 #define MAX_CBS 8
 
 // Uncomment only for debugging of the descriptor heap management; this is very noisy
@@ -967,7 +967,6 @@ protected:
 
 			if (FD3D12DynamicRHI::ResourceViewsIntersect(PipelineState.Graphics.CurrentDepthStencilTarget, SRV))
 			{
-				check(PipelineState.Graphics.CurrentDepthStencilTarget->GetDesc().Flags & D3D12_DSV_FLAG_READ_ONLY_DEPTH);
 				if (!PipelineState.Common.CurrentShaderResourceViewsIntersectWithDepthRT[ShaderFrequency][ResourceIndex])
 				{
 					PipelineState.Common.CurrentShaderResourceViewsIntersectWithDepthRT[ShaderFrequency][ResourceIndex] = true;
diff --git a/Engine/Source/Runtime/Windows/D3D12RHI/Public/D3D12Resources.h b/Engine/Source/Runtime/Windows/D3D12RHI/Public/D3D12Resources.h
--- a/Engine/Source/Runtime/Windows/D3D12RHI/Public/D3D12Resources.h
+++ b/Engine/Source/Runtime/Windows/D3D12RHI/Public/D3D12Resources.h
@@ -2816,7 +2816,7 @@ public:
 class FD3D12DynamicHeapAllocator : public FD3D12ResourceAllocator
 {
 	static const uint32 BucketShift = 6;
-	static const uint32 NumBuckets = 21; // bucket resource sizes range from 64 to 2^24 
+	static const uint32 NumBuckets = 22; // bucket resource sizes range from 64 to 2^28 
 	FThreadsafeQueue<TRefCountPtr<FD3D12ResourceBlockInfo>> AvailableBlocks[NumBuckets];
 	
 	D3D12_HEAP_TYPE HeapType;
