UE4 - D3D12 - Integrate another fix for Kite

diff --git a/Engine/Source/Runtime/Windows/D3D12RHI/Private/D3D12Texture.cpp b/Engine/Source/Runtime/Windows/D3D12RHI/Private/D3D12Texture.cpp
--- a/Engine/Source/Runtime/Windows/D3D12RHI/Private/D3D12Texture.cpp
+++ b/Engine/Source/Runtime/Windows/D3D12RHI/Private/D3D12Texture.cpp
@@ -1783,7 +1783,7 @@ void TD3D12Texture2D<RHIResourceType>::Unlock(uint32 MipIndex, uint32 ArrayIndex
 			FD3D12CommandListHandle& hCommandList = GetParentDevice()->GetDefaultCommandContext().CommandListHandle;
 
 			{
-				FScopeResourceBarrier ScopeResourceBarrierDest(hCommandList, GetResource(), GetResource()->GetDefaultResourceState(), D3D12_RESOURCE_STATE_COPY_DEST, DestCopyLocation.SubresourceIndex);
+				FConditionalScopeResourceBarrier ScopeResourceBarrierDest(hCommandList, GetResource(), D3D12_RESOURCE_STATE_COPY_DEST, DestCopyLocation.SubresourceIndex);
 				// Don't need to transition upload heaps
 
 				GetParentDevice()->GetDefaultCommandContext().numCopies++;
