UE4 - D3D12 - Fix for DSV transition

diff --git a/Engine/Source/Runtime/Windows/D3D12RHI/Private/D3D12StateCache.cpp b/Engine/Source/Runtime/Windows/D3D12RHI/Private/D3D12StateCache.cpp
--- a/Engine/Source/Runtime/Windows/D3D12RHI/Private/D3D12StateCache.cpp
+++ b/Engine/Source/Runtime/Windows/D3D12RHI/Private/D3D12StateCache.cpp
@@ -2778,6 +2778,7 @@ void FD3D12DescriptorCache::SetRenderTargets(FD3D12RenderTargetView **RenderTarg
 		{
 			// Depth is also bound as an SRV. Need to add some SRV state to the depth plane, so transition the depth and stencil separately.
 			const D3D12_DEPTH_STENCIL_VIEW_DESC &desc = DepthStencilTarget->GetDesc();
+			check(desc.Flags & D3D12_DSV_FLAG_READ_ONLY_DEPTH);
 			const D3D12_RESOURCE_STATES depthState = D3D12_RESOURCE_STATE_DEPTH_READ | D3D12_RESOURCE_STATE_PIXEL_SHADER_RESOURCE | D3D12_RESOURCE_STATE_NON_PIXEL_SHADER_RESOURCE;
 			FD3D12DynamicRHI::TransitionResource(CommandList, DepthStencilTarget->GetResource(), depthState, DepthStencilTarget->GetDepthOnlyViewSubresourceSubset());
 
diff --git a/Engine/Source/Runtime/Windows/D3D12RHI/Private/D3D12StateCachePrivate.h b/Engine/Source/Runtime/Windows/D3D12RHI/Private/D3D12StateCachePrivate.h
--- a/Engine/Source/Runtime/Windows/D3D12RHI/Private/D3D12StateCachePrivate.h
+++ b/Engine/Source/Runtime/Windows/D3D12RHI/Private/D3D12StateCachePrivate.h
@@ -967,10 +967,28 @@ protected:
 
 			if (FD3D12DynamicRHI::ResourceViewsIntersect(PipelineState.Graphics.CurrentDepthStencilTarget, SRV))
 			{
-				if (!PipelineState.Common.CurrentShaderResourceViewsIntersectWithDepthRT[ShaderFrequency][ResourceIndex])
+				const D3D12_DEPTH_STENCIL_VIEW_DESC &dsvDesc = PipelineState.Graphics.CurrentDepthStencilTarget->GetDesc();
+				const bool bReadOnlyDepth = (dsvDesc.Flags & D3D12_DSV_FLAG_READ_ONLY_DEPTH) != 0;
+				if (bReadOnlyDepth)
 				{
-					PipelineState.Common.CurrentShaderResourceViewsIntersectWithDepthRT[ShaderFrequency][ResourceIndex] = true;
-					PipelineState.Common.ShaderResourceViewsIntersectWithDepthCount++;
+					// If the DSV has the read only depth then we can leave the depth stencil bound. This should be safe,
+					// as only the stencil bits should have artifacts and the shader reading the DS shouldn't care about it.
+					if (!PipelineState.Common.CurrentShaderResourceViewsIntersectWithDepthRT[ShaderFrequency][ResourceIndex])
+					{
+						PipelineState.Common.CurrentShaderResourceViewsIntersectWithDepthRT[ShaderFrequency][ResourceIndex] = true;
+						PipelineState.Common.ShaderResourceViewsIntersectWithDepthCount++;
+					}
+				}
+				else
+				{
+					// Unbind the DSV because it's being used for depth write
+					check(!bReadOnlyDepth);
+					PipelineState.Graphics.CurrentDepthStencilTarget = nullptr;
+					if (PipelineState.Common.CurrentShaderResourceViewsIntersectWithDepthRT[ShaderFrequency][ResourceIndex])
+					{
+						PipelineState.Common.CurrentShaderResourceViewsIntersectWithDepthRT[ShaderFrequency][ResourceIndex] = false;
+						PipelineState.Common.ShaderResourceViewsIntersectWithDepthCount--;
+					}
 				}
 			}
 			else
