VS2015 Typecast for mismatch in check

diff --git a/Engine/Source/Runtime/Windows/D3D12RHI/Private/D3D12StateCache.cpp b/Engine/Source/Runtime/Windows/D3D12RHI/Private/D3D12StateCache.cpp
--- a/Engine/Source/Runtime/Windows/D3D12RHI/Private/D3D12StateCache.cpp
+++ b/Engine/Source/Runtime/Windows/D3D12RHI/Private/D3D12StateCache.cpp
@@ -628,7 +628,7 @@ void FD3D12DefaultBufferPool::ReleaseAllResources()
 // Grab a buffer from the available buffers or create a new buffer if none are available
 HRESULT FD3D12DefaultBufferAllocator::AllocDefaultResource(const D3D12_RESOURCE_DESC& Desc, D3D12_SUBRESOURCE_DATA *pInitialData, FD3D12ResourceLocation *ResourceLocation, uint32 Alignment)
 {
-	check (Desc.Flags < MAX_DEFAULT_POOLS);
+	check ((uint32)Desc.Flags < MAX_DEFAULT_POOLS);
 
 	return DefaultBufferPools[Desc.Flags]->AllocDefaultResource(Desc, pInitialData, ResourceLocation, Alignment);
 }
