UE-19864 Fixing hitching/blocking of main thread when XAudio2 calls DestroyVoice on an XAudio2 source object.

diff --git a/Engine/Source/Runtime/Windows/XAudio2/Private/XAudio2Source.cpp b/Engine/Source/Runtime/Windows/XAudio2/Private/XAudio2Source.cpp
--- a/Engine/Source/Runtime/Windows/XAudio2/Private/XAudio2Source.cpp
+++ b/Engine/Source/Runtime/Windows/XAudio2/Private/XAudio2Source.cpp
@@ -17,6 +17,39 @@
 #include "XAudio2Support.h"
 #include "IAudioExtensionPlugin.h"
 
+/*------------------------------------------------------------------------------------
+	Simple async task to destroy an xaudio2 source voice without blocking the main thread.
+------------------------------------------------------------------------------------*/
+
+class FAsyncXAudio2SourceDestroyer : public FNonAbandonableTask
+{
+protected:
+	IXAudio2SourceVoice* SourceVoice;
+
+public:
+	FAsyncXAudio2SourceDestroyer(IXAudio2SourceVoice* InSourceVoice)
+		: SourceVoice(InSourceVoice)
+	{
+	}
+
+	~FAsyncXAudio2SourceDestroyer()
+	{
+		check(SourceVoice == nullptr);
+	}
+
+	void DoWork()
+	{
+		check(SourceVoice);
+		SourceVoice->DestroyVoice();
+		SourceVoice = nullptr;
+	}
+
+	FORCEINLINE TStatId GetStatId() const
+	{
+		RETURN_QUICK_DECLARE_CYCLE_STAT(FAsyncXAudio2SourceDestroyer, STATGROUP_ThreadPoolAsyncTasks);
+	}
+};
+
 /*------------------------------------------------------------------------------------
 	For muting user soundtracks during cinematics
 ------------------------------------------------------------------------------------*/
@@ -92,7 +125,9 @@ void FXAudio2SoundSource::FreeResources( void )
 	// Release voice.
 	if( Source )
 	{
-		Source->DestroyVoice();
+		// Because XAudio2's DestroyVoice source is blocking and can be slow on some processors (e.g. AMD), we're creating
+		// a task that destroys the voice on a separate thread to avoid blocking or hitching.
+		(new FAutoDeleteAsyncTask<FAsyncXAudio2SourceDestroyer>(Source))->StartBackgroundTask();
 		Source = nullptr;
 	}
 
