UE4 - D3D12 - Integrate SRV renaming fix and backbuffer format

diff --git a/Engine/Source/Runtime/Windows/D3D12RHI/Private/D3D12Viewport.cpp b/Engine/Source/Runtime/Windows/D3D12RHI/Private/D3D12Viewport.cpp
--- a/Engine/Source/Runtime/Windows/D3D12RHI/Private/D3D12Viewport.cpp
+++ b/Engine/Source/Runtime/Windows/D3D12RHI/Private/D3D12Viewport.cpp
@@ -121,7 +121,7 @@ FD3D12Texture2D* GetSwapChainSurface(FD3D12Device* Parent, EPixelFormat PixelFor
 	// create the render target view
 	TRefCountPtr<FD3D12RenderTargetView> BackBufferRenderTargetView;
 	D3D12_RENDER_TARGET_VIEW_DESC RTVDesc = {};
-	RTVDesc.Format = DXGI_FORMAT_UNKNOWN;
+	RTVDesc.Format = BackBufferDesc.Format;
 	RTVDesc.ViewDimension = D3D12_RTV_DIMENSION_TEXTURE2D;
 	RTVDesc.Texture2D.MipSlice = 0;
 	BackBufferRenderTargetView = new FD3D12RenderTargetView(Parent, &RTVDesc, BackBufferWrappedResource);
@@ -132,7 +132,7 @@ FD3D12Texture2D* GetSwapChainSurface(FD3D12Device* Parent, EPixelFormat PixelFor
 	// create a shader resource view to allow using the backbuffer as a texture
 	D3D12_SHADER_RESOURCE_VIEW_DESC SRVDesc = {};
 	SRVDesc.Shader4ComponentMapping = D3D12_DEFAULT_SHADER_4_COMPONENT_MAPPING;
-	SRVDesc.Format = DXGI_FORMAT_UNKNOWN;
+	SRVDesc.Format = BackBufferDesc.Format;
 	SRVDesc.ViewDimension = D3D12_SRV_DIMENSION_TEXTURE2D;
 	SRVDesc.Texture2D.MostDetailedMip = 0;
 	SRVDesc.Texture2D.MipLevels = 1;
diff --git a/Engine/Source/Runtime/Windows/D3D12RHI/Public/D3D12Resources.h b/Engine/Source/Runtime/Windows/D3D12RHI/Public/D3D12Resources.h
--- a/Engine/Source/Runtime/Windows/D3D12RHI/Public/D3D12Resources.h
+++ b/Engine/Source/Runtime/Windows/D3D12RHI/Public/D3D12Resources.h
@@ -1995,9 +1995,6 @@ public:
         // If no descriptor and heap index were provided...
         if(InDescriptor.ptr == 0)
         {
-            // Allocate a new heap slot
-            AllocateHeapSlot();
-
 			// Create a new view
             CreateView();
         }
diff --git a/Engine/Source/Runtime/Windows/D3D12RHI/Public/D3D12Util.h b/Engine/Source/Runtime/Windows/D3D12RHI/Public/D3D12Util.h
--- a/Engine/Source/Runtime/Windows/D3D12RHI/Public/D3D12Util.h
+++ b/Engine/Source/Runtime/Windows/D3D12RHI/Public/D3D12Util.h
@@ -377,6 +377,7 @@ private: // Methods
 		TRefCountPtr<ID3D12DescriptorHeap> Heap;
 		VERIFYD3D11RESULT( m_pDevice->CreateDescriptorHeap(&m_Desc, IID_PPV_ARGS(Heap.GetInitReference())) );
 		HeapOffset HeapBase = Heap->GetCPUDescriptorHandleForHeapStart();
+		check(HeapBase.ptr != 0);
 
 		// Allocate and initialize a single new entry in the map
 		m_Heaps.SetNum(m_Heaps.Num() + 1); 
