UE4 - D3D12 - Fix for infiltrator flickering colors

diff --git a/Engine/Source/Runtime/Windows/D3D12RHI/Private/D3D12Commands.cpp b/Engine/Source/Runtime/Windows/D3D12RHI/Private/D3D12Commands.cpp
--- a/Engine/Source/Runtime/Windows/D3D12RHI/Private/D3D12Commands.cpp
+++ b/Engine/Source/Runtime/Windows/D3D12RHI/Private/D3D12Commands.cpp
@@ -1105,27 +1105,30 @@ void FD3D12CommandContext::RHISetRenderTargetsAndClear(const FRHISetRenderTarget
 		&RenderTargetsInfo.DepthStencilRenderTarget,
 		0,
 		nullptr);
-	FLinearColor ClearColors[MaxSimultaneousRenderTargets];
-	float DepthClear = 0.0;
-	uint32 StencilClear = 0;
-
-	if (RenderTargetsInfo.bClearColor)
+	if (RenderTargetsInfo.bClearColor || RenderTargetsInfo.bClearStencil || RenderTargetsInfo.bClearDepth)
 	{
-		for (int32 i = 0; i < RenderTargetsInfo.NumColorRenderTargets; ++i)
+		FLinearColor ClearColors[MaxSimultaneousRenderTargets];
+		float DepthClear = 0.0;
+		uint32 StencilClear = 0;
+
+		if (RenderTargetsInfo.bClearColor)
 		{
-			const FClearValueBinding& ClearValue = RenderTargetsInfo.ColorRenderTarget[i].Texture->GetClearBinding();
-			checkf(ClearValue.ColorBinding == EClearBinding::EColorBound, TEXT("Texture: %s does not have a color bound for fast clears"), *RenderTargetsInfo.ColorRenderTarget[i].Texture->GetName().GetPlainNameString());
-			ClearColors[i] = ClearValue.GetClearColor();
+			for (int32 i = 0; i < RenderTargetsInfo.NumColorRenderTargets; ++i)
+			{
+				const FClearValueBinding& ClearValue = RenderTargetsInfo.ColorRenderTarget[i].Texture->GetClearBinding();
+				checkf(ClearValue.ColorBinding == EClearBinding::EColorBound, TEXT("Texture: %s does not have a color bound for fast clears"), *RenderTargetsInfo.ColorRenderTarget[i].Texture->GetName().GetPlainNameString());
+				ClearColors[i] = ClearValue.GetClearColor();
+			}
+		}
+		if (RenderTargetsInfo.bClearDepth || RenderTargetsInfo.bClearStencil)
+		{
+			const FClearValueBinding& ClearValue = RenderTargetsInfo.DepthStencilRenderTarget.Texture->GetClearBinding();
+			checkf(ClearValue.ColorBinding == EClearBinding::EDepthStencilBound, TEXT("Texture: %s does not have a DS value bound for fast clears"), *RenderTargetsInfo.DepthStencilRenderTarget.Texture->GetName().GetPlainNameString());
+			ClearValue.GetDepthStencil(DepthClear, StencilClear);
 		}
-	}
-	if (RenderTargetsInfo.bClearDepth || RenderTargetsInfo.bClearStencil)
-	{
-		const FClearValueBinding& ClearValue = RenderTargetsInfo.DepthStencilRenderTarget.Texture->GetClearBinding();
-		checkf(ClearValue.ColorBinding == EClearBinding::EDepthStencilBound, TEXT("Texture: %s does not have a DS value bound for fast clears"), *RenderTargetsInfo.DepthStencilRenderTarget.Texture->GetName().GetPlainNameString());
-		ClearValue.GetDepthStencil(DepthClear, StencilClear);
-	}
 
-	this->RHIClearMRTImpl(RenderTargetsInfo.bClearColor, RenderTargetsInfo.NumColorRenderTargets, ClearColors, RenderTargetsInfo.bClearDepth, DepthClear, RenderTargetsInfo.bClearStencil, StencilClear, FIntRect(), false);
+		this->RHIClearMRTImpl(RenderTargetsInfo.bClearColor, RenderTargetsInfo.NumColorRenderTargets, ClearColors, RenderTargetsInfo.bClearDepth, DepthClear, RenderTargetsInfo.bClearStencil, StencilClear, FIntRect(), false);
+	}
 }
 
 // Occlusion/Timer queries.
@@ -1756,6 +1759,11 @@ void FD3D12CommandContext::RHIClearMRT(bool bClearColor, int32 NumClearColors, c
 void FD3D12CommandContext::RHIClearMRTImpl(bool bClearColor, int32 NumClearColors, const FLinearColor* ClearColorArray, bool bClearDepth, float Depth, bool bClearStencil, uint32 Stencil, FIntRect ExcludeRect, bool bForceShaderClear)
 {
 	SCOPE_CYCLE_COUNTER(STAT_D3D12ClearMRT);
+
+	//don't force shaders clears for the moment.  There are bugs with the state cache/restore behavior.
+	//will either fix this soon, or move clear out of the RHI entirely.
+	bForceShaderClear = false;
+
 	// Helper struct to record and restore device states RHIClearMRT modifies.
 	class FDeviceStateHelper
 	{
