Stereo Spatialization Feature

- PC/Xaudio2 implementation

- Added new parameter added to sound attenuation object that specifies stereo spread for 3d spatialized assets that use the attenuation setting

- High-level implementation description: split the stereo file channels into 2 mono-source 3d calculations in XAudio2 (though the stereo voice still counts as a single voice).

- Currently HRTF spatialization only works for mono-sources.

- The L/R channels are oriented perpendicular to the listener and spread from each other by the spread parameter.

- Added the ability to visualize 3d audio using
   -- the new audio command is "Audio3dVisualize", we can change the name to something else if there's a different name designers would like to use
   -- 3d visualization will show the name of the wave instance, a white crosshair where the emitter position is, and a red/green crosshair for the left/right channels. This way sound designers can
      visualize the spread parameter. The position used to display the crosshairs is the exact position used by the audio engine to spatialize the channels.
   -- We can add more data to this visualization later(can display current volume, color based on occlusion, etc)
   -- If the sound is mono-spatialized, it will only display a white crosshair (and the name of the wave instance).
   -- It only shows actively playing voices that have audio components (to render into the component's world) so can be used as a means to determine playback activity of audio components.
   -- Only implemented for PC at the moment, though this feature should be ported to our other platforms

diff --git a/Engine/Source/Editor/DetailCustomizations/Private/AttenuationSettingsCustomizations.cpp b/Engine/Source/Editor/DetailCustomizations/Private/AttenuationSettingsCustomizations.cpp
--- a/Engine/Source/Editor/DetailCustomizations/Private/AttenuationSettingsCustomizations.cpp
+++ b/Engine/Source/Editor/DetailCustomizations/Private/AttenuationSettingsCustomizations.cpp
@@ -187,11 +187,12 @@ void FAttenuationSettingsCustomization::CustomizeChildren( TSharedRef<IPropertyH
 
 	ChildBuilder.AddChildProperty(PropertyHandles.FindChecked(GET_MEMBER_NAME_CHECKED(FAttenuationSettings, FalloffDistance)).ToSharedRef());
 	ChildBuilder.AddChildProperty(PropertyHandles.FindChecked(GET_MEMBER_NAME_CHECKED(FAttenuationSettings, OmniRadius)).ToSharedRef());
+	ChildBuilder.AddChildProperty(PropertyHandles.FindChecked(GET_MEMBER_NAME_CHECKED(FAttenuationSettings, StereoSpread)).ToSharedRef());
 	ChildBuilder.AddChildProperty(PropertyHandles.FindChecked(GET_MEMBER_NAME_CHECKED(FAttenuationSettings, bAttenuateWithLPF)).ToSharedRef());
 	ChildBuilder.AddChildProperty(PropertyHandles.FindChecked(GET_MEMBER_NAME_CHECKED(FAttenuationSettings, LPFRadiusMin)).ToSharedRef());
 	ChildBuilder.AddChildProperty(PropertyHandles.FindChecked(GET_MEMBER_NAME_CHECKED(FAttenuationSettings, LPFRadiusMax)).ToSharedRef());
 
-	if (PropertyHandles.Num() != 14)
+	if (PropertyHandles.Num() != 15)
 	{
 		FString PropertyList;
 		for (auto It(PropertyHandles.CreateConstIterator()); It; ++It)
@@ -260,4 +261,4 @@ EVisibility FAttenuationSettingsCustomization::IsCustomCurveSelected() const
 	const ESoundDistanceModel DistanceAlgorithm = (ESoundDistanceModel)DistanceAlgorithmValue;
 
 	return (DistanceAlgorithm == ATTENUATION_Custom ? EVisibility::Visible : EVisibility::Hidden);
-}
\ No newline at end of file
+}
diff --git a/Engine/Source/Runtime/Engine/Classes/Sound/SoundAttenuation.h b/Engine/Source/Runtime/Engine/Classes/Sound/SoundAttenuation.h
--- a/Engine/Source/Runtime/Engine/Classes/Sound/SoundAttenuation.h
+++ b/Engine/Source/Runtime/Engine/Classes/Sound/SoundAttenuation.h
@@ -89,6 +89,10 @@ struct ENGINE_API FAttenuationSettings
 	UPROPERTY(EditAnywhere, BlueprintReadWrite, Category=Attenuation, meta=(ClampMin = "0", EditCondition="bSpatialize", DisplayName="Non-Spatialized Radius"))
 	float OmniRadius;
 
+	/** The distance between left and right stereo channels when stereo assets spatialized. */
+	UPROPERTY(EditAnywhere, BlueprintReadWrite, Category = Attenuation, meta = (ClampMin = "0", EditCondition = "bSpatialize", DisplayName = "3D Stereo Spread"))
+	float StereoSpread;
+
 	/** Which spatialization algorithm to use if spatializing mono sources. */
 	UPROPERTY(EditAnywhere, BlueprintReadWrite, Category = Attenuation, meta = (ClampMin = "0", EditCondition = "bSpatialize", DisplayName = "Spatialization Algorithm"))
 	TEnumAsByte<enum ESoundSpatializationAlgorithm> SpatializationAlgorithm;
diff --git a/Engine/Source/Runtime/Engine/Private/ActiveSound.cpp b/Engine/Source/Runtime/Engine/Private/ActiveSound.cpp
--- a/Engine/Source/Runtime/Engine/Private/ActiveSound.cpp
+++ b/Engine/Source/Runtime/Engine/Private/ActiveSound.cpp
@@ -208,6 +208,7 @@ void FActiveSound::UpdateWaveInstances( FAudioDevice* AudioDevice, TArray<FWaveI
 		{
 			AttenuationSettings.ApplyAttenuation(ParseParams.Transform, Listener.Transform.GetTranslation(), ParseParams.Volume, ParseParams.HighFrequencyGain );
 			ParseParams.OmniRadius = AttenuationSettings.OmniRadius;
+			ParseParams.StereoSpread = AttenuationSettings.StereoSpread;
 			ParseParams.bUseSpatialization = AttenuationSettings.bSpatialize;
 			if (AttenuationSettings.SpatializationAlgorithm == SPATIALIZATION_Default && AudioDevice->IsHRTFEnabledForAll())
 			{
diff --git a/Engine/Source/Runtime/Engine/Private/Audio.cpp b/Engine/Source/Runtime/Engine/Private/Audio.cpp
--- a/Engine/Source/Runtime/Engine/Private/Audio.cpp
+++ b/Engine/Source/Runtime/Engine/Private/Audio.cpp
@@ -388,6 +388,8 @@ FWaveInstance::FWaveInstance( FActiveSound* InActiveSound )
 ,	Pitch( 0.0f )
 ,	Velocity( FVector::ZeroVector )
 ,	Location( FVector::ZeroVector )
+,	OmniRadius(0.0f)
+,	StereoSpread(0.0f)
 ,	UserIndex( 0 )
 {
 	TypeHash = ++TypeHashCounter;
diff --git a/Engine/Source/Runtime/Engine/Private/AudioDevice.cpp b/Engine/Source/Runtime/Engine/Private/AudioDevice.cpp
--- a/Engine/Source/Runtime/Engine/Private/AudioDevice.cpp
+++ b/Engine/Source/Runtime/Engine/Private/AudioDevice.cpp
@@ -772,6 +772,16 @@ bool FAudioDevice::HandleClearSoloCommand(const TCHAR* Cmd, FOutputDevice& Ar)
 	return true;
 }
 
+bool FAudioDevice::HandleAudio3dVisualizeCommand(const TCHAR* Cmd, FOutputDevice& Ar)
+{
+	FAudioDeviceManager* DeviceManager = GEngine->GetAudioDeviceManager();
+	if (DeviceManager)
+	{
+		DeviceManager->ToggleVisualize3dDebug();
+	}
+	return true;
+}
+
 #endif // !UE_BUILD_SHIPPING
 
 EDebugState FAudioDevice::GetMixDebugState( void )
@@ -878,6 +888,10 @@ bool FAudioDevice::Exec( UWorld* InWorld, const TCHAR* Cmd, FOutputDevice& Ar )
 	{
 		return HandleClearSoloCommand(Cmd, Ar);
 	}
+	else if (FParse::Command(&Cmd, TEXT("Audio3dVisualize")))
+	{
+		return HandleAudio3dVisualizeCommand(Cmd, Ar);
+	}
 #endif // !UE_BUILD_SHIPPING
 
 	return false;
diff --git a/Engine/Source/Runtime/Engine/Private/AudioDeviceManager.cpp b/Engine/Source/Runtime/Engine/Private/AudioDeviceManager.cpp
--- a/Engine/Source/Runtime/Engine/Private/AudioDeviceManager.cpp
+++ b/Engine/Source/Runtime/Engine/Private/AudioDeviceManager.cpp
@@ -37,6 +37,7 @@ FAudioDeviceManager::FAudioDeviceManager()
 	, NextResourceID(1)
 	, SoloDeviceHandle(INDEX_NONE)
 	, ActiveAudioDeviceHandle(INDEX_NONE)
+	, bVisualize3dDebug(false)
 {
 }
 
@@ -491,3 +492,15 @@ void FAudioDeviceManager::RemoveSoundMix(USoundMix* SoundMix)
 		}
 	}
 }
+
+bool FAudioDeviceManager::IsVisualizeDebug3dEnabled() const
+{
+	return bVisualize3dDebug;
+}
+
+void FAudioDeviceManager::ToggleVisualize3dDebug()
+{
+	bVisualize3dDebug = !bVisualize3dDebug;
+}
+
+
diff --git a/Engine/Source/Runtime/Engine/Private/DrawDebugHelpers.cpp b/Engine/Source/Runtime/Engine/Private/DrawDebugHelpers.cpp
--- a/Engine/Source/Runtime/Engine/Private/DrawDebugHelpers.cpp
+++ b/Engine/Source/Runtime/Engine/Private/DrawDebugHelpers.cpp
@@ -310,6 +310,28 @@ void DrawDebugCoordinateSystem(const UWorld* InWorld, FVector const& AxisLoc, FR
 	}
 }
 
+ENGINE_API void DrawDebugCrosshairs(const UWorld* InWorld, FVector const& AxisLoc, FRotator const& AxisRot, float Scale, const FColor& Color, bool bPersistentLines, float LifeTime, uint8 DepthPriority)
+{
+	// no debug line drawing on dedicated server
+	if (GEngine->GetNetMode(InWorld) != NM_DedicatedServer)
+	{
+		FRotationMatrix R(AxisRot);
+		FVector const X = 0.5f * R.GetScaledAxis(EAxis::X);
+		FVector const Y = 0.5f * R.GetScaledAxis(EAxis::Y);
+		FVector const Z = 0.5f * R.GetScaledAxis(EAxis::Z);
+
+		// this means foreground lines can't be persistent 
+		ULineBatchComponent* const LineBatcher = GetDebugLineBatcher(InWorld, bPersistentLines, LifeTime, (DepthPriority == SDPG_Foreground));
+		if (LineBatcher != NULL)
+		{
+			LineBatcher->DrawLine(AxisLoc - X*Scale, AxisLoc + X*Scale, Color, DepthPriority);
+			LineBatcher->DrawLine(AxisLoc - Y*Scale, AxisLoc + Y*Scale, Color, DepthPriority);
+			LineBatcher->DrawLine(AxisLoc - Z*Scale, AxisLoc + Z*Scale, Color, DepthPriority);
+		}
+	}
+}
+
+
 static void InternalDrawDebugCircle(const UWorld* InWorld, const FMatrix& TransformMatrix, float Radius, int32 Segments, const FColor& Color, bool bPersistentLines, float LifeTime, uint8 DepthPriority, float Thickness = 0.f)
 {
 	// no debug line drawing on dedicated server
diff --git a/Engine/Source/Runtime/Engine/Private/SoundAttenuation.cpp b/Engine/Source/Runtime/Engine/Private/SoundAttenuation.cpp
--- a/Engine/Source/Runtime/Engine/Private/SoundAttenuation.cpp
+++ b/Engine/Source/Runtime/Engine/Private/SoundAttenuation.cpp
@@ -191,6 +191,8 @@ bool FAttenuationSettings::operator==(const FAttenuationSettings& Other) const
 	return (   bAttenuate			    == Other.bAttenuate
 			&& bSpatialize			    == Other.bSpatialize
 			&& dBAttenuationAtMax	    == Other.dBAttenuationAtMax
+			&& OmniRadius				== Other.OmniRadius
+			&& StereoSpread				== Other.StereoSpread
 			&& DistanceAlgorithm	    == Other.DistanceAlgorithm
 			&& AttenuationShape		    == Other.AttenuationShape
 			&& bAttenuateWithLPF	    == Other.bAttenuateWithLPF
diff --git a/Engine/Source/Runtime/Engine/Private/SoundWave.cpp b/Engine/Source/Runtime/Engine/Private/SoundWave.cpp
--- a/Engine/Source/Runtime/Engine/Private/SoundWave.cpp
+++ b/Engine/Source/Runtime/Engine/Private/SoundWave.cpp
@@ -607,7 +607,7 @@ void USoundWave::Parse( FAudioDevice* AudioDevice, const UPTRINT NodeWaveInstanc
 		WaveInstance->StartTime = ParseParams.StartTime;
 		WaveInstance->UserIndex = ActiveSound.UserIndex;
 		WaveInstance->OmniRadius = ParseParams.OmniRadius;
-
+		WaveInstance->StereoSpread = ParseParams.StereoSpread;
 		bool bAlwaysPlay = false;
 
 		// Properties from the sound class
diff --git a/Engine/Source/Runtime/Engine/Public/ActiveSound.h b/Engine/Source/Runtime/Engine/Public/ActiveSound.h
--- a/Engine/Source/Runtime/Engine/Public/ActiveSound.h
+++ b/Engine/Source/Runtime/Engine/Public/ActiveSound.h
@@ -37,6 +37,9 @@ struct FSoundParseParameters
 	// At what distance from the source of the sound should spatialization begin
 	float OmniRadius;
 
+	// The distance between left and right channels when spatializing stereo assets
+	float StereoSpread;
+
 	// Which spatialization algorithm to use
 	ESoundSpatializationAlgorithm SpatializationAlgorithm;
 
@@ -55,6 +58,7 @@ struct FSoundParseParameters
 		, HighFrequencyGain(1.f)
 		, StartTime(-1.f)
 		, OmniRadius(0.0f)
+		, StereoSpread(0.0f)
 		, SpatializationAlgorithm(SPATIALIZATION_Default)
 		, bUseSpatialization(false)
 		, bLooping(false)
diff --git a/Engine/Source/Runtime/Engine/Public/Audio.h b/Engine/Source/Runtime/Engine/Public/Audio.h
--- a/Engine/Source/Runtime/Engine/Public/Audio.h
+++ b/Engine/Source/Runtime/Engine/Public/Audio.h
@@ -243,6 +243,8 @@ struct ENGINE_API FWaveInstance
 	FVector				Location;
 	/** At what distance we start transforming into omnidirectional soundsource */
 	float				OmniRadius;
+	/** Amount of spread for 3d multi-channel asset spatialization */
+	float				StereoSpread;
 	/** Cached type hash */
 	uint32				TypeHash;
 	/** Hash value for finding the wave instance based on the path through the cue to get to it */
diff --git a/Engine/Source/Runtime/Engine/Public/AudioDevice.h b/Engine/Source/Runtime/Engine/Public/AudioDevice.h
--- a/Engine/Source/Runtime/Engine/Public/AudioDevice.h
+++ b/Engine/Source/Runtime/Engine/Public/AudioDevice.h
@@ -212,6 +212,7 @@ public:
 	bool HandleEnableHRTFForAllCommand(const TCHAR* Cmd, FOutputDevice& Ar);
 	bool HandleSoloCommand( const TCHAR* Cmd, FOutputDevice& Ar );
 	bool HandleClearSoloCommand( const TCHAR* Cmd, FOutputDevice& Ar );
+	bool HandleAudio3dVisualizeCommand(const TCHAR* Cmd, FOutputDevice& Ar);
 #endif
 
 	/**
diff --git a/Engine/Source/Runtime/Engine/Public/AudioDeviceManager.h b/Engine/Source/Runtime/Engine/Public/AudioDeviceManager.h
--- a/Engine/Source/Runtime/Engine/Public/AudioDeviceManager.h
+++ b/Engine/Source/Runtime/Engine/Public/AudioDeviceManager.h
@@ -114,6 +114,12 @@ public:
 	/** Removes sound mix from all audio devices */
 	void RemoveSoundMix(USoundMix* SoundMix);
 
+	/** Is debug visualization of 3d sounds enabled */
+	bool IsVisualizeDebug3dEnabled() const;
+
+	/** Toggles 3d visualization of 3d sounds on/off */
+	void ToggleVisualize3dDebug();
+
 public:
 
 	/** Array of all created buffers */
@@ -166,6 +172,9 @@ private:
 
 	/** Which audio device is currently active */
 	uint32 ActiveAudioDeviceHandle;
+
+	/** Whether or not 3d debug visualization is enabled. */
+	bool bVisualize3dDebug;
 };
 
 
diff --git a/Engine/Source/Runtime/Engine/Public/DrawDebugHelpers.h b/Engine/Source/Runtime/Engine/Public/DrawDebugHelpers.h
--- a/Engine/Source/Runtime/Engine/Public/DrawDebugHelpers.h
+++ b/Engine/Source/Runtime/Engine/Public/DrawDebugHelpers.h
@@ -21,6 +21,8 @@ ENGINE_API void DrawDebugBox(const UWorld* InWorld, FVector const& Center, FVect
 ENGINE_API void DrawDebugBox(const UWorld* InWorld, FVector const& Center, FVector const& Box, const FQuat& Rotation, FColor const& Color, bool bPersistentLines = false, float LifeTime=-1.f, uint8 DepthPriority = 0);
 /** Draw Debug coordinate system */
 ENGINE_API void DrawDebugCoordinateSystem(const UWorld* InWorld, FVector const& AxisLoc, FRotator const& AxisRot, float Scale, bool bPersistentLines = false, float LifeTime=-1.f, uint8 DepthPriority = 0);
+/** Draw Debug crosshair */
+ENGINE_API void DrawDebugCrosshairs(const UWorld* InWorld, FVector const& AxisLoc, FRotator const& AxisRot, float Scale, const FColor& Color = FColor::White, bool bPersistentLines = false, float LifeTime = -1.f, uint8 DepthPriority = 0);
 /** Draw Debug Circle */
 ENGINE_API void DrawDebugCircle(const UWorld* InWorld, const FMatrix& TransformMatrix, float Radius, int32 Segments, const FColor& Color, bool bPersistentLines = false, float LifeTime=-1.f, uint8 DepthPriority = 0, float Thickness = 0.f, bool bDrawAxis=true);
 /** Draw Debug Circle */
diff --git a/Engine/Source/Runtime/Windows/XAudio2/Private/XAudio2Device.cpp b/Engine/Source/Runtime/Windows/XAudio2/Private/XAudio2Device.cpp
--- a/Engine/Source/Runtime/Windows/XAudio2/Private/XAudio2Device.cpp
+++ b/Engine/Source/Runtime/Windows/XAudio2/Private/XAudio2Device.cpp
@@ -248,7 +248,7 @@ void FXAudio2Device::UpdateHardware()
 		// at the Y component after normalization to determine spatialization.
 		const FVector Up = Listeners[0].GetUp();
 		const FVector Right = Listeners[0].GetFront();
-		InverseTransform = FMatrix(Up, Right, Up ^ Right, Listeners[0].Transform.GetTranslation()).InverseFast();
+		InverseListenerTransform = FMatrix(Up, Right, Up ^ Right, Listeners[0].Transform.GetTranslation()).InverseFast();
 	}
 }
 
diff --git a/Engine/Source/Runtime/Windows/XAudio2/Private/XAudio2Source.cpp b/Engine/Source/Runtime/Windows/XAudio2/Private/XAudio2Source.cpp
--- a/Engine/Source/Runtime/Windows/XAudio2/Private/XAudio2Source.cpp
+++ b/Engine/Source/Runtime/Windows/XAudio2/Private/XAudio2Source.cpp
@@ -67,15 +67,17 @@ FXMPHelper* FXMPHelper::GetXMPHelper( void )
  * Simple constructor
  */
 FXAudio2SoundSource::FXAudio2SoundSource(FAudioDevice* InAudioDevice)
-:	FSoundSource( InAudioDevice ),
-	Source( NULL ),
-	RealtimeAsyncTask( nullptr ),
-	CurrentBuffer( 0 ),
-	bBuffersToFlush( false ),
-	bLoopCallback( false ),
-	bResourcesNeedFreeing(false),
-	VoiceId(-1),
-	bUsingDefaultSpatialization(true)
+	: FSoundSource( InAudioDevice )
+	, Source( NULL )
+	, RealtimeAsyncTask( nullptr )
+	, CurrentBuffer( 0 )
+	, bBuffersToFlush( false )
+	, bLoopCallback( false )
+	, bResourcesNeedFreeing(false)
+	, VoiceId(-1)
+	, bUsingHRTFSpatialization(false)
+	, LeftChannelSourceLocation(0)
+	, RightChannelSourceLocation(0)
 {
 	AudioDevice = ( FXAudio2Device* )InAudioDevice;
 	check( AudioDevice );
@@ -403,7 +405,7 @@ bool FXAudio2SoundSource::CreateSource( void )
 	Flags |= XAUDIO2_VOICE_USEFILTER;
 
 	// Reset the bUsingSpatializationEffect flag
-	bUsingDefaultSpatialization = true;
+	bUsingHRTFSpatialization = false;
 	bool bCreatedWithSpatializationEffect = false;
 
 	if (CreateWithSpatializationEffect())
@@ -415,7 +417,7 @@ bool FXAudio2SoundSource::CreateSource( void )
 			// Indicate that this source is currently using the 3d spatialization effect. We can't stop using it
 			// for the lifetime of this sound, so if this if the spatialization effect is toggled off, we're still 
 			// going to hear the sound for the duration of this sound.
-			bUsingDefaultSpatialization = false;
+			bUsingHRTFSpatialization = true;
 
 			XAUDIO2_EFFECT_DESCRIPTOR EffectDescriptor[1];
 			EffectDescriptor[0].pEffect = Effect;
@@ -542,7 +544,7 @@ bool FXAudio2SoundSource::Init(FWaveInstance* InWaveInstance)
 /**
  * Calculates the volume for each channel
  */
-void FXAudio2SoundSource::GetChannelVolumes( float ChannelVolumes[CHANNELOUT_COUNT], float AttenuatedVolume )
+void FXAudio2SoundSource::GetChannelVolumes(float ChannelVolumes[CHANNEL_MATRIX_COUNT], float AttenuatedVolume)
 {
 	if (FApp::GetVolumeMultiplier() == 0.0f || AudioDevice->bIsDeviceMuted)
 	{
@@ -591,7 +593,7 @@ void FXAudio2SoundSource::GetChannelVolumes( float ChannelVolumes[CHANNELOUT_COU
 		break;
 	};
 
-	for (int32 i = 0; i < CHANNELOUT_COUNT; i++)
+	for (int32 i = 0; i < CHANNEL_MATRIX_COUNT; i++)
 	{
 		// Detect and warn about NaN and INF volumes. XAudio does not do this internally and behavior is undefined.
 		// This is known to happen in X3DAudioCalculate in channel 0 and the cause is unknown.
@@ -611,36 +613,80 @@ void FXAudio2SoundSource::GetChannelVolumes( float ChannelVolumes[CHANNELOUT_COU
 	}
 }
 
-void FXAudio2SoundSource::GetMonoChannelVolumes(float ChannelVolumes[CHANNELOUT_COUNT], float AttenuatedVolume)
+FVector FXAudio2SoundSource::ConvertToXAudio2Orientation(const FVector& InputVector)
 {
-	if (IsUsingDefaultSpatializer())
+	return FVector(InputVector.Y, InputVector.X, -InputVector.Z);
+
+}
+
+FVector FXAudio2SoundSource::GetListenerTransformedDirection(const FVector& Position, float* OutDistance)
+{
+	FVector UnnormalizedDirection = AudioDevice->InverseListenerTransform.TransformPosition(Position);
+	if (OutDistance)
+	{
+		*OutDistance = UnnormalizedDirection.Size();
+	}
+	return UnnormalizedDirection.GetSafeNormal();
+}
+
+void FXAudio2SoundSource::GetSpatializationParams(FSpatializationParams& Params)
+{
+	// Calculate direction from listener to sound, where the sound is at the origin if unspatialised.
+	if (WaveInstance->bUseSpatialization)
 	{
-		// Calculate direction from listener to sound, where the sound is at the origin if unspatialised.
-		FVector Direction = FVector::ZeroVector;
-		float NormalizedOmniRadius = 0;
-		if (WaveInstance->bUseSpatialization)
+		FVector EmitterDirection = GetListenerTransformedDirection(WaveInstance->Location, &Params.Distance);
+		Params.NormalizedOmniRadius = (Params.Distance > 0) ? (WaveInstance->OmniRadius / Params.Distance) : FLT_MAX;
+
+		if (Buffer->NumChannels == 2)
+		{
+			FVector EmitterDirection = GetListenerTransformedDirection(LeftChannelSourceLocation, nullptr);
+			Params.LeftChannelPosition = ConvertToXAudio2Orientation(EmitterDirection);
+
+			EmitterDirection = GetListenerTransformedDirection(RightChannelSourceLocation, nullptr);
+			Params.RightChannelPosition = ConvertToXAudio2Orientation(EmitterDirection);
+
+			Params.EmitterPosition = FVector::ZeroVector;
+		}
+		else
 		{
-			FVector UnnormalizedDirection = AudioDevice->InverseTransform.TransformPosition(WaveInstance->Location);
-			Direction = UnnormalizedDirection.GetSafeNormal();
-			float Distance = UnnormalizedDirection.Size();
-			NormalizedOmniRadius = (Distance > 0) ? (WaveInstance->OmniRadius / Distance) : FLT_MAX;
+			Params.EmitterPosition = ConvertToXAudio2Orientation(EmitterDirection);
 		}
+	}
+	else
+	{
+		Params.NormalizedOmniRadius = 0.0f;
+		Params.Distance = 0.0f;
+		Params.EmitterPosition = FVector::ZeroVector;
+	}
 
-		// Calculate 5.1 channel volume
-		FVector OrientFront;
-		OrientFront.X = 0.0f;
-		OrientFront.Y = 0.0f;
-		OrientFront.Z = 1.0f;
+	// We are currently always computing spatialization for XAudio2 relative to the listener!
+	Params.ListenerOrientation = FVector::UpVector;
+	Params.ListenerPosition = FVector::ZeroVector;
+}
 
-		FVector ListenerPosition;
-		ListenerPosition.X = 0.0f;
-		ListenerPosition.Y = 0.0f;
-		ListenerPosition.Z = 0.0f;
 
-		FVector EmitterPosition;
-		EmitterPosition.X = Direction.Y;
-		EmitterPosition.Y = Direction.X;
-		EmitterPosition.Z = -Direction.Z;
+void FXAudio2SoundSource::GetMonoChannelVolumes(float ChannelVolumes[CHANNEL_MATRIX_COUNT], float AttenuatedVolume)
+{
+	if (IsUsingHrtfSpatializer())
+	{
+		// If we are using a HRTF spatializer, we are going to be using an XAPO effect that takes a mono stream and splits it into stereo
+		// So in th at case we will just set the emitter position as a parameter to the XAPO plugin and then treat the
+		// sound as if it was a non-spatialized stereo asset
+		check(WaveInstance->SpatializationAlgorithm == SPATIALIZATION_HRTF);
+		check(AudioDevice->SpatializeProcessor != nullptr);
+
+		// If we're using an HRTF spatialization algorithm, we need to find the position of the emitter and set it as a parameter
+		FVector EmitterPosition = FVector::ZeroVector;
+		FVector UnnormalizedDirection = AudioDevice->InverseListenerTransform.TransformPosition(WaveInstance->Location);
+		EmitterPosition = UnnormalizedDirection.GetSafeNormal();
+
+		AudioDevice->SpatializeProcessor->SetSpatializationParameters(VoiceId, EmitterPosition, (ESpatializationEffectType)WaveInstance->SpatializationAlgorithm);
+		GetStereoChannelVolumes(ChannelVolumes, AttenuatedVolume);
+	}
+	else // Spatialize the mono stream using the normal 3d audio algorithm
+	{
+		FSpatializationParams SpatializationParams;
+		GetSpatializationParams(SpatializationParams);
 
 		// Calculate 5.1 channel dolby surround rate/multipliers.
 		ChannelVolumes[CHANNELOUT_FRONTLEFT] = AttenuatedVolume;
@@ -657,7 +703,7 @@ void FXAudio2SoundSource::GetMonoChannelVolumes(float ChannelVolumes[CHANNELOUT_
 		ChannelVolumes[CHANNELOUT_RADIO] = 0.0f;
 
 		// Call the spatialisation magic
-		AudioDevice->SpatializationHelper.CalculateDolbySurroundRate(OrientFront, ListenerPosition, EmitterPosition, NormalizedOmniRadius, ChannelVolumes);
+		AudioDevice->SpatializationHelper.CalculateDolbySurroundRate(SpatializationParams.ListenerOrientation, SpatializationParams.ListenerPosition, SpatializationParams.EmitterPosition, SpatializationParams.NormalizedOmniRadius, ChannelVolumes);
 
 		// Handle any special post volume processing
 		if (WaveInstance->bApplyRadioFilter)
@@ -696,50 +742,106 @@ void FXAudio2SoundSource::GetMonoChannelVolumes(float ChannelVolumes[CHANNELOUT_
 			ChannelVolumes[CHANNELOUT_FRONTCENTER] = FMath::Max(ChannelVolumes[CHANNELOUT_FRONTCENTER], WaveInstance->VoiceCenterChannelVolume * AttenuatedVolume);
 		}
 	}
-	else if (WaveInstance->bUseSpatialization)
-	{
-		check(WaveInstance->SpatializationAlgorithm == SPATIALIZATION_HRTF);
-		check(AudioDevice->SpatializeProcessor != nullptr);
+}
 
-		// If we're using an HRTF spatialization algorithm, we need to find the position of the emitter and set it as a parameter
-		FVector EmitterPosition = FVector::ZeroVector;
-		FVector UnnormalizedDirection = AudioDevice->InverseTransform.TransformPosition(WaveInstance->Location);
-		EmitterPosition = UnnormalizedDirection.GetSafeNormal();
+void FXAudio2SoundSource::UpdateStereoEmitterPositions()
+{
+	// We need to compute the stereo left/right channel positions using the audio component position and the spread 
+	const FVector& ListenerPosition = AudioDevice->Listeners[0].Transform.GetLocation();
+	FVector ListenerToSourceDir = WaveInstance->Location - ListenerPosition;
+	ListenerToSourceDir.Normalize();
 
-		AudioDevice->SpatializeProcessor->SetSpatializationParameters(VoiceId, EmitterPosition, (ESpatializationEffectType)WaveInstance->SpatializationAlgorithm);
-		GetStereoChannelVolumes(ChannelVolumes, AttenuatedVolume);
-	}
+	float HalfSpread = 0.5f * WaveInstance->StereoSpread;
+
+	// Get direction of left emitter from true emitter position (left hand rule)
+	FVector LeftEmitterDir = FVector::CrossProduct(ListenerToSourceDir, FVector::UpVector);
+	FVector LeftEmitterOffset = LeftEmitterDir * HalfSpread;
+
+	// Get position vector of left emitter by adding to true emitter the dir scaled by half the spread
+	LeftChannelSourceLocation = WaveInstance->Location + LeftEmitterOffset;
+
+	// Right emitter position is same as right but opposite direction
+	RightChannelSourceLocation = WaveInstance->Location - LeftEmitterOffset;
 }
 
-void FXAudio2SoundSource::GetStereoChannelVolumes(float ChannelVolumes[CHANNELOUT_COUNT], float AttenuatedVolume)
+void FXAudio2SoundSource::GetStereoChannelVolumes(float ChannelVolumes[CHANNEL_MATRIX_COUNT], float AttenuatedVolume)
 {
-	// Stereo is always treated as unspatialized (except when the oculus stereo effect is being used)
-	ChannelVolumes[CHANNELOUT_FRONTLEFT] = AttenuatedVolume;
-	ChannelVolumes[CHANNELOUT_FRONTRIGHT] = AttenuatedVolume;
-
-	// Potentially bleed to the rear speakers from 2.0 channel to simulated 4.0 channel
-	if (IsUsingDefaultSpatializer() && FXAudioDeviceProperties::NumSpeakers == 6)
+	// If we're doing 3d spatializaton of stereo sounds
+	if (!IsUsingHrtfSpatializer() && WaveInstance->bUseSpatialization)
 	{
-		ChannelVolumes[CHANNELOUT_LEFTSURROUND] = AttenuatedVolume * StereoBleed;
-		ChannelVolumes[CHANNELOUT_RIGHTSURROUND] = AttenuatedVolume * StereoBleed;
+		check(MAX_INPUT_CHANNELS_SPATIALIZED >= 2);
 
-		ChannelVolumes[CHANNELOUT_LOWFREQUENCY] = AttenuatedVolume * LFEBleed * 0.5f;
-	}
+		// Loop through the left and right input channels and set the attenuation volumes
+		for (int32 i = 0; i < 2; ++i)
+		{
+			// Offset is the offset into the channel matrix
+			int32 Offset = CHANNELOUT_COUNT*i;
+			ChannelVolumes[CHANNELOUT_FRONTLEFT + Offset] = AttenuatedVolume;
+			ChannelVolumes[CHANNELOUT_FRONTRIGHT + Offset] = AttenuatedVolume;
+			ChannelVolumes[CHANNELOUT_FRONTCENTER + Offset] = AttenuatedVolume;
+			ChannelVolumes[CHANNELOUT_LEFTSURROUND + Offset] = AttenuatedVolume;
+			ChannelVolumes[CHANNELOUT_RIGHTSURROUND + Offset] = AttenuatedVolume;
+
+			if (bReverbApplied)
+			{
+				ChannelVolumes[CHANNELOUT_REVERB + Offset] = AttenuatedVolume;
+			}
 
-	if (bReverbApplied)
-	{
-		ChannelVolumes[CHANNELOUT_REVERB] = AttenuatedVolume;
-	}
+			ChannelVolumes[CHANNELOUT_RADIO + Offset] = 0.0f;
+
+			// Add some LFE bleed 
+			if (FXAudioDeviceProperties::NumSpeakers == 6)
+			{
+				ChannelVolumes[CHANNELOUT_LOWFREQUENCY + Offset] = AttenuatedVolume * LFEBleed;
+			}
+		}
+
+		// Make sure we have up-to-date left and right channel positions for stereo spatialization
+		UpdateStereoEmitterPositions();
 
-	// Handle radio distortion if the sound can handle it. 
-	ChannelVolumes[CHANNELOUT_RADIO] = 0.0f;
-	if (WaveInstance->bApplyRadioFilter)
+		// Now get the spatialization params transformed into listener-space
+		FSpatializationParams SpatializationParams;
+		GetSpatializationParams(SpatializationParams);
+
+		// Compute the speaker mappings for the left channel
+		float* ChannelMap = ChannelVolumes;
+		AudioDevice->SpatializationHelper.CalculateDolbySurroundRate(SpatializationParams.ListenerOrientation, SpatializationParams.ListenerPosition, SpatializationParams.LeftChannelPosition, SpatializationParams.NormalizedOmniRadius, ChannelMap);
+		
+		// Now compute the speaker mappings for the right channel
+		ChannelMap = &ChannelVolumes[CHANNELOUT_COUNT];
+		AudioDevice->SpatializationHelper.CalculateDolbySurroundRate(SpatializationParams.ListenerOrientation, SpatializationParams.ListenerPosition, SpatializationParams.RightChannelPosition, SpatializationParams.NormalizedOmniRadius, ChannelMap);
+	}
+	else
 	{
-		ChannelVolumes[CHANNELOUT_RADIO] = WaveInstance->RadioFilterVolume;
+		// Stereo is always treated as unspatialized (except when the HRTF spatialization effect is being used)
+		ChannelVolumes[CHANNELOUT_FRONTLEFT] = AttenuatedVolume;
+		ChannelVolumes[CHANNELOUT_FRONTRIGHT] = AttenuatedVolume;
+
+		// Potentially bleed to the rear speakers from 2.0 channel to simulated 4.0 channel
+		// but only if this is not an HRTF-spatialized mono sound
+		if (!IsUsingHrtfSpatializer() && FXAudioDeviceProperties::NumSpeakers == 6)
+		{
+			ChannelVolumes[CHANNELOUT_LEFTSURROUND] = AttenuatedVolume * StereoBleed;
+			ChannelVolumes[CHANNELOUT_RIGHTSURROUND] = AttenuatedVolume * StereoBleed;
+
+			ChannelVolumes[CHANNELOUT_LOWFREQUENCY] = AttenuatedVolume * LFEBleed * 0.5f;
+		}
+
+		if (bReverbApplied)
+		{
+			ChannelVolumes[CHANNELOUT_REVERB] = AttenuatedVolume;
+		}
+
+		// Handle radio distortion if the sound can handle it. 
+		ChannelVolumes[CHANNELOUT_RADIO] = 0.0f;
+		if (WaveInstance->bApplyRadioFilter)
+		{
+			ChannelVolumes[CHANNELOUT_RADIO] = WaveInstance->RadioFilterVolume;
+		}
 	}
 }
 
-void FXAudio2SoundSource::GetQuadChannelVolumes(float ChannelVolumes[CHANNELOUT_COUNT], float AttenuatedVolume)
+void FXAudio2SoundSource::GetQuadChannelVolumes(float ChannelVolumes[CHANNEL_MATRIX_COUNT], float AttenuatedVolume)
 {
 	ChannelVolumes[CHANNELOUT_FRONTLEFT] = AttenuatedVolume;
 	ChannelVolumes[CHANNELOUT_FRONTRIGHT] = AttenuatedVolume;
@@ -752,7 +854,7 @@ void FXAudio2SoundSource::GetQuadChannelVolumes(float ChannelVolumes[CHANNELOUT_
 	}
 }
 
-void FXAudio2SoundSource::GetHexChannelVolumes(float ChannelVolumes[CHANNELOUT_COUNT], float AttenuatedVolume)
+void FXAudio2SoundSource::GetHexChannelVolumes(float ChannelVolumes[CHANNEL_MATRIX_COUNT], float AttenuatedVolume)
 {
 	ChannelVolumes[CHANNELOUT_FRONTLEFT] = AttenuatedVolume;
 	ChannelVolumes[CHANNELOUT_FRONTRIGHT] = AttenuatedVolume;
@@ -765,7 +867,7 @@ void FXAudio2SoundSource::GetHexChannelVolumes(float ChannelVolumes[CHANNELOUT_C
 /** 
  * Maps a sound with a given number of channels to to expected speakers
  */
-void FXAudio2SoundSource::RouteDryToSpeakers( float ChannelVolumes[CHANNELOUT_COUNT] )
+void FXAudio2SoundSource::RouteDryToSpeakers(float ChannelVolumes[CHANNEL_MATRIX_COUNT])
 {
 	// Only need to account to the special cases that are not a simple match of channel to speaker
 	switch( Buffer->NumChannels )
@@ -791,9 +893,14 @@ void FXAudio2SoundSource::RouteDryToSpeakers( float ChannelVolumes[CHANNELOUT_CO
 	};
 }
 
-void FXAudio2SoundSource::RouteMonoToDry(float ChannelVolumes[CHANNELOUT_COUNT])
+void FXAudio2SoundSource::RouteMonoToDry(float ChannelVolumes[CHANNEL_MATRIX_COUNT])
 {
-	if (IsUsingDefaultSpatializer())
+	if (IsUsingHrtfSpatializer())
+	{
+		// If we're spatializing using HRTF algorithms, then our output is actually stereo
+		RouteStereoToDry(ChannelVolumes);
+	}
+	else
 	{
 		// Spatialised audio maps 1 channel to 6 speakers
 		float SpatialisationMatrix[SPEAKER_COUNT * 1] =
@@ -809,40 +916,53 @@ void FXAudio2SoundSource::RouteMonoToDry(float ChannelVolumes[CHANNELOUT_COUNT])
 		// Update the dry output to the mastering voice
 		AudioDevice->ValidateAPICall(TEXT("SetOutputMatrix (mono)"), Source->SetOutputMatrix(Destinations[DEST_DRY].pOutputVoice, 1, SPEAKER_COUNT, SpatialisationMatrix));
 	}
-	else
-	{
-		// If we're spatializing using HRTF algorithms, then our output is actually stereo
-		RouteStereoToDry(ChannelVolumes);
-	}
 }
 
-void FXAudio2SoundSource::RouteStereoToDry(float ChannelVolumes[CHANNELOUT_COUNT])
+void FXAudio2SoundSource::RouteStereoToDry(float Chans[CHANNEL_MATRIX_COUNT])
 {
-	if (IsUsingDefaultSpatializer())
+	if (IsUsingHrtfSpatializer())
 	{
+		// A 2d sound
 		float SpatialisationMatrix[SPEAKER_COUNT * 2] =
 		{
-			ChannelVolumes[CHANNELOUT_FRONTLEFT], 0.0f,
-			0.0f, ChannelVolumes[CHANNELOUT_FRONTRIGHT],
-			0.0f, 0.0f,
-			ChannelVolumes[CHANNELOUT_LOWFREQUENCY], ChannelVolumes[CHANNELOUT_LOWFREQUENCY],
-			ChannelVolumes[CHANNELOUT_LEFTSURROUND], 0.0f,
-			0.0f, ChannelVolumes[CHANNELOUT_RIGHTSURROUND]
+			// Left Input					Right Input
+			Chans[CHANNELOUT_FRONTLEFT],	0.0f,							// Left
+			0.0f,							Chans[CHANNELOUT_FRONTRIGHT],	// Right
+			0.0f,							0.0f,							// Center
+			0.0f,							0.0f,							// LFE
+			0.0f,							0.0f,							// Left Surround
+			0.0f,							0.0f							// Right Surround
 		};
-
 		// Stereo sounds map 2 channels to 6 speakers
+		AudioDevice->ValidateAPICall(TEXT("SetOutputMatrix (stereo)"), Source->SetOutputMatrix(Destinations[DEST_DRY].pOutputVoice, 2, SPEAKER_COUNT, SpatialisationMatrix));		// Build a non-3d "multi-channel" blend from the stereo channels
+	}
+	else if (WaveInstance->bUseSpatialization)
+	{
+		// Build a non-3d "multi-channel" blend from the stereo channels
+		float SpatialisationMatrix[SPEAKER_COUNT * 2] =
+		{
+			// Left Input						Right Input
+			Chans[CHANNELOUT_FRONTLEFT],		Chans[CHANNELOUT_COUNT + CHANNELOUT_FRONTLEFT],		// Left
+			Chans[CHANNELOUT_FRONTRIGHT],		Chans[CHANNELOUT_COUNT + CHANNELOUT_FRONTRIGHT],	// Right
+			Chans[CHANNELOUT_FRONTCENTER],		Chans[CHANNELOUT_COUNT + CHANNELOUT_FRONTCENTER],	// Right
+			Chans[CHANNELOUT_LOWFREQUENCY],		Chans[CHANNELOUT_COUNT + CHANNELOUT_LOWFREQUENCY],	// LFE
+			Chans[CHANNELOUT_LEFTSURROUND],		Chans[CHANNELOUT_COUNT + CHANNELOUT_LEFTSURROUND],	// Left Surround
+			Chans[CHANNELOUT_RIGHTSURROUND],	Chans[CHANNELOUT_COUNT + CHANNELOUT_RIGHTSURROUND],	// Right Surround
+		};
+
 		AudioDevice->ValidateAPICall(TEXT("SetOutputMatrix (stereo)"), Source->SetOutputMatrix(Destinations[DEST_DRY].pOutputVoice, 2, SPEAKER_COUNT, SpatialisationMatrix));
 	}
 	else
 	{
 		float SpatialisationMatrix[SPEAKER_COUNT * 2] =
 		{
-			ChannelVolumes[CHANNELOUT_FRONTLEFT], 0.0f,
-			0.0f, ChannelVolumes[CHANNELOUT_FRONTRIGHT],
-			0.0f, 0.0f,
-			0.0f, 0.0f,
-			0.0f, 0.0f,
-			0.0f, 0.0f
+			// Left Input					Right Input
+			Chans[CHANNELOUT_FRONTLEFT],	0.0f,								// Left
+			0.0f,							Chans[CHANNELOUT_FRONTRIGHT],		// Right
+			0.0f,							0.0f,								// Center
+			Chans[CHANNELOUT_LOWFREQUENCY], Chans[CHANNELOUT_LOWFREQUENCY],		// LFE
+			Chans[CHANNELOUT_LEFTSURROUND], 0.0f,								// Left Surround
+			0.0f,							Chans[CHANNELOUT_RIGHTSURROUND]		// Right Surround
 		};
 
 		// Stereo sounds map 2 channels to 6 speakers
@@ -850,16 +970,17 @@ void FXAudio2SoundSource::RouteStereoToDry(float ChannelVolumes[CHANNELOUT_COUNT
 	}
 }
 
-void FXAudio2SoundSource::RouteQuadToDry(float ChannelVolumes[CHANNELOUT_COUNT])
+void FXAudio2SoundSource::RouteQuadToDry(float Chans[CHANNEL_MATRIX_COUNT])
 {
 	float SpatialisationMatrix[SPEAKER_COUNT * 4] =
 	{
-		ChannelVolumes[CHANNELOUT_FRONTLEFT], 0.0f, 0.0f, 0.0f,
-		0.0f, ChannelVolumes[CHANNELOUT_FRONTRIGHT], 0.0f, 0.0f,
-		0.0f, 0.0f, 0.0f, 0.0f,
-		ChannelVolumes[CHANNELOUT_LOWFREQUENCY], ChannelVolumes[CHANNELOUT_LOWFREQUENCY], ChannelVolumes[CHANNELOUT_LOWFREQUENCY], ChannelVolumes[CHANNELOUT_LOWFREQUENCY],
-		0.0f, 0.0f, ChannelVolumes[CHANNELOUT_LEFTSURROUND], 0.0f,
-		0.0f, 0.0f, 0.0f, ChannelVolumes[CHANNELOUT_RIGHTSURROUND]
+		// Left Input						Right Input						Left Surround Input				Right Surround Input
+		Chans[CHANNELOUT_FRONTLEFT],		0.0f,							0.0f,							0.0f,								// Left
+		0.0f,								Chans[CHANNELOUT_FRONTRIGHT],	0.0f,							0.0f,								// Right
+		0.0f,								0.0f,							0.0f,							0.0f,								// Center
+		Chans[CHANNELOUT_LOWFREQUENCY],		Chans[CHANNELOUT_LOWFREQUENCY],	Chans[CHANNELOUT_LOWFREQUENCY],	Chans[CHANNELOUT_LOWFREQUENCY],		// LFE
+		0.0f,								0.0f,							Chans[CHANNELOUT_LEFTSURROUND],	0.0f,								// Left Surround
+		0.0f,								0.0f,							0.0f,							Chans[CHANNELOUT_RIGHTSURROUND]		// Right Surround
 	};
 
 	// Quad sounds map 4 channels to 6 speakers
@@ -867,7 +988,7 @@ void FXAudio2SoundSource::RouteQuadToDry(float ChannelVolumes[CHANNELOUT_COUNT])
 		Source->SetOutputMatrix(Destinations[DEST_DRY].pOutputVoice, 4, SPEAKER_COUNT, SpatialisationMatrix));
 }
 
-void FXAudio2SoundSource::RouteHexToDry(float ChannelVolumes[CHANNELOUT_COUNT])
+void FXAudio2SoundSource::RouteHexToDry(float Chans[CHANNEL_MATRIX_COUNT])
 {
 	if ((XAudio2Buffer->DecompressionState && XAudio2Buffer->DecompressionState->UsesVorbisChannelOrdering())
 		|| WaveInstance->WaveData->bDecompressedFromOgg)
@@ -875,12 +996,13 @@ void FXAudio2SoundSource::RouteHexToDry(float ChannelVolumes[CHANNELOUT_COUNT])
 		// Ordering of channels is different for 6 channel OGG
 		float SpatialisationMatrix[SPEAKER_COUNT * 6] =
 		{
-			ChannelVolumes[CHANNELOUT_FRONTLEFT], 0.0f, 0.0f, 0.0f, 0.0f, 0.0f,
-			0.0f, 0.0f, ChannelVolumes[CHANNELOUT_FRONTRIGHT], 0.0f, 0.0f, 0.0f,
-			0.0f, ChannelVolumes[CHANNELOUT_FRONTCENTER], 0.0f, 0.0f, 0.0f, 0.0f,
-			0.0f, 0.0f, 0.0f, 0.0f, 0.0f, ChannelVolumes[CHANNELOUT_LOWFREQUENCY],
-			0.0f, 0.0f, 0.0f, ChannelVolumes[CHANNELOUT_LEFTSURROUND], 0.0f, 0.0f,
-			0.0f, 0.0f, 0.0f, 0.0f, ChannelVolumes[CHANNELOUT_RIGHTSURROUND], 0.0f
+			// Left In						Center In						Right In						Left Surround In				Right Surround In					LFE In
+			Chans[CHANNELOUT_FRONTLEFT],	0.0f,							0.0f,							0.0f,							0.0f,								0.0f,							// Left Out
+			0.0f,							0.0f,							Chans[CHANNELOUT_FRONTRIGHT],	0.0f,							0.0f,								0.0f,							// Right Out
+			0.0f,							Chans[CHANNELOUT_FRONTCENTER],	0.0f,							0.0f,							0.0f,								0.0f,							// Center Out
+			0.0f,							0.0f,							0.0f,							0.0f,							0.0f,								Chans[CHANNELOUT_LOWFREQUENCY],	// LFE Out
+			0.0f,							0.0f,							0.0f,							Chans[CHANNELOUT_LEFTSURROUND], 0.0f,								0.0f,							// Left Surround Out
+			0.0f,							0.0f,							0.0f,							0.0f,							Chans[CHANNELOUT_RIGHTSURROUND],	0.0f							// Right Surround Out
 		};
 
 		// 5.1 sounds map 6 channels to 6 speakers
@@ -891,12 +1013,13 @@ void FXAudio2SoundSource::RouteHexToDry(float ChannelVolumes[CHANNELOUT_COUNT])
 	{
 		float SpatialisationMatrix[SPEAKER_COUNT * 6] =
 		{
-			ChannelVolumes[CHANNELOUT_FRONTLEFT], 0.0f, 0.0f, 0.0f, 0.0f, 0.0f,
-			0.0f, ChannelVolumes[CHANNELOUT_FRONTRIGHT], 0.0f, 0.0f, 0.0f, 0.0f,
-			0.0f, 0.0f, ChannelVolumes[CHANNELOUT_FRONTCENTER], 0.0f, 0.0f, 0.0f,
-			0.0f, 0.0f, 0.0f, ChannelVolumes[CHANNELOUT_LOWFREQUENCY], 0.0f, 0.0f,
-			0.0f, 0.0f, 0.0f, 0.0f, ChannelVolumes[CHANNELOUT_LEFTSURROUND], 0.0f,
-			0.0f, 0.0f, 0.0f, 0.0f, 0.0f, ChannelVolumes[CHANNELOUT_RIGHTSURROUND]
+			// Left In						Center In						Right In						Left Surround In				Right Surround In					LFE In
+			Chans[CHANNELOUT_FRONTLEFT],	0.0f,							0.0f,							0.0f,							0.0f,								0.0f,								// Left Out
+			0.0f,							Chans[CHANNELOUT_FRONTRIGHT],	0.0f,							0.0f,							0.0f,								0.0f,								// Right Out
+			0.0f,							0.0f,							Chans[CHANNELOUT_FRONTCENTER],	0.0f,							0.0f,								0.0f,								// Center Out
+			0.0f,							0.0f,							0.0f,							Chans[CHANNELOUT_LOWFREQUENCY], 0.0f,								0.0f,								// LFE Out
+			0.0f,							0.0f,							0.0f,							0.0f,							Chans[CHANNELOUT_LEFTSURROUND],		0.0f,								// Left Surround Out
+			0.0f,							0.0f,							0.0f,							0.0f,							0.0f,								Chans[CHANNELOUT_RIGHTSURROUND]		// Right Surround Out
 		};
 
 		// 5.1 sounds map 6 channels to 6 speakers
@@ -908,7 +1031,7 @@ void FXAudio2SoundSource::RouteHexToDry(float ChannelVolumes[CHANNELOUT_COUNT])
 /** 
  * Maps the sound to the relevant reverb effect
  */
-void FXAudio2SoundSource::RouteToReverb( float ChannelVolumes[CHANNELOUT_COUNT] )
+void FXAudio2SoundSource::RouteToReverb(float ChannelVolumes[CHANNEL_MATRIX_COUNT])
 {
 	// Reverb must be applied to process this function because the index 
 	// of the destination output voice may not be at DEST_REVERB.
@@ -926,9 +1049,13 @@ void FXAudio2SoundSource::RouteToReverb( float ChannelVolumes[CHANNELOUT_COUNT]
 	}
 }
 
-void FXAudio2SoundSource::RouteMonoToReverb(float ChannelVolumes[CHANNELOUT_COUNT])
+void FXAudio2SoundSource::RouteMonoToReverb(float ChannelVolumes[CHANNEL_MATRIX_COUNT])
 {
-	if (IsUsingDefaultSpatializer())
+	if (IsUsingHrtfSpatializer())
+	{
+		RouteStereoToReverb(ChannelVolumes);
+	}
+	else
 	{
 		float SpatialisationMatrix[2] =
 		{
@@ -940,13 +1067,9 @@ void FXAudio2SoundSource::RouteMonoToReverb(float ChannelVolumes[CHANNELOUT_COUN
 		AudioDevice->ValidateAPICall(TEXT("SetOutputMatrix (Mono reverb)"),
 			Source->SetOutputMatrix(Destinations[DEST_REVERB].pOutputVoice, 1, 2, SpatialisationMatrix));
 	}
-	else
-	{
-		RouteStereoToReverb(ChannelVolumes);
-	}
 }
 
-void FXAudio2SoundSource::RouteStereoToReverb(float ChannelVolumes[CHANNELOUT_COUNT])
+void FXAudio2SoundSource::RouteStereoToReverb(float ChannelVolumes[CHANNEL_MATRIX_COUNT])
 {
 	float SpatialisationMatrix[4] =
 	{
@@ -966,7 +1089,7 @@ void FXAudio2SoundSource::RouteStereoToReverb(float ChannelVolumes[CHANNELOUT_CO
  * @param	ChannelVolumes	The volumes associated to each channel. 
  *							Note: Not all channels are mapped directly to a speaker.
  */
-void FXAudio2SoundSource::RouteToRadio( float ChannelVolumes[CHANNELOUT_COUNT] )
+void FXAudio2SoundSource::RouteToRadio(float ChannelVolumes[CHANNEL_MATRIX_COUNT])
 {
 	// Radio distortion must be applied to process this function because 
 	// the index of the destination output voice would be incorrect. 
@@ -1092,16 +1215,35 @@ FString FXAudio2SoundSource::Describe_Internal(bool bUseLongName, bool bIncludeC
 	FString SpatializedVolumeInfo;
 	if (bIncludeChannelVolumes && WaveInstance->bUseSpatialization)
 	{
-		float ChannelVolumes[CHANNELOUT_COUNT] = { 0.0f };
+		float ChannelVolumes[CHANNEL_MATRIX_COUNT] = { 0.0f };
 		GetChannelVolumes( ChannelVolumes, WaveInstance->GetActualVolume() );
 
-		SpatializedVolumeInfo = FString::Printf(TEXT(" (FL: %.2f FR: %.2f FC: %.2f LF: %.2f, LS: %.2f, RS: %.2f)"),
-			ChannelVolumes[CHANNELOUT_FRONTLEFT],
-			ChannelVolumes[CHANNELOUT_FRONTRIGHT],
-			ChannelVolumes[CHANNELOUT_FRONTCENTER],
-			ChannelVolumes[CHANNELOUT_LOWFREQUENCY],
-			ChannelVolumes[CHANNELOUT_LEFTSURROUND],
-			ChannelVolumes[CHANNELOUT_RIGHTSURROUND]);
+		if (Buffer->NumChannels == 1)
+		{
+			SpatializedVolumeInfo = FString::Printf(TEXT(" (FL: %.2f FR: %.2f FC: %.2f LF: %.2f, LS: %.2f, RS: %.2f)"),
+													ChannelVolumes[CHANNELOUT_FRONTLEFT],
+													ChannelVolumes[CHANNELOUT_FRONTRIGHT],
+													ChannelVolumes[CHANNELOUT_FRONTCENTER],
+													ChannelVolumes[CHANNELOUT_LOWFREQUENCY],
+													ChannelVolumes[CHANNELOUT_LEFTSURROUND],
+													ChannelVolumes[CHANNELOUT_RIGHTSURROUND]);
+		}
+		else if (Buffer->NumChannels == 2)
+		{
+			SpatializedVolumeInfo = FString::Printf(TEXT(" Left: (FL: %.2f FR: %.2f FC: %.2f LF: %.2f, LS: %.2f, RS: %.2f), Right: (FL: %.2f FR: %.2f FC: %.2f LF: %.2f, LS: %.2f, RS: %.2f)"),
+													ChannelVolumes[CHANNELOUT_FRONTLEFT],		
+													ChannelVolumes[CHANNELOUT_FRONTRIGHT],		
+													ChannelVolumes[CHANNELOUT_FRONTCENTER],		
+													ChannelVolumes[CHANNELOUT_LOWFREQUENCY],	
+													ChannelVolumes[CHANNELOUT_LEFTSURROUND],	
+													ChannelVolumes[CHANNELOUT_RIGHTSURROUND],	
+													ChannelVolumes[CHANNELOUT_FRONTLEFT + CHANNELOUT_COUNT],
+													ChannelVolumes[CHANNELOUT_FRONTRIGHT + CHANNELOUT_COUNT],
+													ChannelVolumes[CHANNELOUT_FRONTCENTER + CHANNELOUT_COUNT],
+													ChannelVolumes[CHANNELOUT_LOWFREQUENCY + CHANNELOUT_COUNT],
+													ChannelVolumes[CHANNELOUT_LEFTSURROUND + CHANNELOUT_COUNT],
+													ChannelVolumes[CHANNELOUT_RIGHTSURROUND + CHANNELOUT_COUNT]);
+		}
 	}
 
 	return FString::Printf(TEXT("Wave: %s, Volume: %6.2f%s, Owner: %s"), 
@@ -1151,7 +1293,7 @@ void FXAudio2SoundSource::Update( void )
 		Source->SetFilterParameters( &LPFParameters ) );	
 
 	// Initialize channel volumes
-	float ChannelVolumes[CHANNELOUT_COUNT] = { 0.0f };
+	float ChannelVolumes[CHANNEL_MATRIX_COUNT] = { 0.0f };
 
 	GetChannelVolumes( ChannelVolumes, WaveInstance->GetActualVolume() );
 
@@ -1170,8 +1312,38 @@ void FXAudio2SoundSource::Update( void )
 	{
 		RouteToRadio( ChannelVolumes );
 	}
+
+	// Draw 3d Debug information about this source, if enabled
+	// TODO - Audio Threading. This won't work cross thread.
+	FAudioDeviceManager* DeviceManager = GEngine->GetAudioDeviceManager();
+
+	if (DeviceManager && DeviceManager->IsVisualizeDebug3dEnabled())
+	{
+		if (WaveInstance->ActiveSound && WaveInstance->ActiveSound->AudioComponent.IsValid())
+		{
+			TWeakObjectPtr<class UAudioComponent>& AudioComponent = WaveInstance->ActiveSound->AudioComponent;
+			UWorld* SoundWorld = AudioComponent->GetWorld();
+			if (SoundWorld)
+			{
+				FRotator SoundRotation = WaveInstance->ActiveSound->AudioComponent->GetComponentRotation();
+				DrawDebugCrosshairs(SoundWorld, WaveInstance->Location, SoundRotation, 20.0f, FColor::White, false, -1.0f, SDPG_Foreground);
+
+				FString Name;
+				WaveInstance->ActiveSound->Sound->GetName(Name);
+
+				if (Buffer->NumChannels == 2 && WaveInstance->bUseSpatialization)
+				{
+					DrawDebugCrosshairs(SoundWorld, LeftChannelSourceLocation, SoundRotation, 20.0f, FColor::Red, false, -1.0f, SDPG_Foreground);
+					DrawDebugCrosshairs(SoundWorld, RightChannelSourceLocation, SoundRotation, 20.0f, FColor::Green, false, -1.0f, SDPG_Foreground);
+				}
+
+				DrawDebugString(SoundWorld, AudioComponent->GetComponentLocation() + FVector(0, 0, 32), *Name, nullptr, FColor::White, 0.033, false);
+			}
+		}
+	}
 }
 
+
 /**
  * Plays the current wave instance.	
  */
@@ -1398,9 +1570,9 @@ bool FXAudio2SoundSource::IsFinished( void )
 	return( true );
 }
 
-bool FXAudio2SoundSource::IsUsingDefaultSpatializer()
+bool FXAudio2SoundSource::IsUsingHrtfSpatializer()
 {
-	return bUsingDefaultSpatialization;
+	return bUsingHRTFSpatialization;
 }
 
 bool FXAudio2SoundSource::CreateWithSpatializationEffect()
diff --git a/Engine/Source/Runtime/Windows/XAudio2/Private/XAudio2Support.h b/Engine/Source/Runtime/Windows/XAudio2/Private/XAudio2Support.h
--- a/Engine/Source/Runtime/Windows/XAudio2/Private/XAudio2Support.h
+++ b/Engine/Source/Runtime/Windows/XAudio2/Private/XAudio2Support.h
@@ -365,7 +365,7 @@ public:
 	/**
 	 * Calculates the volume for each channel
 	 */
-	void GetChannelVolumes( float ChannelVolumes[CHANNELOUT_COUNT], float AttenuatedVolume );
+	void GetChannelVolumes(float ChannelVolumes[CHANNEL_MATRIX_COUNT], float AttenuatedVolume);
 
 	/**
 	 * Returns a string describing the source
@@ -380,12 +380,12 @@ public:
 	/** 
 	 * Maps a sound with a given number of channels to to expected speakers
 	 */
-	void RouteDryToSpeakers( float ChannelVolumes[CHANNELOUT_COUNT] );
+	void RouteDryToSpeakers(float ChannelVolumes[CHANNEL_MATRIX_COUNT]);
 
 	/** 
 	 * Maps the sound to the relevant reverb effect
 	 */
-	void RouteToReverb( float ChannelVolumes[CHANNELOUT_COUNT] );
+	void RouteToReverb(float ChannelVolumes[CHANNEL_MATRIX_COUNT]);
 
 	/** 
 	 * Maps the sound to the relevant radio effect.
@@ -393,7 +393,7 @@ public:
 	 * @param	ChannelVolumes	The volumes associated to each channel. 
 	 *							Note: Not all channels are mapped directly to a speaker.
 	 */
-	void RouteToRadio( float ChannelVolumes[CHANNELOUT_COUNT] );
+	void RouteToRadio(float ChannelVolumes[CHANNEL_MATRIX_COUNT]);
 
 protected:
 
@@ -408,7 +408,7 @@ protected:
 	bool ReadMorePCMData(const int32 BufferIndex, EDataReadMode DataReadMode);
 
 	/** Returns if the source is using the default 3d spatialization. */
-	bool IsUsingDefaultSpatializer();
+	bool IsUsingHrtfSpatializer();
 
 	/** Returns Whether or not to create this source with the 3d spatialization effect. */
 	bool CreateWithSpatializationEffect();
@@ -423,23 +423,54 @@ protected:
 	 */
 	int32 GetDestinationVoiceIndexForEffect( SourceDestinations Effect );
 
+	/**
+	* FSpatializationParams
+	* Struct for retrieving parameters needed for computing 3d spatialization
+	*/
+	struct FSpatializationParams
+	{
+		FVector ListenerPosition;
+		FVector ListenerOrientation;
+		FVector EmitterPosition;
+		FVector LeftChannelPosition;
+		FVector RightChannelPosition;
+		float Distance;
+		float NormalizedOmniRadius;
+	};
+
+	/**
+	* Converts a vector orientation from UE4 coordinates to XAudio2 coordinates
+	*/
+	inline FVector ConvertToXAudio2Orientation(const FVector& InputVector);
+
+	/**
+	* Transform the given position in absolute world-space to listener-relative space. Optionally returns distance.
+	*/
+	FVector GetListenerTransformedDirection(const FVector& Position, float* OutDistance = nullptr);
+
+	/**
+	* Gets parameters necessary for computing 3d spatialization of sources
+	*/
+	void GetSpatializationParams(FSpatializationParams& Params);
+	void UpdateStereoEmitterPositions();
+
 	/**
 	* Calculates the channel volumes for various input channel configurations.
 	*/
-	void GetMonoChannelVolumes(float ChannelVolumes[CHANNELOUT_COUNT], float AttenuatedVolume);
-	void GetStereoChannelVolumes(float ChannelVolumes[CHANNELOUT_COUNT], float AttenuatedVolume);
-	void GetQuadChannelVolumes(float ChannelVolumes[CHANNELOUT_COUNT], float AttenuatedVolume);
-	void GetHexChannelVolumes(float ChannelVolumes[CHANNELOUT_COUNT], float AttenuatedVolume);
+	void GetMonoChannelVolumes(float ChannelVolumes[CHANNEL_MATRIX_COUNT], float AttenuatedVolume);
+	void GetStereoChannelVolumes(float ChannelVolumes[CHANNEL_MATRIX_COUNT], float AttenuatedVolume);
+	void GetQuadChannelVolumes(float ChannelVolumes[CHANNEL_MATRIX_COUNT], float AttenuatedVolume);
+	void GetHexChannelVolumes(float ChannelVolumes[CHANNEL_MATRIX_COUNT], float AttenuatedVolume);
 
 	/**
 	* Routes channel sends for various input channel configurations.
 	*/
-	void RouteMonoToDry(float ChannelVolumes[CHANNELOUT_COUNT]);
-	void RouteStereoToDry(float ChannelVolumes[CHANNELOUT_COUNT]);
-	void RouteQuadToDry(float ChannelVolumes[CHANNELOUT_COUNT]);
-	void RouteHexToDry(float ChannelVolumes[CHANNELOUT_COUNT]);
-	void RouteMonoToReverb(float ChannelVolumes[CHANNELOUT_COUNT]);
-	void RouteStereoToReverb(float ChannelVolumes[CHANNELOUT_COUNT]);
+	void RouteMonoToDry(float ChannelVolumes[CHANNEL_MATRIX_COUNT]);
+	void RouteStereoToDry(float ChannelVolumes[CHANNEL_MATRIX_COUNT]);
+	void RouteQuadToDry(float ChannelVolumes[CHANNEL_MATRIX_COUNT]);
+	void RouteHexToDry(float ChannelVolumes[CHANNEL_MATRIX_COUNT]);
+	void RouteMonoToReverb(float ChannelVolumes[CHANNEL_MATRIX_COUNT]);
+	void RouteStereoToReverb(float ChannelVolumes[CHANNEL_MATRIX_COUNT]);
 
 	/** Owning classes */
 	FXAudio2Device*				AudioDevice;
@@ -470,8 +501,12 @@ protected:
 	uint32						bResourcesNeedFreeing:1;
 	/** Index of this sound source in the audio device sound source array. */
 	uint32						VoiceId;
-	/** Whether or not this sound is spatializing using default spatialization algorithm. */
-	bool						bUsingDefaultSpatialization;
+	/** Whether or not this sound is spatializing using an HRTF spatialization algorithm. */
+	bool						bUsingHRTFSpatialization;
+	/** The location of the left-channel source for stereo spatialization. */
+	FVector						LeftChannelSourceLocation;
+	/** The location of the right-channel source for stereo spatialization. */
+	FVector						RightChannelSourceLocation;
 
 	friend class FXAudio2Device;
 	friend class FXAudio2SoundSourceCallback;
diff --git a/Engine/Source/Runtime/Windows/XAudio2/Public/XAudio2Device.h b/Engine/Source/Runtime/Windows/XAudio2/Public/XAudio2Device.h
--- a/Engine/Source/Runtime/Windows/XAudio2/Public/XAudio2Device.h
+++ b/Engine/Source/Runtime/Windows/XAudio2/Public/XAudio2Device.h
@@ -52,6 +52,12 @@ enum ChannelOutputs
 	CHANNELOUT_COUNT
 };
 
+// We are supporting spatializing non-mono assets, so our channel output will be a "matrix" 
+// of channel-out mappings per input channel only supporting stereo 3d spatialization, but 
+// if we support 3d spatialization of other N channel source files, we'd bump this up
+#define MAX_INPUT_CHANNELS_SPATIALIZED (2)
+#define CHANNEL_MATRIX_COUNT (MAX_INPUT_CHANNELS_SPATIALIZED*CHANNELOUT_COUNT)
+
 #define SPEAKER_5POINT0          ( SPEAKER_FRONT_LEFT | SPEAKER_FRONT_RIGHT | SPEAKER_FRONT_CENTER | SPEAKER_SIDE_LEFT | SPEAKER_SIDE_RIGHT )
 #define SPEAKER_6POINT1          ( SPEAKER_FRONT_LEFT | SPEAKER_FRONT_RIGHT | SPEAKER_FRONT_CENTER | SPEAKER_LOW_FREQUENCY | SPEAKER_SIDE_LEFT | SPEAKER_SIDE_RIGHT | SPEAKER_BACK_CENTER )
 
@@ -162,10 +168,10 @@ protected:
 	void TimeTest( FOutputDevice& Ar, const TCHAR* WaveAssetName );
 
 	/** Inverse listener transformation, used for spatialization */
-	FMatrix								InverseTransform;
+	FMatrix InverseListenerTransform;
 
 	// For calculating spatialised volumes
-	static FSpatializationHelper		SpatializationHelper;
+	static FSpatializationHelper SpatializationHelper;
 
 	friend class FXAudio2SoundBuffer;
 	friend class FXAudio2SoundSource;
