UE-20133 Fixing async destroy voice crash

- Moving the async worker thread declaration and implementation to be inside XAudio2Device.cpp
- Added a threadsafe counter to track the number of in-flight async destroy-voice calls per audio device
- Blocked the call to tear-down the audio device until all in-flight destroy-voice calls occur.
- Should rarely actualy block but this will prevent Xaudio2 from shutting down before an XAudio2 source get "Destroy" called on it. If a block does occur on shutdown, will only be for a neglible time.
- Wasn't able to repro after rigorous and prolonged audio start/stop spammage in the blueprint editor.

diff --git a/Engine/Source/Runtime/Windows/XAudio2/Private/XAudio2Device.cpp b/Engine/Source/Runtime/Windows/XAudio2/Private/XAudio2Device.cpp
--- a/Engine/Source/Runtime/Windows/XAudio2/Private/XAudio2Device.cpp
+++ b/Engine/Source/Runtime/Windows/XAudio2/Private/XAudio2Device.cpp
@@ -204,6 +204,12 @@ void FXAudio2Device::TeardownHardware()
 {
 	if (DeviceProperties)
 	{
+		// Block device shutdown until all voices have been destroyed.
+		while (NumSourcesDestroying.GetValue() > 0)
+		{
+			// spin
+		}
+
 		// close hardware interfaces
 		if (DeviceProperties->MasteringVoice)
 		{
@@ -640,3 +646,41 @@ void* FXAudio2Device::AllocatePermanentMemory( int32 Size, bool& AllocatedInPool
 	
 	return( Allocation );
 }
+
+void FXAudio2Device::AsyncDestroyXAudio2Source(IXAudio2SourceVoice* Source)
+{
+	// Publish that there is a voice in-flight to be asynchronously destroyed.
+	NumSourcesDestroying.Increment();
+
+	(new FAutoDeleteAsyncTask<FAsyncXAudio2SourceDestroyer>(this, Source))->StartBackgroundTask();
+}
+
+
+/*------------------------------------------------------------------------------------
+FAsyncXAudio2SourceDestroyer
+------------------------------------------------------------------------------------*/
+
+FAsyncXAudio2SourceDestroyer::FAsyncXAudio2SourceDestroyer(FXAudio2Device* InAudioDevice, IXAudio2SourceVoice* InSourceVoice)
+	: SourceVoice(InSourceVoice)
+	, AudioDevice(InAudioDevice)
+{
+	check(AudioDevice != nullptr);
+}
+
+FAsyncXAudio2SourceDestroyer::~FAsyncXAudio2SourceDestroyer()
+{
+	check(SourceVoice == nullptr);
+}
+
+void FAsyncXAudio2SourceDestroyer::DoWork()
+{
+	check(SourceVoice);
+	SourceVoice->DestroyVoice();
+	SourceVoice = nullptr;
+
+	// Publish that this voice has been destroyed
+	check(AudioDevice != nullptr);
+	check(AudioDevice->NumSourcesDestroying.GetValue() > 0);
+	AudioDevice->NumSourcesDestroying.Decrement();
+}
+
diff --git a/Engine/Source/Runtime/Windows/XAudio2/Private/XAudio2Source.cpp b/Engine/Source/Runtime/Windows/XAudio2/Private/XAudio2Source.cpp
--- a/Engine/Source/Runtime/Windows/XAudio2/Private/XAudio2Source.cpp
+++ b/Engine/Source/Runtime/Windows/XAudio2/Private/XAudio2Source.cpp
@@ -17,38 +17,6 @@
 #include "XAudio2Support.h"
 #include "IAudioExtensionPlugin.h"
 
-/*------------------------------------------------------------------------------------
-	Simple async task to destroy an xaudio2 source voice without blocking the main thread.
-------------------------------------------------------------------------------------*/
-
-class FAsyncXAudio2SourceDestroyer : public FNonAbandonableTask
-{
-protected:
-	IXAudio2SourceVoice* SourceVoice;
-
-public:
-	FAsyncXAudio2SourceDestroyer(IXAudio2SourceVoice* InSourceVoice)
-		: SourceVoice(InSourceVoice)
-	{
-	}
-
-	~FAsyncXAudio2SourceDestroyer()
-	{
-		check(SourceVoice == nullptr);
-	}
-
-	void DoWork()
-	{
-		check(SourceVoice);
-		SourceVoice->DestroyVoice();
-		SourceVoice = nullptr;
-	}
-
-	FORCEINLINE TStatId GetStatId() const
-	{
-		RETURN_QUICK_DECLARE_CYCLE_STAT(FAsyncXAudio2SourceDestroyer, STATGROUP_ThreadPoolAsyncTasks);
-	}
-};
 
 /*------------------------------------------------------------------------------------
 	For muting user soundtracks during cinematics
@@ -129,7 +97,7 @@ void FXAudio2SoundSource::FreeResources( void )
 	{
 		// Because XAudio2's DestroyVoice source is blocking and can be slow on some processors (e.g. AMD), we're creating
 		// a task that destroys the voice on a separate thread to avoid blocking or hitching.
-		(new FAutoDeleteAsyncTask<FAsyncXAudio2SourceDestroyer>(Source))->StartBackgroundTask();
+		AudioDevice->AsyncDestroyXAudio2Source(Source);
 		Source = nullptr;
 	}
 
diff --git a/Engine/Source/Runtime/Windows/XAudio2/Public/XAudio2Device.h b/Engine/Source/Runtime/Windows/XAudio2/Public/XAudio2Device.h
--- a/Engine/Source/Runtime/Windows/XAudio2/Public/XAudio2Device.h
+++ b/Engine/Source/Runtime/Windows/XAudio2/Public/XAudio2Device.h
@@ -74,6 +74,28 @@ enum ESoundFormat
 	SoundFormat_Streaming
 };
 
+
+/** 
+* Simple async task to destroy an xaudio2 source voice without blocking the main thread.
+*/
+class FAsyncXAudio2SourceDestroyer : public FNonAbandonableTask
+{
+protected:
+	struct IXAudio2SourceVoice* SourceVoice;
+	FXAudio2Device* AudioDevice;
+
+public:
+	FAsyncXAudio2SourceDestroyer(FXAudio2Device* InAudioDevice, struct IXAudio2SourceVoice* InSourceVoice);
+	~FAsyncXAudio2SourceDestroyer();
+	void DoWork();
+
+	FORCEINLINE TStatId GetStatId() const
+	{
+		RETURN_QUICK_DECLARE_CYCLE_STAT(FAsyncXAudio2SourceDestroyer, STATGROUP_ThreadPoolAsyncTasks);
+	}
+};
+
+
 class FXAudio2SoundBuffer;
 class FXAudio2SoundSource;
 class FSpatializationHelper;
@@ -144,6 +166,12 @@ class FXAudio2Device : public FAudioDevice
 	 */
 	virtual bool Exec( UWorld* InWorld, const TCHAR* Cmd, FOutputDevice& Ar ) override;
 
+public:
+	/** 
+	 * Safely destroy an XAudio2 source asynchronously.
+	*/
+	void AsyncDestroyXAudio2Source(struct IXAudio2SourceVoice* Source);
+
 protected:
 
 	/**
@@ -176,9 +204,11 @@ protected:
 	friend class FXAudio2SoundBuffer;
 	friend class FXAudio2SoundSource;
 	friend class FXAudio2EffectsManager;
+	friend class FAsyncXAudio2SourceDestroyer;
 
 private:
 	struct FXAudioDeviceProperties* DeviceProperties;
+	FThreadSafeCounter NumSourcesDestroying;
 
 #if PLATFORM_WINDOWS
 	// We need to keep track whether com was successfully initialized so we can clean 
