Fix AudioComponent direct accesses

diff --git a/Engine/Source/Runtime/Windows/XAudio2/Private/XAudio2Source.cpp b/Engine/Source/Runtime/Windows/XAudio2/Private/XAudio2Source.cpp
--- a/Engine/Source/Runtime/Windows/XAudio2/Private/XAudio2Source.cpp
+++ b/Engine/Source/Runtime/Windows/XAudio2/Private/XAudio2Source.cpp
@@ -1297,13 +1297,13 @@ void FXAudio2SoundSource::Update( void )
 
 	if (DeviceManager && DeviceManager->IsVisualizeDebug3dEnabled())
 	{
-		if (WaveInstance->ActiveSound && WaveInstance->ActiveSound->AudioComponent.IsValid())
+		UAudioComponent* AudioComponent = (WaveInstance->ActiveSound ? WaveInstance->ActiveSound->GetAudioComponent() : nullptr);
+		if (AudioComponent)
 		{
-			TWeakObjectPtr<class UAudioComponent>& AudioComponent = WaveInstance->ActiveSound->AudioComponent;
 			UWorld* SoundWorld = AudioComponent->GetWorld();
 			if (SoundWorld)
 			{
-				FRotator SoundRotation = WaveInstance->ActiveSound->AudioComponent->GetComponentRotation();
+				FRotator SoundRotation = AudioComponent->GetComponentRotation();
 				DrawDebugCrosshairs(SoundWorld, WaveInstance->Location, SoundRotation, 20.0f, FColor::White, false, -1.0f, SDPG_Foreground);
 
 				FString Name;
