Enable GRHISupportsMSAADepthSampleAccess  for additional platforms which support it.

diff --git a/Engine/Source/Runtime/Windows/D3D12RHI/Private/D3D12Device.cpp b/Engine/Source/Runtime/Windows/D3D12RHI/Private/D3D12Device.cpp
--- a/Engine/Source/Runtime/Windows/D3D12RHI/Private/D3D12Device.cpp
+++ b/Engine/Source/Runtime/Windows/D3D12RHI/Private/D3D12Device.cpp
@@ -390,6 +390,7 @@ FD3D12DynamicRHI::FD3D12DynamicRHI(IDXGIFactory4* InDXGIFactory, FD3D12Adapter&
 	GMaxTextureDimensions = D3D12_REQ_TEXTURE2D_U_OR_V_DIMENSION;
 	GMaxCubeTextureDimensions = D3D12_REQ_TEXTURECUBE_DIMENSION;
 	GMaxTextureArrayLayers = D3D12_REQ_TEXTURE2D_ARRAY_AXIS_DIMENSION;
+	GRHISupportsMSAADepthSampleAccess = true;
 
 	GMaxShadowDepthBufferSizeX = 4096;
 	GMaxShadowDepthBufferSizeY = 4096;
