Fixed: Crash on Mobile Previewing or Playing #jira UE-20639
FCaptureRenderTarget does not allocate target texture

diff --git a/Engine/Source/Runtime/Windows/D3D11RHI/Private/D3D11Commands.cpp b/Engine/Source/Runtime/Windows/D3D11RHI/Private/D3D11Commands.cpp
--- a/Engine/Source/Runtime/Windows/D3D11RHI/Private/D3D11Commands.cpp
+++ b/Engine/Source/Runtime/Windows/D3D11RHI/Private/D3D11Commands.cpp
@@ -974,9 +974,12 @@ void FD3D11DynamicRHI::RHISetRenderTargetsAndClear(const FRHISetRenderTargetsInf
 		{
 			for (int32 i = 0; i < RenderTargetsInfo.NumColorRenderTargets; ++i)
 			{
-				const FClearValueBinding& ClearValue = RenderTargetsInfo.ColorRenderTarget[i].Texture->GetClearBinding();
-				checkf(ClearValue.ColorBinding == EClearBinding::EColorBound, TEXT("Texture: %s does not have a color bound for fast clears"), *RenderTargetsInfo.ColorRenderTarget[i].Texture->GetName().GetPlainNameString());
-				ClearColors[i] = ClearValue.GetClearColor();
+				if (RenderTargetsInfo.ColorRenderTarget[i].Texture != nullptr)
+				{
+					const FClearValueBinding& ClearValue = RenderTargetsInfo.ColorRenderTarget[i].Texture->GetClearBinding();
+					checkf(ClearValue.ColorBinding == EClearBinding::EColorBound, TEXT("Texture: %s does not have a color bound for fast clears"), *RenderTargetsInfo.ColorRenderTarget[i].Texture->GetName().GetPlainNameString());
+					ClearColors[i] = ClearValue.GetClearColor();
+				}
 			}
 		}
 		if (RenderTargetsInfo.bClearDepth || RenderTargetsInfo.bClearStencil)
