Implenting a voice pool for XAudio2 voices

- Removing the async DestroyVoice code
- Adding format-based pools
- Created a separate pool for voices that use the spatialization effect (splitting mono-to-stereo) since xaudio2 defaults to a max per-voice effect output channel count equal to the input channel count unless created with an effect chain. Therefore xaudio2 voices that need to use the mono-to-stereo effect need to be created with it.
- Put new pools and related code in FXAudioDeviceProperties since it's the struct we use for xaudio2-specific handles to resources, e.g. IXAudio2, IXAudio2MasteringVoice, etc.
- Added pool shutdown code and moved other xaudio2 shutdown code into the destructor of FXAudioDeviceProperties

diff --git a/Engine/Source/Runtime/Windows/XAudio2/Private/XAudio2Device.cpp b/Engine/Source/Runtime/Windows/XAudio2/Private/XAudio2Device.cpp
--- a/Engine/Source/Runtime/Windows/XAudio2/Private/XAudio2Device.cpp
+++ b/Engine/Source/Runtime/Windows/XAudio2/Private/XAudio2Device.cpp
@@ -204,36 +204,6 @@ void FXAudio2Device::TeardownHardware()
 {
 	if (DeviceProperties)
 	{
-		// Block device shutdown until all voices have been destroyed.
-		while (NumSourcesDestroying.GetValue() > 0)
-		{
-			// spin
-		}
-
-		// close hardware interfaces
-		if (DeviceProperties->MasteringVoice)
-		{
-			DeviceProperties->MasteringVoice->DestroyVoice();
-			DeviceProperties->MasteringVoice = NULL;
-		}
-
-		if (DeviceProperties->XAudio2)
-		{
-			// Force the hardware to release all references
-			DeviceProperties->XAudio2->Release();
-			DeviceProperties->XAudio2 = NULL;
-		}
-
-#if PLATFORM_WINDOWS && PLATFORM_64BITS
-		if (DeviceProperties->XAudio2Dll)
-		{
-			if (!FreeLibrary(DeviceProperties->XAudio2Dll))
-			{
-				UE_LOG(LogAudio, Warning, TEXT("Failed to free XAudio2 Dll"));
-			}
-		}
-#endif
-
 		delete DeviceProperties;
 		DeviceProperties = nullptr;
 	}
@@ -646,41 +616,3 @@ void* FXAudio2Device::AllocatePermanentMemory( int32 Size, bool& AllocatedInPool
 	
 	return( Allocation );
 }
-
-void FXAudio2Device::AsyncDestroyXAudio2Source(IXAudio2SourceVoice* Source)
-{
-	// Publish that there is a voice in-flight to be asynchronously destroyed.
-	NumSourcesDestroying.Increment();
-
-	(new FAutoDeleteAsyncTask<FAsyncXAudio2SourceDestroyer>(this, Source))->StartBackgroundTask();
-}
-
-
-/*------------------------------------------------------------------------------------
-FAsyncXAudio2SourceDestroyer
-------------------------------------------------------------------------------------*/
-
-FAsyncXAudio2SourceDestroyer::FAsyncXAudio2SourceDestroyer(FXAudio2Device* InAudioDevice, IXAudio2SourceVoice* InSourceVoice)
-	: SourceVoice(InSourceVoice)
-	, AudioDevice(InAudioDevice)
-{
-	check(AudioDevice != nullptr);
-}
-
-FAsyncXAudio2SourceDestroyer::~FAsyncXAudio2SourceDestroyer()
-{
-	check(SourceVoice == nullptr);
-}
-
-void FAsyncXAudio2SourceDestroyer::DoWork()
-{
-	check(SourceVoice);
-	SourceVoice->DestroyVoice();
-	SourceVoice = nullptr;
-
-	// Publish that this voice has been destroyed
-	check(AudioDevice != nullptr);
-	check(AudioDevice->NumSourcesDestroying.GetValue() > 0);
-	AudioDevice->NumSourcesDestroying.Decrement();
-}
-
diff --git a/Engine/Source/Runtime/Windows/XAudio2/Private/XAudio2Source.cpp b/Engine/Source/Runtime/Windows/XAudio2/Private/XAudio2Source.cpp
--- a/Engine/Source/Runtime/Windows/XAudio2/Private/XAudio2Source.cpp
+++ b/Engine/Source/Runtime/Windows/XAudio2/Private/XAudio2Source.cpp
@@ -37,6 +37,7 @@ FXMPHelper* FXMPHelper::GetXMPHelper( void )
 FXAudio2SoundSource::FXAudio2SoundSource(FAudioDevice* InAudioDevice)
 	: FSoundSource( InAudioDevice )
 	, Source( NULL )
+	, MaxEffectChainChannels(0)
 	, RealtimeAsyncTask( nullptr )
 	, CurrentBuffer( 0 )
 	, bBuffersToFlush( false )
@@ -93,11 +94,11 @@ void FXAudio2SoundSource::FreeResources( void )
 	}
 
 	// Release voice.
-	if( Source )
+	if (Source)
 	{
 		// Because XAudio2's DestroyVoice source is blocking and can be slow on some processors (e.g. AMD), we're creating
 		// a task that destroys the voice on a separate thread to avoid blocking or hitching.
-		AudioDevice->AsyncDestroyXAudio2Source(Source);
+		AudioDevice->DeviceProperties->ReleaseSourceVoice(Source, XAudio2Buffer->PCM, MaxEffectChainChannels);
 		Source = nullptr;
 	}
 
@@ -364,13 +365,6 @@ bool FXAudio2SoundSource::CreateSource( void )
 #endif	//XAUDIO2_SUPPORTS_SENDLIST
 
 	// Mark the source as music if it is a member of the music group and allow low, band and high pass filters
-	UINT32 Flags = 
-#if XAUDIO2_SUPPORTS_MUSIC
-		( WaveInstance->bIsMusic ? XAUDIO2_VOICE_MUSIC : 0 );
-#else	//XAUDIO2_SUPPORTS_MUSIC
-		0;
-#endif	//XAUDIO2_SUPPORTS_MUSIC
-	Flags |= XAUDIO2_VOICE_USEFILTER;
 
 	// Reset the bUsingSpatializationEffect flag
 	bUsingHRTFSpatialization = false;
@@ -387,24 +381,25 @@ bool FXAudio2SoundSource::CreateSource( void )
 			// going to hear the sound for the duration of this sound.
 			bUsingHRTFSpatialization = true;
 
+			MaxEffectChainChannels = 2;
+
 			XAUDIO2_EFFECT_DESCRIPTOR EffectDescriptor[1];
 			EffectDescriptor[0].pEffect = Effect;
-			EffectDescriptor[0].OutputChannels = 2;
+			EffectDescriptor[0].OutputChannels = MaxEffectChainChannels;
 			EffectDescriptor[0].InitialState = true;
 
 			const XAUDIO2_EFFECT_CHAIN EffectChain = { 1, EffectDescriptor };
+			AudioDevice->DeviceProperties->GetFreeSourceVoice(&Source, XAudio2Buffer->PCM, &EffectChain, MaxEffectChainChannels);
 
-			// All sound formats start with the WAVEFORMATEX structure (PCM, XMA2, XWMA)
-			if (!AudioDevice->ValidateAPICall(TEXT("CreateSourceVoice (mono source)"),
-#if XAUDIO2_SUPPORTS_SENDLIST
-				AudioDevice->DeviceProperties->XAudio2->CreateSourceVoice(&Source, (WAVEFORMATEX*)&XAudio2Buffer->PCM, Flags, MAX_PITCH, &SourceCallback, &SourceSendList, &EffectChain)))
-#else	//XAUDIO2_SUPPORTS_SENDLIST
-				AudioDevice->DeviceProperties->XAudio2->CreateSourceVoice(&Source, (WAVEFORMATEX*)&XAudio2Buffer->PCM, Flags, MAX_PITCH, &SourceCallback, nullptr, &EffectChain)))
-#endif	//XAUDIO2_SUPPORTS_SENDLIST
+			if (!Source)
 			{
-				return(false);
+				return false;
 			}
 
+#if XAUDIO2_SUPPORTS_SENDLIST
+			Source->SetOutputVoices(&SourceSendList);
+#endif
+
 			// We succeeded, then return
 			bCreatedWithSpatializationEffect = true;
 		}
@@ -414,16 +409,18 @@ bool FXAudio2SoundSource::CreateSource( void )
 	{
 		check(AudioDevice->DeviceProperties != nullptr);
 		check(AudioDevice->DeviceProperties->XAudio2 != nullptr);
-		// All sound formats start with the WAVEFORMATEX structure (PCM, XMA2, XWMA)
-		if (!AudioDevice->ValidateAPICall(TEXT("CreateSourceVoice (source)"),
-#if XAUDIO2_SUPPORTS_SENDLIST
-			AudioDevice->DeviceProperties->XAudio2->CreateSourceVoice(&Source, (WAVEFORMATEX*)&XAudio2Buffer->PCM, Flags, MAX_PITCH, &SourceCallback, &SourceSendList, nullptr)))
-#else	//XAUDIO2_SUPPORTS_SENDLIST
-			AudioDevice->DeviceProperties->XAudio2->CreateSourceVoice(&Source, (WAVEFORMATEX*)&XAudio2Buffer->PCM, Flags, MAX_PITCH, &SourceCallback)))
-#endif	//XAUDIO2_SUPPORTS_SENDLIST
+
+		AudioDevice->DeviceProperties->GetFreeSourceVoice(&Source, XAudio2Buffer->PCM, nullptr);
+
+		if (!Source)
 		{
 			return false;
 		}
+
+#if XAUDIO2_SUPPORTS_SENDLIST
+		Source->SetOutputVoices(&SourceSendList);
+#endif
+		Source->SetEffectChain(nullptr);
 	}
 
 	return true;
diff --git a/Engine/Source/Runtime/Windows/XAudio2/Private/XAudio2Support.h b/Engine/Source/Runtime/Windows/XAudio2/Private/XAudio2Support.h
--- a/Engine/Source/Runtime/Windows/XAudio2/Private/XAudio2Support.h
+++ b/Engine/Source/Runtime/Windows/XAudio2/Private/XAudio2Support.h
@@ -480,8 +480,8 @@ protected:
 	FXAudio2SoundBuffer*		XAudio2Buffer;
 	/** XAudio2 source voice associated with this source. */
 	IXAudio2SourceVoice*		Source;
-	/** Structure to handle looping sound callbacks */
-	FXAudio2SoundSourceCallback	SourceCallback;
+	/** The max channels in the voice's effect chain. This is used to classify a pool for IXAudio2SourceVoice. */
+	int32						MaxEffectChainChannels;
 
 	/** Asynchronous task for real time audio sources */
 	FAsyncRealtimeAudioTask* RealtimeAsyncTask;
@@ -564,7 +564,37 @@ public:
 	void CalculateDolbySurroundRate( const FVector& OrientFront, const FVector& ListenerPosition, const FVector& EmitterPosition, float OmniRadius, float* OutVolumes );
 };
 
-/** Variables required for the early init */
+/** A pool entry for related IXAudio2SourceVoices */
+struct FSourceVoicePoolEntry
+{
+	/** The format for all voices in this entry */
+	WAVEFORMATEX Format;
+	
+	/** The max number of channels used in the effect chain for this voice. This is needed because
+		XAudio2 defaults the max output channels for any effect chain to be the number of input channels. So
+		a mono-to-stereo effect (e.g. for HRTF processing) would not work. 
+	*/
+	int32 MaxEffectChainChannels;
+	
+	/** The array of free voices in this pool entry. */
+	TArray<struct IXAudio2SourceVoice*> FreeVoices;
+};
+
+/** Function to compare two WAVEFORMATEX structs */
+bool operator==(const WAVEFORMATEX& FormatA, const WAVEFORMATEX& FormatB)
+{
+	/** Unfortunately, need to compare every member of the WAVEFORMATEX struct */
+	return FormatA.cbSize 			== FormatB.cbSize &&
+			FormatA.nAvgBytesPerSec == FormatB.nAvgBytesPerSec &&
+			FormatA.nBlockAlign 	== FormatB.nBlockAlign &&
+			FormatA.nChannels 		== FormatB.nChannels &&
+			FormatA.nSamplesPerSec 	== FormatB.nSamplesPerSec &&
+			FormatA.wBitsPerSample 	== FormatB.wBitsPerSample &&
+			FormatA.wFormatTag 		== FormatB.wFormatTag;
+}
+
+
+/** This structure holds any singleton XAudio2 resources which need to be used, not just "properties" of the device. */
 struct FXAudioDeviceProperties
 {
 	// These variables are non-static to support multiple audio device instances
@@ -579,12 +609,125 @@ struct FXAudioDeviceProperties
 	static XAUDIO2_DEVICE_DETAILS		DeviceDetails;
 #endif	//XAUDIO_SUPPORTS_DEVICE_DETAILS
 
+	/** Source callback to handle looping sound callbacks */
+	FXAudio2SoundSourceCallback	SourceCallback;
+	
+	/** The array of voice pools. Each pool is according to the sound format (and max effect chain channels) */
+	TArray<FSourceVoicePoolEntry*> VoicePool;
+
 	FXAudioDeviceProperties()
 		: XAudio2(nullptr)
 		, MasteringVoice(nullptr)
 		, XAudio2Dll(nullptr)
 	{
 	}
+	
+	~FXAudioDeviceProperties()
+	{
+		// Destroy all the xaudio2 voices allocated in our pools
+		for (int32 i = 0; i < VoicePool.Num(); ++i)
+		{
+			for (int32 j = 0; j < VoicePool[i]->FreeVoices.Num(); ++j)
+			{
+				IXAudio2SourceVoice* Voice = VoicePool[i]->FreeVoices[j];
+				Voice->DestroyVoice();
+			}
+		}
+
+		// Now delete all the pool entries
+		for (int32 i = 0; i < VoicePool.Num(); ++i)
+		{
+			delete VoicePool[i];
+			VoicePool[i] = nullptr;
+		}
+
+		// close hardware interfaces
+		if (MasteringVoice)
+		{
+			MasteringVoice->DestroyVoice();
+			MasteringVoice = NULL;
+		}
+
+		if (XAudio2)
+		{
+			// Force the hardware to release all references
+			XAudio2->Release();
+			XAudio2 = NULL;
+		}
+
+#if PLATFORM_WINDOWS && PLATFORM_64BITS
+		if (XAudio2Dll)
+		{
+			if (!FreeLibrary(XAudio2Dll))
+			{
+				UE_LOG(LogAudio, Warning, TEXT("Failed to free XAudio2 Dll"));
+			}
+		}
+#endif
+	}
+
+	/** Returns either a new IXAudio2SourceVoice or a recycled IXAudio2SourceVoice according to the sound format and max channel count in the voice's effect chain*/
+	void GetFreeSourceVoice(IXAudio2SourceVoice** Voice, const FPCMBufferInfo& BufferInfo, const XAUDIO2_EFFECT_CHAIN* EffectChain = nullptr, int32 MaxEffectChainChannels = 0)
+	{
+		// First find the pool for the given format
+		FSourceVoicePoolEntry* VoicePoolEntry = nullptr;
+		for (int32 i = 0; i < VoicePool.Num(); ++i)
+		{
+			if (VoicePool[i]->Format == BufferInfo.PCMFormat && VoicePool[i]->MaxEffectChainChannels == MaxEffectChainChannels)
+			{
+				VoicePoolEntry = VoicePool[i];
+				break;
+			}
+		}
+
+		// If we found a voice pool entry for this format and we have free voices
+		// then use the voice
+		if (VoicePoolEntry && VoicePoolEntry->FreeVoices.Num() > 0)
+		{
+			*Voice = VoicePoolEntry->FreeVoices.Pop(false);
+		}
+		else
+		{
+			// Create a brand new source voice with this format.
+			XAudio2->CreateSourceVoice(Voice, &BufferInfo.PCMFormat, XAUDIO2_VOICE_USEFILTER, MAX_PITCH, &SourceCallback, nullptr, EffectChain);
+		}
+	}
+
+	/** Releases the voice into a pool of free voices according to the voice format and the max effect chain channels */
+	void ReleaseSourceVoice(IXAudio2SourceVoice* Voice, const FPCMBufferInfo& BufferInfo, const int32 MaxEffectChainChannels)
+	{
+		// Make sure the voice is stopped
+		Voice->Stop();
+
+		// And make sure there's no audio remaining the voice so when it's re-used it's fresh.
+		Voice->FlushSourceBuffers();
+
+		// See if there is an existing pool for this source voice
+		FSourceVoicePoolEntry* VoicePoolEntry = nullptr;
+		for (int32 i = 0; i < VoicePool.Num(); ++i)
+		{
+			if (VoicePool[i]->Format == BufferInfo.PCMFormat && VoicePool[i]->MaxEffectChainChannels == MaxEffectChainChannels)
+			{
+				VoicePoolEntry = VoicePool[i];
+				break;
+			}
+		}
+
+		// If we found a voice pool entry for this voice, add the voice to the list of free voices
+		if (VoicePoolEntry)
+		{
+			VoicePoolEntry->FreeVoices.Add(Voice);
+		}
+		else
+		{
+			// Otherwise We need to make a new voice pool entry with this format and max effect chain channels
+			VoicePoolEntry = new FSourceVoicePoolEntry();
+			VoicePoolEntry->Format = BufferInfo.PCMFormat;
+			VoicePoolEntry->FreeVoices.Add(Voice);
+			VoicePoolEntry->MaxEffectChainChannels = MaxEffectChainChannels;
+			VoicePool.Add(VoicePoolEntry);
+		}
+	}
 };
 
 #if XAUDIO_SUPPORTS_DEVICE_DETAILS
diff --git a/Engine/Source/Runtime/Windows/XAudio2/Public/XAudio2Device.h b/Engine/Source/Runtime/Windows/XAudio2/Public/XAudio2Device.h
--- a/Engine/Source/Runtime/Windows/XAudio2/Public/XAudio2Device.h
+++ b/Engine/Source/Runtime/Windows/XAudio2/Public/XAudio2Device.h
@@ -74,28 +74,6 @@ enum ESoundFormat
 	SoundFormat_Streaming
 };
 
-
-/** 
-* Simple async task to destroy an xaudio2 source voice without blocking the main thread.
-*/
-class FAsyncXAudio2SourceDestroyer : public FNonAbandonableTask
-{
-protected:
-	struct IXAudio2SourceVoice* SourceVoice;
-	FXAudio2Device* AudioDevice;
-
-public:
-	FAsyncXAudio2SourceDestroyer(FXAudio2Device* InAudioDevice, struct IXAudio2SourceVoice* InSourceVoice);
-	~FAsyncXAudio2SourceDestroyer();
-	void DoWork();
-
-	FORCEINLINE TStatId GetStatId() const
-	{
-		RETURN_QUICK_DECLARE_CYCLE_STAT(FAsyncXAudio2SourceDestroyer, STATGROUP_ThreadPoolAsyncTasks);
-	}
-};
-
-
 class FXAudio2SoundBuffer;
 class FXAudio2SoundSource;
 class FSpatializationHelper;
@@ -105,7 +83,6 @@ class FSpatializationHelper;
  */
 class FXAudio2Device : public FAudioDevice
 {
-
 	/** Starts up any platform specific hardware/APIs */
 	virtual bool InitializeHardware() override;
 
@@ -166,12 +143,6 @@ class FXAudio2Device : public FAudioDevice
 	 */
 	virtual bool Exec( UWorld* InWorld, const TCHAR* Cmd, FOutputDevice& Ar ) override;
 
-public:
-	/** 
-	 * Safely destroy an XAudio2 source asynchronously.
-	*/
-	void AsyncDestroyXAudio2Source(struct IXAudio2SourceVoice* Source);
-
 protected:
 
 	/**
@@ -208,7 +179,6 @@ protected:
 
 private:
 	struct FXAudioDeviceProperties* DeviceProperties;
-	FThreadSafeCounter NumSourcesDestroying;
 
 #if PLATFORM_WINDOWS
 	// We need to keep track whether com was successfully initialized so we can clean 
