Xaudio2 non unity compile fix. Force line comparison struct.

diff --git a/Engine/Source/Runtime/Windows/XAudio2/Private/XAudio2Support.h b/Engine/Source/Runtime/Windows/XAudio2/Private/XAudio2Support.h
--- a/Engine/Source/Runtime/Windows/XAudio2/Private/XAudio2Support.h
+++ b/Engine/Source/Runtime/Windows/XAudio2/Private/XAudio2Support.h
@@ -581,7 +581,7 @@ struct FSourceVoicePoolEntry
 };
 
 /** Function to compare two WAVEFORMATEX structs */
-bool operator==(const WAVEFORMATEX& FormatA, const WAVEFORMATEX& FormatB)
+FORCEINLINE bool operator==(const WAVEFORMATEX& FormatA, const WAVEFORMATEX& FormatB)
 {
 	/** Unfortunately, need to compare every member of the WAVEFORMATEX struct */
 	return FormatA.cbSize 			== FormatB.cbSize &&
