add null check because of crash

diff --git a/Engine/Source/Runtime/Windows/D3D11RHI/Private/D3D11Commands.cpp b/Engine/Source/Runtime/Windows/D3D11RHI/Private/D3D11Commands.cpp
--- a/Engine/Source/Runtime/Windows/D3D11RHI/Private/D3D11Commands.cpp
+++ b/Engine/Source/Runtime/Windows/D3D11RHI/Private/D3D11Commands.cpp
@@ -2101,24 +2101,27 @@ void FD3D11DynamicRHI::RHITransitionResources(EResourceTransitionAccess Transiti
 		FTextureRHIParamRef RenderTarget = InTextures[i];
 		SCOPED_RHI_CONDITIONAL_DRAW_EVENTF(*this, RHITransitionResourcesLoop, bShowTransitionEvents, TEXT("To:%i - %s"), i, *RenderTarget->GetName().ToString());
 
-		FD3D11BaseShaderResource* Resource = nullptr;
-		FD3D11Texture2D* SourceTexture2D = static_cast<FD3D11Texture2D*>(RenderTarget->GetTexture2D());
-		if (SourceTexture2D)
+		if (RenderTarget)
 		{
-			Resource = SourceTexture2D;
-		}
-		FD3D11TextureCube* SourceTextureCube = static_cast<FD3D11TextureCube*>(RenderTarget->GetTextureCube());
-		if (SourceTextureCube)
-		{
-			Resource = SourceTextureCube;			
-		}
-		FD3D11Texture3D* SourceTexture3D = static_cast<FD3D11Texture3D*>(RenderTarget->GetTexture3D());
-		if (SourceTexture3D)
-		{
-			Resource = SourceTexture3D;			
-		}
+			FD3D11BaseShaderResource* Resource = nullptr;
+			FD3D11Texture2D* SourceTexture2D = static_cast<FD3D11Texture2D*>(RenderTarget->GetTexture2D());
+			if(SourceTexture2D)
+			{
+				Resource = SourceTexture2D;
+			}
+			FD3D11TextureCube* SourceTextureCube = static_cast<FD3D11TextureCube*>(RenderTarget->GetTextureCube());
+			if(SourceTextureCube)
+			{
+				Resource = SourceTextureCube;
+			}
+			FD3D11Texture3D* SourceTexture3D = static_cast<FD3D11Texture3D*>(RenderTarget->GetTexture3D());
+			if(SourceTexture3D)
+			{
+				Resource = SourceTexture3D;
+			}
 
-		Resource->SetCurrentGPUAccess(TransitionType, PresentCounter);
+			Resource->SetCurrentGPUAccess(TransitionType, PresentCounter);
+		}
 	}
 }
 
