UE4 - D3D12 - Fix crash on machines with > 8 threads

diff --git a/Engine/Source/Runtime/Windows/D3D12RHI/Private/D3D12RHIPrivate.h b/Engine/Source/Runtime/Windows/D3D12RHI/Private/D3D12RHIPrivate.h
--- a/Engine/Source/Runtime/Windows/D3D12RHI/Private/D3D12RHIPrivate.h
+++ b/Engine/Source/Runtime/Windows/D3D12RHI/Private/D3D12RHIPrivate.h
@@ -1617,7 +1617,7 @@ public:
 
 		return *HelperThreadDynamicHeapAllocator;
 	}
-	FD3D12DynamicHeapAllocator* ThreadDynamicHeapAllocatorArray[8];
+	FD3D12DynamicHeapAllocator* ThreadDynamicHeapAllocatorArray[16];
 	uint32 NumThreadDynamicHeapAllocators;
 	static __declspec(thread) FD3D12DynamicHeapAllocator* HelperThreadDynamicHeapAllocator;
 
