Integrate from Orion

CL 2704794 2704931 2704948 2704962 2705238 2706353 2700643 2705458

        Uniform Buffer layout name

and

more debug info for this and other bugs: OR-7159 CRASH: Client crashed at start of match in D3D11Commands.cpp

Adding OptionalData (key value pairs) to the shaders (to put shader name into the shader code for better debugging), affects all RHI, invalidate DDC key for all shaders, cleanup existing code (was adding 0/1/5 bytes and had to compensate in many areas)

OptionalData:
key is a char, value is up to 255 bytes, up to 64K in total

diff --git a/Engine/Source/Developer/Apple/MetalShaderFormat/Private/MetalShaderCompiler.cpp b/Engine/Source/Developer/Apple/MetalShaderFormat/Private/MetalShaderCompiler.cpp
--- a/Engine/Source/Developer/Apple/MetalShaderFormat/Private/MetalShaderCompiler.cpp
+++ b/Engine/Source/Developer/Apple/MetalShaderFormat/Private/MetalShaderCompiler.cpp
@@ -311,12 +311,15 @@ static void BuildMetalShaderOutput(
 	else if(ShaderInput.Environment.CompilerFlags.Contains(CFLAG_Debug))
 	{
 		// Write out the header and shader source code.
-		FMemoryWriter Ar(ShaderOutput.Code, true);
+		FMemoryWriter Ar(ShaderOutput.ShaderCode.GetWriteAccess(), true);
 		uint8 PrecompiledFlag = 0;
 		Ar << PrecompiledFlag;
 		Ar << Header;
 		Ar.Serialize((void*)USFSource, SourceLen + 1 - (USFSource - InShaderSource));
 		
+		// store data we can pickup later with ShaderCode.FindOptionalData('n'), could be removed for shipping
+		ShaderOutput.ShaderCode.AddOptionalData('n', TCHAR_TO_UTF8(*ShaderInput.GenerateShaderName()));
+
 		ShaderOutput.NumInstructions = 0;
 		ShaderOutput.NumTextureSamplers = Header.Bindings.NumSamplers;
 		ShaderOutput.bSucceeded = true;
@@ -407,7 +410,7 @@ static void BuildMetalShaderOutput(
 						bCompileAtRuntime = false;
 						
 						// Write out the header and compiled shader code
-						FMemoryWriter Ar(ShaderOutput.Code, true);
+						FMemoryWriter Ar(ShaderOutput.Code.ShaderCode.GetWriteAccess(), true);
 						uint8 PrecompiledFlag = 1;
 						Ar << PrecompiledFlag;
 						Ar << Header;
@@ -419,6 +422,9 @@ static void BuildMetalShaderOutput(
 						// jam it into the output bytes
 						Ar.Serialize(CompiledShader.GetData(), CompiledShader.Num());
 						
+						// store data we can pickup later with ShaderCode.FindOptionalData('n'), could be removed for shipping
+						ShaderOutput.ShaderCode.AddOptionalData('n', TCHAR_TO_UTF8(*ShaderInput.GenerateShaderName()));
+
 						ShaderOutput.NumInstructions = 0;
 						ShaderOutput.NumTextureSamplers = Header.Bindings.NumSamplers;
 						ShaderOutput.bSucceeded = true;
@@ -435,12 +441,15 @@ static void BuildMetalShaderOutput(
 		if (bCompileAtRuntime)
 		{
 			// Write out the header and shader source code.
-			FMemoryWriter Ar(ShaderOutput.Code, true);
+			FMemoryWriter Ar(ShaderOutput.ShaderCode.GetWriteAccess(), true);
 			uint8 PrecompiledFlag = 0;
 			Ar << PrecompiledFlag;
 			Ar << Header;
 			Ar.Serialize((void*)USFSource, SourceLen + 1 - (USFSource - InShaderSource));
 			
+			// store data we can pickup later with ShaderCode.FindOptionalData('n'), could be removed for shipping
+			ShaderOutput.ShaderCode.AddOptionalData('n', TCHAR_TO_UTF8(*ShaderInput.GenerateShaderName()));
+			
 			ShaderOutput.NumInstructions = 0;
 			ShaderOutput.NumTextureSamplers = Header.Bindings.NumSamplers;
 			ShaderOutput.bSucceeded = bSucceeded || ShaderOutput.bSucceeded;
diff --git a/Engine/Source/Developer/Apple/MetalShaderFormat/Private/MetalShaderFormat.cpp b/Engine/Source/Developer/Apple/MetalShaderFormat/Private/MetalShaderFormat.cpp
--- a/Engine/Source/Developer/Apple/MetalShaderFormat/Private/MetalShaderFormat.cpp
+++ b/Engine/Source/Developer/Apple/MetalShaderFormat/Private/MetalShaderFormat.cpp
@@ -17,7 +17,7 @@ class FMetalShaderFormat : public IShaderFormat
 {
 	enum
 	{
-		HEADER_VERSION = 25,
+		HEADER_VERSION = 26,
 	};
 	struct FVersion
 	{
diff --git a/Engine/Source/Developer/ShaderFormatOpenGL/Private/OpenGLShaderCompiler.cpp b/Engine/Source/Developer/ShaderFormatOpenGL/Private/OpenGLShaderCompiler.cpp
--- a/Engine/Source/Developer/ShaderFormatOpenGL/Private/OpenGLShaderCompiler.cpp
+++ b/Engine/Source/Developer/ShaderFormatOpenGL/Private/OpenGLShaderCompiler.cpp
@@ -832,10 +832,13 @@ static void BuildShaderOutput(
 	else
 	{
 		// Write out the header and shader source code.
-		FMemoryWriter Ar(ShaderOutput.Code, true);
+		FMemoryWriter Ar(ShaderOutput.ShaderCode.GetWriteAccess(), true);
 		Ar << Header;
 		Ar.Serialize((void*)USFSource, SourceLen + 1 - (USFSource - InShaderSource));
 		
+		// store data we can pickup later with ShaderCode.FindOptionalData('n'), could be removed for shipping
+		ShaderOutput.ShaderCode.AddOptionalData('n', TCHAR_TO_UTF8(*ShaderInput.GenerateShaderName()));
+
 		ShaderOutput.NumInstructions = 0;
 		ShaderOutput.NumTextureSamplers = Header.Bindings.NumSamplers;
 		ShaderOutput.bSucceeded = true;
diff --git a/Engine/Source/Developer/ShaderFormatOpenGL/Private/ShaderFormatOpenGL.cpp b/Engine/Source/Developer/ShaderFormatOpenGL/Private/ShaderFormatOpenGL.cpp
--- a/Engine/Source/Developer/ShaderFormatOpenGL/Private/ShaderFormatOpenGL.cpp
+++ b/Engine/Source/Developer/ShaderFormatOpenGL/Private/ShaderFormatOpenGL.cpp
@@ -25,16 +25,16 @@ class FShaderFormatGLSL : public IShaderFormat
 	enum 
 	{
 		/** Version for shader format, this becomes part of the DDC key. */
-		UE_SHADER_GLSL_150_VER = 28,
-		UE_SHADER_GLSL_150_MAC_VER = 8,
+		UE_SHADER_GLSL_150_VER = 60,
+		UE_SHADER_GLSL_150_MAC_VER = 60,
 		UE_SHADER_GLSL_430_VER = 60,
-		UE_SHADER_GLSL_ES2_VER = 18,
-		UE_SHADER_GLSL_150ES2_VER = 20,
-		UE_SHADER_GLSL_150ES2NOUB_VER = 4,
-		UE_SHADER_GLSL_150ES3_1_VER = 4,
-		UE_SHADER_GLSL_ES2_VER_WEBGL  = 18,
-		UE_SHADER_GLSL_ES2_IOS_VER  = 7,
-		UE_SHADER_GLSL_310_ES_EXT_VER = 7,
+		UE_SHADER_GLSL_ES2_VER = 60,
+		UE_SHADER_GLSL_150ES2_VER = 60,
+		UE_SHADER_GLSL_150ES2NOUB_VER = 60,
+		UE_SHADER_GLSL_150ES3_1_VER = 60,
+		UE_SHADER_GLSL_ES2_VER_WEBGL  = 60,
+		UE_SHADER_GLSL_ES2_IOS_VER  = 60,
+		UE_SHADER_GLSL_310_ES_EXT_VER = 60,
 	}; 
 
 	void CheckFormat(FName Format) const
diff --git a/Engine/Source/Developer/Windows/ShaderFormatD3D/Private/D3D11ShaderCompiler.cpp b/Engine/Source/Developer/Windows/ShaderFormatD3D/Private/D3D11ShaderCompiler.cpp
--- a/Engine/Source/Developer/Windows/ShaderFormatD3D/Private/D3D11ShaderCompiler.cpp
+++ b/Engine/Source/Developer/Windows/ShaderFormatD3D/Private/D3D11ShaderCompiler.cpp
@@ -689,16 +689,19 @@ void CompileD3D11Shader(const FShaderCompilerInput& Input,FShaderCompilerOutput&
 
 			}
 
-			FMemoryWriter Ar( Output.Code, true );
+			FMemoryWriter Ar(Output.ShaderCode.GetWriteAccess(), true);
 			Ar << SRT;
-			Ar.Serialize( CompressedData->GetBufferPointer(), CompressedData->GetBufferSize() );
-
-			// Pack bGlobalUniformBufferUsed and resource counts in the last few bytes
-			Output.Code.Add(bGlobalUniformBufferUsed);
-			Output.Code.Add(NumSamplers);
-			Output.Code.Add(NumSRVs);
-			Output.Code.Add(NumCBs);
-			Output.Code.Add(NumUAVs);
+			Ar.Serialize(CompressedData->GetBufferPointer(), CompressedData->GetBufferSize());
+
+			// append data that is generate from the shader code and assist the usage, mostly needed for DX12 
+			{
+				FShaderCodePackedResourceCounts PackedResourceCounts = { bGlobalUniformBufferUsed, NumSamplers, NumSRVs, NumCBs, NumUAVs };
+
+				Output.ShaderCode.AddOptionalData(PackedResourceCounts);
+			}
+
+			// store data we can pickup later with ShaderCode.FindOptionalData('n'), could be removed for shipping
+			Output.ShaderCode.AddOptionalData('n', TCHAR_TO_UTF8(*Input.GenerateShaderName()));
 
 			// Set the number of instructions.
 			Output.NumInstructions = ShaderDesc.InstructionCount;
diff --git a/Engine/Source/Developer/Windows/ShaderFormatD3D/Private/ShaderFormatD3D.cpp b/Engine/Source/Developer/Windows/ShaderFormatD3D/Private/ShaderFormatD3D.cpp
--- a/Engine/Source/Developer/Windows/ShaderFormatD3D/Private/ShaderFormatD3D.cpp
+++ b/Engine/Source/Developer/Windows/ShaderFormatD3D/Private/ShaderFormatD3D.cpp
@@ -16,10 +16,10 @@ class FShaderFormatD3D : public IShaderFormat
 	enum
 	{
 		/** Version for shader format, this becomes part of the DDC key. */
-		UE_SHADER_PCD3D_SM5_VER = 3,
-		UE_SHADER_PCD3D_SM4_VER = 3,
-		UE_SHADER_PCD3D_ES2_VER = 3,
-		UE_SHADER_PCD3D_ES3_1_VER = 2,
+		UE_SHADER_PCD3D_SM5_VER = 4,
+		UE_SHADER_PCD3D_SM4_VER = 4,
+		UE_SHADER_PCD3D_ES2_VER = 4,
+		UE_SHADER_PCD3D_ES3_1_VER = 4,
 	};
 
 	void CheckFormat(FName Format) const
diff --git a/Engine/Source/Programs/ShaderCompileWorker/Private/ShaderCompileWorker.cpp b/Engine/Source/Programs/ShaderCompileWorker/Private/ShaderCompileWorker.cpp
--- a/Engine/Source/Programs/ShaderCompileWorker/Private/ShaderCompileWorker.cpp
+++ b/Engine/Source/Programs/ShaderCompileWorker/Private/ShaderCompileWorker.cpp
@@ -13,8 +13,10 @@
 
 #define DEBUG_USING_CONSOLE	0
 
-const int32 ShaderCompileWorkerInputVersion = 3;
-const int32 ShaderCompileWorkerOutputVersion = 1;
+// this is for the protocol, not the data, bump if FShaderCompilerInput or ProcessInputFromArchive changes (also search for the second one with the same name, todo: put into one header file)
+const int32 ShaderCompileWorkerInputVersion = 4;
+// this is for the protocol, not the data, bump if FShaderCompilerOutput or WriteToOutputArchive changes (also search for the second one with the same name, todo: put into one header file)
+const int32 ShaderCompileWorkerOutputVersion = 2;
 
 double LastCompileTime = 0.0;
 
diff --git a/Engine/Source/Runtime/Apple/MetalRHI/Private/MetalShaders.cpp b/Engine/Source/Runtime/Apple/MetalRHI/Private/MetalShaders.cpp
--- a/Engine/Source/Runtime/Apple/MetalRHI/Private/MetalShaders.cpp
+++ b/Engine/Source/Runtime/Apple/MetalRHI/Private/MetalShaders.cpp
@@ -49,10 +49,12 @@ static FMetalCompiledShaderCache& GetMetalCompiledShaderCache()
 
 /** Initialization constructor. */
 template<typename BaseResourceType, int32 ShaderType>
-TMetalBaseShader<BaseResourceType, ShaderType>::TMetalBaseShader(const TArray<uint8>& InCode)
+TMetalBaseShader<BaseResourceType, ShaderType>::TMetalBaseShader(const TArray<uint8>& InShaderCode)
 	: DirtyUniformBuffers(0)
 {
-	FMemoryReader Ar(InCode, true);
+	FMemoryReader Ar(InShaderCode, true);
+	
+	Ar.SetLimitSize(FShaderCodeReader(InShaderCode).GetActualShaderCodeSize());
 
 	// was the shader already compiled offline?
 	uint8 OfflineCompiledFlag;
diff --git a/Engine/Source/Runtime/Core/Public/Serialization/Archive.h b/Engine/Source/Runtime/Core/Public/Serialization/Archive.h
--- a/Engine/Source/Runtime/Core/Public/Serialization/Archive.h
+++ b/Engine/Source/Runtime/Core/Public/Serialization/Archive.h
@@ -163,12 +163,12 @@ public:
 
 	int64 TotalSize()
 	{
-		return Bytes.Num();
+		return FMath::Min((int64)Bytes.Num(), LimitSize);
 	}
 
 	void Seek(int64 InPos )
 	{
-		check(InPos<=Bytes.Num());
+		check(InPos <= TotalSize());
 		FMemoryArchive::Seek(InPos);
 	}
 
@@ -177,7 +177,7 @@ public:
 		if (Num && !ArIsError)
 		{
 			// Only serialize if we have the requested amount of data
-			if (Offset + Num <= Bytes.Num())
+			if (Offset + Num <= TotalSize())
 			{
 				FMemory::Memcpy( Data, &Bytes[Offset], Num );
 				Offset += Num;
@@ -188,17 +188,26 @@ public:
 			}
 		}
 	}
+
 	FMemoryReader( const TArray<uint8>& InBytes, bool bIsPersistent = false )
 	: FMemoryArchive()
 	, Bytes(InBytes)
+	, LimitSize(INT64_MAX)
 	{
 		ArIsLoading		= true;
 		ArIsPersistent	= bIsPersistent;
 	}
 
+	// with this method it's possible to attach data behind some serialized data 
+	void SetLimitSize(int64 NewLimitSize)
+	{
+		LimitSize = NewLimitSize;
+	}
+
 protected:
 
 	const TArray<uint8>& Bytes;
+	int64 LimitSize;
 };
 
 /**
diff --git a/Engine/Source/Runtime/Engine/Private/Materials/MaterialUniformExpressions.cpp b/Engine/Source/Runtime/Engine/Private/Materials/MaterialUniformExpressions.cpp
--- a/Engine/Source/Runtime/Engine/Private/Materials/MaterialUniformExpressions.cpp
+++ b/Engine/Source/Runtime/Engine/Private/Materials/MaterialUniformExpressions.cpp
@@ -294,8 +294,10 @@ void FUniformExpressionSet::CreateBufferStruct()
 	new(Members) FUniformBufferStruct::FMember(TEXT("Clamp_WorldGroupSettings"),TEXT("SamplerState"),NextMemberOffset,UBMT_SAMPLER,EShaderPrecisionModifier::Float,1,1,1,NULL);
 	NextMemberOffset += 8;
 
+	static FName LayoutName(TEXT("MaterialUniforms"));
 	const uint32 StructSize = Align(NextMemberOffset,UNIFORM_BUFFER_STRUCT_ALIGNMENT);
 	UniformBufferStruct = new FUniformBufferStruct(
+		LayoutName,
 		TEXT("MaterialUniforms"),
 		TEXT("Material"),
 		ConstructMaterialUniformBufferParameter,
diff --git a/Engine/Source/Runtime/Engine/Private/Materials/ParameterCollection.cpp b/Engine/Source/Runtime/Engine/Private/Materials/ParameterCollection.cpp
--- a/Engine/Source/Runtime/Engine/Private/Materials/ParameterCollection.cpp
+++ b/Engine/Source/Runtime/Engine/Private/Materials/ParameterCollection.cpp
@@ -364,9 +364,10 @@ void UMaterialParameterCollection::CreateBufferStruct()
 	new(Members) FUniformBufferStruct::FMember(TEXT("Vectors"),TEXT(""),NextMemberOffset,UBMT_FLOAT32,EShaderPrecisionModifier::Half,1,4,NumVectors,NULL);
 	const uint32 VectorArraySize = NumVectors * sizeof(FVector4);
 	NextMemberOffset += VectorArraySize;
-
+	static FName LayoutName(TEXT("MaterialCollection"));
 	const uint32 StructSize = Align(NextMemberOffset,UNIFORM_BUFFER_STRUCT_ALIGNMENT);
 	UniformBufferStruct = new FUniformBufferStruct(
+		LayoutName,
 		TEXT("MaterialCollection"),
 		TEXT("MaterialCollection"),
 		ConstructCollectionUniformBufferParameter,
@@ -572,6 +573,12 @@ void FMaterialParameterCollectionInstanceResource::GameThread_Destroy()
 	});
 }
 
+static FName MaterialParameterCollectionInstanceResourceName(TEXT("MaterialParameterCollectionInstanceResource"));
+FMaterialParameterCollectionInstanceResource::FMaterialParameterCollectionInstanceResource() :
+	UniformBufferLayout(MaterialParameterCollectionInstanceResourceName)
+{
+}
+
 FMaterialParameterCollectionInstanceResource::~FMaterialParameterCollectionInstanceResource()
 {
 	check(IsInRenderingThread());
@@ -592,4 +599,4 @@ void FMaterialParameterCollectionInstanceResource::UpdateContents(const FGuid& I
 		check(UniformBufferLayout.Resources.Num() == 0);
 		UniformBuffer = RHICreateUniformBuffer(Data.GetData(), UniformBufferLayout, UniformBuffer_MultiFrame);
 	}
-}
\ No newline at end of file
+}
diff --git a/Engine/Source/Runtime/Engine/Private/ShaderCompiler/ShaderCompiler.cpp b/Engine/Source/Runtime/Engine/Private/ShaderCompiler/ShaderCompiler.cpp
--- a/Engine/Source/Runtime/Engine/Private/ShaderCompiler/ShaderCompiler.cpp
+++ b/Engine/Source/Runtime/Engine/Private/ShaderCompiler/ShaderCompiler.cpp
@@ -162,7 +162,7 @@ static inline void GetFormatVersionMap(TMap<FString, uint16>& OutFormatVersionMa
 // Serialize Queued Job information
 static bool DoWriteTasks(const TArray<FShaderCompileJob*>& QueuedJobs, FArchive& TransferFile)
 {
-	int32 ShaderCompileWorkerInputVersion = 3;
+	int32 ShaderCompileWorkerInputVersion = 4;
 	TransferFile << ShaderCompileWorkerInputVersion;
 
 	static TMap<FString, uint16> FormatVersionMap;
@@ -185,7 +185,7 @@ static bool DoWriteTasks(const TArray<FShaderCompileJob*>& QueuedJobs, FArchive&
 // Process results from Worker Process
 static void DoReadTaskResults(const TArray<FShaderCompileJob*>& QueuedJobs, FArchive& OutputFile)
 {
-	const int32 OutputVersion = 1;
+	const int32 OutputVersion = 2;
 	int32 ShaderCompileWorkerOutputVersion;
 	OutputFile << ShaderCompileWorkerOutputVersion;
 
@@ -1159,7 +1159,7 @@ void FShaderCompilingManager::ProcessCompiledShaderMaps(
 
 				if (CurrentJob.bSucceeded)
 				{
-					check(CurrentJob.Output.Code.Num() > 0);
+					check(CurrentJob.Output.ShaderCode.GetShaderCodeSize() > 0);
 				}
 				else 
 				{
@@ -1793,6 +1793,8 @@ void GlobalBeginCompileShader(
 	Input.SourceFilename = SourceFilename;
 	Input.EntryPointName = FunctionName;
 	Input.DumpDebugInfoRootPath = GShaderCompilingManager->GetAbsoluteShaderDebugInfoDirectory() / Input.ShaderFormat.ToString();
+	// asset material name or "Global"
+	Input.DebugGroupName = DebugGroupName;
 
 	if (GDumpShaderDebugInfo != 0)
 	{
diff --git a/Engine/Source/Runtime/Engine/Public/ParameterCollection.h b/Engine/Source/Runtime/Engine/Public/ParameterCollection.h
--- a/Engine/Source/Runtime/Engine/Public/ParameterCollection.h
+++ b/Engine/Source/Runtime/Engine/Public/ParameterCollection.h
@@ -37,6 +37,7 @@ public:
 		return UniformBuffer;
 	}
 
+	FMaterialParameterCollectionInstanceResource();
 	~FMaterialParameterCollectionInstanceResource();
 
 private:
diff --git a/Engine/Source/Runtime/OpenGLDrv/Private/OpenGLShaders.cpp b/Engine/Source/Runtime/OpenGLDrv/Private/OpenGLShaders.cpp
--- a/Engine/Source/Runtime/OpenGLDrv/Private/OpenGLShaders.cpp
+++ b/Engine/Source/Runtime/OpenGLDrv/Private/OpenGLShaders.cpp
@@ -448,14 +448,19 @@ static void BindShaderLocations(GLenum TypeEnum, GLuint Resource, uint16 InOutMa
  * @returns the compiled shader upon success.
  */
 template <typename ShaderType>
-ShaderType* CompileOpenGLShader(const TArray<uint8>& Code)
+ShaderType* CompileOpenGLShader(const TArray<uint8>& InShaderCode)
 {
 	SCOPE_CYCLE_COUNTER(STAT_OpenGLShaderCompileTime);
 	VERIFY_GL_SCOPE();
 
+	FShaderCodeReader ShaderCode(InShaderCode);
+
 	ShaderType* Shader = nullptr;
 	const GLenum TypeEnum = ShaderType::TypeEnum;
-	FMemoryReader Ar(Code, true);
+	FMemoryReader Ar(InShaderCode, true);
+
+	Ar.SetLimitSize(ShaderCode.GetActualShaderCodeSize());
+
 	FOpenGLCodeHeader Header = { 0 };
 
 	Ar << Header;
@@ -482,7 +487,7 @@ ShaderType* CompileOpenGLShader(const TArray<uint8>& Code)
 
 	// The code as given to us.
 	FAnsiCharArray GlslCodeOriginal;
-	AppendCString(GlslCodeOriginal, (ANSICHAR*)Code.GetData() + CodeOffset);
+	AppendCString(GlslCodeOriginal, (ANSICHAR*)InShaderCode.GetData() + CodeOffset);
 	uint32 GlslCodeOriginalCRC = FCrc::MemCrc_DEPRECATED(GlslCodeOriginal.GetData(), GlslCodeOriginal.Num());
 
 	// The amended code we actually compile.
diff --git a/Engine/Source/Runtime/RHI/Public/RHIResources.h b/Engine/Source/Runtime/RHI/Public/RHIResources.h
--- a/Engine/Source/Runtime/RHI/Public/RHIResources.h
+++ b/Engine/Source/Runtime/RHI/Public/RHIResources.h
@@ -112,6 +112,11 @@ public:
 	void SetHash(FSHAHash InHash) { Hash = InHash; }
 	FSHAHash GetHash() const { return Hash; }
 
+#if !(UE_BUILD_SHIPPING || UE_BUILD_TEST)
+	// for debugging only e.g. MaterialName:ShaderFile.usf or ShaderFile.usf/EntryFunc
+	FString ShaderName;
+#endif
+
 private:
 	FSHAHash Hash;
 };
@@ -162,15 +167,20 @@ struct FRHIUniformBufferLayout
 		return Hash;
 	}
 
-	FRHIUniformBufferLayout() :
+	FRHIUniformBufferLayout(FName InName) :
 		ConstantBufferSize(0),
 		ResourceOffset(0),
+		Name(InName),
 		Hash(0),
 		bComputedHash(false)
 	{
 	}
 
+	const FName GetDebugName() const { return Name; }
+
 private:
+	// for debugging / error message
+	const FName Name;
 	mutable uint32 Hash;
 	mutable bool bComputedHash;
 };
diff --git a/Engine/Source/Runtime/RenderCore/Public/UniformBuffer.h b/Engine/Source/Runtime/RenderCore/Public/UniformBuffer.h
--- a/Engine/Source/Runtime/RenderCore/Public/UniformBuffer.h
+++ b/Engine/Source/Runtime/RenderCore/Public/UniformBuffer.h
@@ -192,11 +192,12 @@ public:
 	typedef class FShaderUniformBufferParameter* (*ConstructUniformBufferParameterType)();
 
 	/** Initialization constructor. */
-	FUniformBufferStruct(const TCHAR* InStructTypeName, const TCHAR* InShaderVariableName, ConstructUniformBufferParameterType InConstructRef, uint32 InSize, const TArray<FMember>& InMembers, bool bRegisterForAutoBinding)
+	FUniformBufferStruct(const FName& InLayoutName, const TCHAR* InStructTypeName, const TCHAR* InShaderVariableName, ConstructUniformBufferParameterType InConstructRef, uint32 InSize, const TArray<FMember>& InMembers, bool bRegisterForAutoBinding)
 	:	StructTypeName(InStructTypeName)
 	,	ShaderVariableName(InShaderVariableName)
 	,	ConstructUniformBufferParameterRef(InConstructRef)
 	,	Size(InSize)
+	,	Layout(InLayoutName)
 	,	Members(InMembers)
 	,	GlobalListLink(this)
 	{
@@ -581,6 +582,7 @@ private:
 
 #define IMPLEMENT_UNIFORM_BUFFER_STRUCT(StructTypeName,ShaderVariableName) \
 	FUniformBufferStruct StructTypeName::StaticStruct( \
+	FName(TEXT(#StructTypeName)), \
 	TEXT(#StructTypeName), \
 	ShaderVariableName, \
 	StructTypeName::ConstructUniformBufferParameter, \
diff --git a/Engine/Source/Runtime/ShaderCore/Private/Shader.cpp b/Engine/Source/Runtime/ShaderCore/Private/Shader.cpp
--- a/Engine/Source/Runtime/ShaderCore/Private/Shader.cpp
+++ b/Engine/Source/Runtime/ShaderCore/Private/Shader.cpp
@@ -313,7 +313,9 @@ FShaderResource::FShaderResource(const FShaderCompilerOutput& Output, FShaderTyp
 	
 {
 	Target = Output.Target;
-	Code = Output.Code;
+	// todo: can we avoid the memcpy?
+	Code = Output.ShaderCode.GetReadAccess();
+
 	check(Code.Num() > 0);
 
 	OutputHash = Output.OutputHash;
diff --git a/Engine/Source/Runtime/ShaderCore/Private/ShaderCore.cpp b/Engine/Source/Runtime/ShaderCore/Private/ShaderCore.cpp
--- a/Engine/Source/Runtime/ShaderCore/Private/ShaderCore.cpp
+++ b/Engine/Source/Runtime/ShaderCore/Private/ShaderCore.cpp
@@ -179,7 +179,13 @@ void FShaderParameterMap::RemoveParameterAllocation(const TCHAR* ParameterName)
 void FShaderCompilerOutput::GenerateOutputHash()
 {
 	FSHA1 HashState;
-	HashState.Update(Code.GetData(), Code.Num() * Code.GetTypeSize());
+	
+	const TArray<uint8>& Code = ShaderCode.GetReadAccess();
+
+	// we don't hash the optional attachments as they would prevent sharing (e.g. many material share the save VS)
+	uint32 ShaderCodeSize = ShaderCode.GetShaderCodeSize();
+
+	HashState.Update(Code.GetData(), ShaderCodeSize * Code.GetTypeSize());
 	ParameterMap.UpdateHash(HashState);
 	HashState.Final();
 	HashState.GetHash(&OutputHash.Hash[0]);
diff --git a/Engine/Source/Runtime/ShaderCore/Public/Shader.h b/Engine/Source/Runtime/ShaderCore/Public/Shader.h
--- a/Engine/Source/Runtime/ShaderCore/Public/Shader.h
+++ b/Engine/Source/Runtime/ShaderCore/Public/Shader.h
@@ -472,7 +472,7 @@ public:
 			):
 			Type(InType),
 			Target(CompilerOutput.Target),
-			Code(CompilerOutput.Code),
+			Code(CompilerOutput.ShaderCode.GetReadAccess()),
 			ParameterMap(CompilerOutput.ParameterMap),
 			OutputHash(CompilerOutput.OutputHash),
 			Resource(InResource),
diff --git a/Engine/Source/Runtime/ShaderCore/Public/ShaderCore.h b/Engine/Source/Runtime/ShaderCore/Public/ShaderCore.h
--- a/Engine/Source/Runtime/ShaderCore/Public/ShaderCore.h
+++ b/Engine/Source/Runtime/ShaderCore/Public/ShaderCore.h
@@ -416,13 +416,36 @@ struct FShaderCompilerInput
 	FShaderTarget Target;
 	FName ShaderFormat;
 	FString SourceFilePrefix;
+	// e.g. BasePassPixelShader, ReflectionEnvironmentShaders, SlateElementPixelShader, PostProcessCombineLUTs
 	FString SourceFilename;
 	FString EntryPointName;
-	FString DumpDebugInfoRootPath;	// Dump debug path (up to platform)
-	FString DumpDebugInfoPath;		// Dump debug path (platform/groupname)
+	// Dump debug path (up to platform) e.g. "D:/MMittring-Z3941-A/UE4-Orion/OrionGame/Saved/ShaderDebugInfo/PCD3D_SM5"
+	FString DumpDebugInfoRootPath;
+	// only used if enabled by r.DumpShaderDebugInfo (platform/groupname) e.g. ""
+	FString DumpDebugInfoPath;
+	// materialname or "Global" "for debugging and better error messages
+	FString DebugGroupName;
 	FShaderCompilerEnvironment Environment;
 	TRefCountPtr<FShaderCompilerEnvironment> SharedEnvironment;
 
+	// generate human readable name for debugging
+	FString GenerateShaderName() const
+	{
+		FString Name;
+
+		if(DebugGroupName == TEXT("Global"))
+		{
+			Name = SourceFilename + TEXT(".usf|") + EntryPointName;
+		}
+		else
+		{
+			// we skip EntryPointName as it's usually not useful
+			Name = DebugGroupName + TEXT(":") + SourceFilename + TEXT(".usf");
+		}
+
+		return Name;
+	}
+
 	friend FArchive& operator<<(FArchive& Ar,FShaderCompilerInput& Input)
 	{
 		// Note: this serialize is used to pass between UE4 and the shader compile worker, recompile both when modifying
@@ -438,6 +461,7 @@ struct FShaderCompilerInput
 		Ar << Input.EntryPointName;
 		Ar << Input.DumpDebugInfoRootPath;
 		Ar << Input.DumpDebugInfoPath;
+		Ar << Input.DebugGroupName;
 		Ar << Input.Environment;
 
 		bool bHasSharedEnvironment = IsValidRef(Input.SharedEnvironment);
@@ -487,6 +511,239 @@ struct FShaderCompilerError
 	}
 };
 
+// if this changes you need to make sure all D3D11 shaders get invalidated
+struct FShaderCodePackedResourceCounts
+{
+	// for FindOptionalData() and AddOptionalData()
+	static const uint8 Key = 'p';
+
+	bool bGlobalUniformBufferUsed;
+	uint8 NumSamplers;
+	uint8 NumSRVs;
+	uint8 NumCBs;
+	uint8 NumUAVs;
+};
+
+// later we can transform that to the actual class passed around at the RHI level
+class FShaderCodeReader
+{
+	const TArray<uint8>& ShaderCode;
+
+public:
+	FShaderCodeReader(const TArray<uint8>& InShaderCode)
+		: ShaderCode(InShaderCode)
+	{
+		check(ShaderCode.Num());
+	}
+
+	uint32 GetActualShaderCodeSize() const
+	{
+		return ShaderCode.Num() - GetOptionalDataSize();
+	}
+
+	// for convenience
+	template <class T>
+	const T* FindOptionalData() const
+	{
+		return (const T*)FindOptionalData(T::Key, sizeof(T));
+	}
+
+
+	// @param InKey e.g. FShaderCodePackedResourceCounts::Key
+	// @return 0 if not found
+	const uint8* FindOptionalData(uint8 InKey, uint8 ValueSize) const
+	{
+		check(ValueSize);
+
+		const uint8* End = &ShaderCode[0] + ShaderCode.Num();
+
+		int32 LocalOptionalDataSize = GetOptionalDataSize();
+
+		const uint8* Start = End - LocalOptionalDataSize;
+		const uint8* Current = Start;
+
+		while(Current < End)
+		{
+			uint8 Key = *Current++;
+			uint8 Size = *Current++;
+
+			if(Key == InKey)
+			{
+				return Current;
+			}
+
+			Current += Size;
+		}
+
+		return 0;
+	}
+
+	const ANSICHAR* FindOptionalData(uint8 InKey) const
+	{
+		check(ShaderCode.Num() >= 4);
+
+		const uint8* End = &ShaderCode[0] + ShaderCode.Num();
+
+		int32 LocalOptionalDataSize = GetOptionalDataSize();
+
+		const uint8* Start = End - LocalOptionalDataSize;
+		const uint8* Current = Start;
+
+		while(Current < End)
+		{
+			uint8 Key = *Current++;
+			uint8 Size = *Current++;
+
+			if(Key == InKey)
+			{
+				return (ANSICHAR*)Current;
+			}
+
+			Current += Size;
+		}
+
+		return 0;
+	}
+
+	int32 GetOptionalDataSize() const
+	{
+		if(ShaderCode.Num() < 2)
+		{
+			return 0;
+		}
+
+		const uint8* End = &ShaderCode[0] + ShaderCode.Num();
+
+		int32 LocalOptionalDataSize = (((uint32)End[-2]) << 8 ) + ((uint32)End[-1]);
+
+		check(LocalOptionalDataSize >= 0);
+		check(ShaderCode.Num() >= LocalOptionalDataSize);
+
+		return LocalOptionalDataSize;
+	}
+
+	int32 GetShaderCodeSize() const
+	{
+		return ShaderCode.Num() - GetOptionalDataSize();
+	}
+};
+
+class FShaderCode
+{
+	// -1 if ShaderData was finalized 
+	mutable int32 OptionalDataSize;
+	// access through class methods
+	mutable TArray<uint8> ShaderCodeWithOptionalData;
+
+public:
+
+	FShaderCode()
+	: OptionalDataSize(0)
+	{
+	}
+
+	// adds CustomData or does nothing if that was already done before
+	void FinalizeShaderCode() const
+	{
+		if(OptionalDataSize != -1)
+		{
+			OptionalDataSize += sizeof(uint8) + sizeof(uint8);
+
+			check(OptionalDataSize <= 0xffff);
+
+			ShaderCodeWithOptionalData.Add(OptionalDataSize >> 8);
+			ShaderCodeWithOptionalData.Add(OptionalDataSize);
+			OptionalDataSize = -1;
+		}
+	}
+
+	// for write access
+	TArray<uint8>& GetWriteAccess()
+	{
+		return ShaderCodeWithOptionalData;
+	}
+
+	int32 GetShaderCodeSize() const
+	{
+		FinalizeShaderCode();
+
+		FShaderCodeReader Wrapper(ShaderCodeWithOptionalData);
+
+		return Wrapper.GetShaderCodeSize();
+	}
+
+	// inefficient, will/should be replaced by GetShaderCodeToRead()
+	void GetShaderCodeLegacy(TArray<uint8>& Out) const
+	{
+		Out.Empty();
+
+		Out.AddUninitialized(GetShaderCodeSize());
+		FMemory::Memcpy(Out.GetData(), GetReadAccess().GetData(), ShaderCodeWithOptionalData.Num());
+	}
+
+	// for read access, can have additional data attached to the end
+	const TArray<uint8>& GetReadAccess() const
+	{
+		FinalizeShaderCode();
+
+		return ShaderCodeWithOptionalData;
+	}
+
+	// for convenience
+	template <class T>
+	void AddOptionalData(const T &In)
+	{
+		AddOptionalData(T::Key, (uint8*)&In, sizeof(T));
+	}
+	
+	// Note: we don't hash the optional attachments in GenerateOutputHash() as they would prevent sharing (e.g. many material share the save VS)
+	// can be called after the non optional data was stored in ShaderData
+	// @param Key uint8 to save memory so max 255, e.g. FShaderCodePackedResourceCounts::Key
+	// @param Size max 255, >0, check if too large
+	void AddOptionalData(uint8 Key, const uint8* ValuePtr, uint32 ValueSize)
+	{
+		check(ValueSize <= 255);
+		check(ValuePtr);
+
+		// don't add after Finalize happened
+		check(OptionalDataSize >= 0);
+
+		ShaderCodeWithOptionalData.Add(Key);
+		ShaderCodeWithOptionalData.Add(ValueSize);
+		ShaderCodeWithOptionalData.Append(ValuePtr, ValueSize);
+		OptionalDataSize += sizeof(uint8) + sizeof(uint8) + (uint32)ValueSize;
+	}
+
+	// Note: we don't hash the optional attachments in GenerateOutputHash() as they would prevent sharing (e.g. many material share the save VS)
+	// convenience, silently drops the data if string is too long 
+	// @param e.g. 'n' for the ShaderSourceFileName
+	void AddOptionalData(uint8 Key, const ANSICHAR* InString)
+	{
+		uint32 Size = FCStringAnsi::Strlen(InString) + 1;
+
+		if(Size < 255)
+		{
+			AddOptionalData(Key, (uint8*)InString, Size);
+		}
+	}
+
+	friend FArchive& operator<<(FArchive& Ar, FShaderCode& Output)
+	{
+		if(Ar.IsLoading())
+		{
+			Output.OptionalDataSize = -1;
+		}
+		else
+		{
+			Output.FinalizeShaderCode();
+		}
+
+		// Note: this serialize is used to pass between UE4 and the shader compile worker, recompile both when modifying
+		Ar << Output.ShaderCodeWithOptionalData;
+		return Ar;
+	}
+};
+
 /** The output of the shader compiler. */
 struct FShaderCompilerOutput
 {
@@ -500,7 +757,7 @@ struct FShaderCompilerOutput
 	FShaderParameterMap ParameterMap;
 	TArray<FShaderCompilerError> Errors;
 	FShaderTarget Target;
-	TArray<uint8> Code;
+	FShaderCode ShaderCode;
 	FSHAHash OutputHash;
 	uint32 NumInstructions;
 	uint32 NumTextureSamplers;
@@ -512,7 +769,7 @@ struct FShaderCompilerOutput
 	friend FArchive& operator<<(FArchive& Ar,FShaderCompilerOutput& Output)
 	{
 		// Note: this serialize is used to pass between UE4 and the shader compile worker, recompile both when modifying
-		return Ar << Output.ParameterMap << Output.Errors << Output.Target << Output.Code << Output.NumInstructions << Output.NumTextureSamplers << Output.bSucceeded;
+		return Ar << Output.ParameterMap << Output.Errors << Output.Target << Output.ShaderCode << Output.NumInstructions << Output.NumTextureSamplers << Output.bSucceeded;
 	}
 };
 
diff --git a/Engine/Source/Runtime/Windows/D3D11RHI/Private/D3D11Commands.cpp b/Engine/Source/Runtime/Windows/D3D11RHI/Private/D3D11Commands.cpp
--- a/Engine/Source/Runtime/Windows/D3D11RHI/Private/D3D11Commands.cpp
+++ b/Engine/Source/Runtime/Windows/D3D11RHI/Private/D3D11Commands.cpp
@@ -1265,7 +1265,23 @@ void FD3D11DynamicRHI::SetResourcesFromTables(const ShaderType* RESTRICT Shader)
 		FD3D11UniformBuffer* Buffer = (FD3D11UniformBuffer*)BoundUniformBuffers[ShaderType::StaticFrequency][BufferIndex].GetReference();
 		check(Buffer);
 		check(BufferIndex < Shader->ShaderResourceTable.ResourceTableLayoutHashes.Num());
-		check(Buffer->GetLayout().GetHash() == Shader->ShaderResourceTable.ResourceTableLayoutHashes[BufferIndex]);
+
+#if !(UE_BUILD_SHIPPING || UE_BUILD_TEST)
+		// to track down OR-7159 CRASH: Client crashed at start of match in D3D11Commands.cpp
+		{
+			if (Buffer->GetLayout().GetHash() != Shader->ShaderResourceTable.ResourceTableLayoutHashes[BufferIndex])
+			{
+				FString DebugName = Buffer->GetLayout().GetDebugName().GetPlainNameString();
+				const FString& ShaderName = Shader->ShaderName;
+				
+				UE_LOG(LogD3D11RHI, Error, TEXT("SetResourcesFromTables upcoming check() info Layout='%s' Shader='%s'"), *DebugName, *ShaderName);
+
+				// this might mean you are accessing a data you haven't bound e.g. GBuffer
+				check(Buffer->GetLayout().GetHash() == Shader->ShaderResourceTable.ResourceTableLayoutHashes[BufferIndex]);
+			}
+		}
+#endif
+
 		Buffer->CacheResources(ResourceTableFrameCounter);
 
 		// todo: could make this two pass: gather then set
diff --git a/Engine/Source/Runtime/Windows/D3D11RHI/Private/D3D11Shaders.cpp b/Engine/Source/Runtime/Windows/D3D11RHI/Private/D3D11Shaders.cpp
--- a/Engine/Source/Runtime/Windows/D3D11RHI/Private/D3D11Shaders.cpp
+++ b/Engine/Source/Runtime/Windows/D3D11RHI/Private/D3D11Shaders.cpp
@@ -8,7 +8,7 @@
 
 FVertexShaderRHIRef FD3D11DynamicRHI::RHICreateVertexShader(const TArray<uint8>& Code)
 {
-	check(Code.Num());
+	FShaderCodeReader ShaderCode(Code);
 
 	FD3D11VertexShader* Shader = new FD3D11VertexShader;
 
@@ -16,12 +16,16 @@ FVertexShaderRHIRef FD3D11DynamicRHI::RHICreateVertexShader(const TArray<uint8>&
 	Ar << Shader->ShaderResourceTable;
 	int32 Offset = Ar.Tell();
 	const uint8* CodePtr = Code.GetData() + Offset;
-	const size_t CodeSize = Code.Num() - Offset - 5;
+	const size_t CodeSize = ShaderCode.GetActualShaderCodeSize() - Offset;
 
 	VERIFYD3D11RESULT( Direct3DDevice->CreateVertexShader( (void*)CodePtr, CodeSize, NULL, Shader->Resource.GetInitReference() ) );
-	
-	// bGlobalUniformBufferUsed and resource counts are packed in the last few bytes, see CompileD3D11Shader
-	Shader->bShaderNeedsGlobalConstantBuffer = Code[Code.Num() - 5] != 0;
+
+	auto PackedResourceCounts = ShaderCode.FindOptionalData<FShaderCodePackedResourceCounts>();
+	check(PackedResourceCounts);
+	Shader->bShaderNeedsGlobalConstantBuffer = PackedResourceCounts->bGlobalUniformBufferUsed;
+#if !(UE_BUILD_SHIPPING || UE_BUILD_TEST)
+	Shader->ShaderName = ShaderCode.FindOptionalData('n');
+#endif
 	
 	// TEMP
 	Shader->Code = Code;
@@ -32,7 +36,7 @@ FVertexShaderRHIRef FD3D11DynamicRHI::RHICreateVertexShader(const TArray<uint8>&
 
 FGeometryShaderRHIRef FD3D11DynamicRHI::RHICreateGeometryShader(const TArray<uint8>& Code) 
 { 
-	check(Code.Num());
+	FShaderCodeReader ShaderCode(Code);
 
 	FD3D11GeometryShader* Shader = new FD3D11GeometryShader;
 
@@ -40,19 +44,23 @@ FGeometryShaderRHIRef FD3D11DynamicRHI::RHICreateGeometryShader(const TArray<uin
 	Ar << Shader->ShaderResourceTable;
 	int32 Offset = Ar.Tell();
 	const uint8* CodePtr = Code.GetData() + Offset;
-	const size_t CodeSize = Code.Num() - Offset - 5;
+	const size_t CodeSize = ShaderCode.GetActualShaderCodeSize() - Offset;
 
 	VERIFYD3D11RESULT( Direct3DDevice->CreateGeometryShader( (void*)CodePtr, CodeSize, NULL, Shader->Resource.GetInitReference() ) );
 	
-	// bGlobalUniformBufferUsed and resource counts are packed in the last few bytes, see CompileD3D11Shader
-	Shader->bShaderNeedsGlobalConstantBuffer = Code[Code.Num() - 5] != 0;
-
+	auto PackedResourceCounts = ShaderCode.FindOptionalData<FShaderCodePackedResourceCounts>();
+	check(PackedResourceCounts);
+	Shader->bShaderNeedsGlobalConstantBuffer = PackedResourceCounts->bGlobalUniformBufferUsed;
+#if !(UE_BUILD_SHIPPING || UE_BUILD_TEST)
+	Shader->ShaderName = ShaderCode.FindOptionalData('n');
+#endif
+	
 	return Shader;
 }
 
 FGeometryShaderRHIRef FD3D11DynamicRHI::RHICreateGeometryShaderWithStreamOutput(const TArray<uint8>& Code, const FStreamOutElementList& ElementList, uint32 NumStrides, const uint32* Strides, int32 RasterizedStream) 
 { 
-	check(Code.Num());
+	FShaderCodeReader ShaderCode(Code);
 
 	FD3D11GeometryShader* Shader = new FD3D11GeometryShader;
 
@@ -60,7 +68,7 @@ FGeometryShaderRHIRef FD3D11DynamicRHI::RHICreateGeometryShaderWithStreamOutput(
 	Ar << Shader->ShaderResourceTable;
 	int32 Offset = Ar.Tell();
 	const uint8* CodePtr = Code.GetData() + Offset;
-	const size_t CodeSize = Code.Num() - Offset - 5;
+	const size_t CodeSize = ShaderCode.GetActualShaderCodeSize() - Offset;
 
 	uint32 D3DRasterizedStream = RasterizedStream;
 	if (RasterizedStream == -1)
@@ -92,15 +100,19 @@ FGeometryShaderRHIRef FD3D11DynamicRHI::RHICreateGeometryShaderWithStreamOutput(
 		Shader->Resource.GetInitReference()
 	) );
 	
-	// bGlobalUniformBufferUsed and resource counts are packed in the last few bytes, see CompileD3D11Shader
-	Shader->bShaderNeedsGlobalConstantBuffer = Code[Code.Num() - 5] != 0;
+	auto PackedResourceCounts = ShaderCode.FindOptionalData<FShaderCodePackedResourceCounts>();
+	check(PackedResourceCounts);
+	Shader->bShaderNeedsGlobalConstantBuffer = PackedResourceCounts->bGlobalUniformBufferUsed;
+#if !(UE_BUILD_SHIPPING || UE_BUILD_TEST)
+	Shader->ShaderName = ShaderCode.FindOptionalData('n');
+#endif
 
 	return Shader;
 }
 
 FHullShaderRHIRef FD3D11DynamicRHI::RHICreateHullShader(const TArray<uint8>& Code) 
 { 
-	check(Code.Num());
+	FShaderCodeReader ShaderCode(Code);
 
 	FD3D11HullShader* Shader = new FD3D11HullShader;
 
@@ -108,19 +120,23 @@ FHullShaderRHIRef FD3D11DynamicRHI::RHICreateHullShader(const TArray<uint8>& Cod
 	Ar << Shader->ShaderResourceTable;
 	int32 Offset = Ar.Tell();
 	const uint8* CodePtr = Code.GetData() + Offset;
-	const size_t CodeSize = Code.Num() - Offset - 5;
+	const size_t CodeSize = ShaderCode.GetActualShaderCodeSize() - Offset;
 
 	VERIFYD3D11RESULT( Direct3DDevice->CreateHullShader( (void*)CodePtr, CodeSize, NULL, Shader->Resource.GetInitReference() ) );
 	
-	// bGlobalUniformBufferUsed and resource counts are packed in the last few bytes, see CompileD3D11Shader
-	Shader->bShaderNeedsGlobalConstantBuffer = Code[Code.Num() - 5] != 0;
+	auto PackedResourceCounts = ShaderCode.FindOptionalData<FShaderCodePackedResourceCounts>();
+	check(PackedResourceCounts);
+	Shader->bShaderNeedsGlobalConstantBuffer = PackedResourceCounts->bGlobalUniformBufferUsed;
+#if !(UE_BUILD_SHIPPING || UE_BUILD_TEST)
+	Shader->ShaderName = ShaderCode.FindOptionalData('n');
+#endif
 
 	return Shader;
 }
 
 FDomainShaderRHIRef FD3D11DynamicRHI::RHICreateDomainShader(const TArray<uint8>& Code) 
 { 
-	check(Code.Num());
+	FShaderCodeReader ShaderCode(Code);
 
 	FD3D11DomainShader* Shader = new FD3D11DomainShader;
 
@@ -128,19 +144,23 @@ FDomainShaderRHIRef FD3D11DynamicRHI::RHICreateDomainShader(const TArray<uint8>&
 	Ar << Shader->ShaderResourceTable;
 	int32 Offset = Ar.Tell();
 	const uint8* CodePtr = Code.GetData() + Offset;
-	const size_t CodeSize = Code.Num() - Offset - 5;
+	const size_t CodeSize = ShaderCode.GetActualShaderCodeSize() - Offset;
 
 	VERIFYD3D11RESULT( Direct3DDevice->CreateDomainShader( (void*)CodePtr, CodeSize, NULL, Shader->Resource.GetInitReference() ) );
 	
-	// bGlobalUniformBufferUsed and resource counts are packed in the last few bytes, see CompileD3D11Shader
-	Shader->bShaderNeedsGlobalConstantBuffer = Code[Code.Num() - 5] != 0;
+	auto PackedResourceCounts = ShaderCode.FindOptionalData<FShaderCodePackedResourceCounts>();
+	check(PackedResourceCounts);
+	Shader->bShaderNeedsGlobalConstantBuffer = PackedResourceCounts->bGlobalUniformBufferUsed;
+#if !(UE_BUILD_SHIPPING || UE_BUILD_TEST)
+	Shader->ShaderName = ShaderCode.FindOptionalData('n');
+#endif
 
 	return Shader;
 }
 
 FPixelShaderRHIRef FD3D11DynamicRHI::RHICreatePixelShader(const TArray<uint8>& Code)
 {
-	check(Code.Num());
+	FShaderCodeReader ShaderCode(Code);
 
 	FD3D11PixelShader* Shader = new FD3D11PixelShader;
 
@@ -148,19 +168,23 @@ FPixelShaderRHIRef FD3D11DynamicRHI::RHICreatePixelShader(const TArray<uint8>& C
 	Ar << Shader->ShaderResourceTable;
 	int32 Offset = Ar.Tell();
 	const uint8* CodePtr = Code.GetData() + Offset;
-	const size_t CodeSize = Code.Num() - Offset - 5;
+	const size_t CodeSize = ShaderCode.GetActualShaderCodeSize() - Offset;
 
 	VERIFYD3D11RESULT( Direct3DDevice->CreatePixelShader( (void*)CodePtr, CodeSize, NULL, Shader->Resource.GetInitReference() ) );
 	
-	// bGlobalUniformBufferUsed and resource counts are packed in the last few bytes, see CompileD3D11Shader
-	Shader->bShaderNeedsGlobalConstantBuffer = Code[Code.Num() - 5] != 0;
+	auto PackedResourceCounts = ShaderCode.FindOptionalData<FShaderCodePackedResourceCounts>();
+	check(PackedResourceCounts);
+	Shader->bShaderNeedsGlobalConstantBuffer = PackedResourceCounts->bGlobalUniformBufferUsed;
+#if !(UE_BUILD_SHIPPING || UE_BUILD_TEST)
+	Shader->ShaderName = ShaderCode.FindOptionalData('n');
+#endif
 
 	return Shader;
 }
 
 FComputeShaderRHIRef FD3D11DynamicRHI::RHICreateComputeShader(const TArray<uint8>& Code) 
 { 
-	check(Code.Num());
+	FShaderCodeReader ShaderCode(Code);
 
 	FD3D11ComputeShader* Shader = new FD3D11ComputeShader;
 
@@ -168,12 +192,16 @@ FComputeShaderRHIRef FD3D11DynamicRHI::RHICreateComputeShader(const TArray<uint8
 	Ar << Shader->ShaderResourceTable;
 	int32 Offset = Ar.Tell();
 	const uint8* CodePtr = Code.GetData() + Offset;
-	const size_t CodeSize = Code.Num() - Offset - 5;
+	const size_t CodeSize = ShaderCode.GetActualShaderCodeSize() - Offset;
 
 	VERIFYD3D11RESULT( Direct3DDevice->CreateComputeShader( (void*)CodePtr, CodeSize, NULL, Shader->Resource.GetInitReference() ) );
 	
-	// bGlobalUniformBufferUsed and resource counts are packed in the last few bytes, see CompileD3D11Shader
-	Shader->bShaderNeedsGlobalConstantBuffer = Code[Code.Num() - 5] != 0;
+	auto PackedResourceCounts = ShaderCode.FindOptionalData<FShaderCodePackedResourceCounts>();
+	check(PackedResourceCounts);
+	Shader->bShaderNeedsGlobalConstantBuffer = PackedResourceCounts->bGlobalUniformBufferUsed;
+#if !(UE_BUILD_SHIPPING || UE_BUILD_TEST)
+	Shader->ShaderName = ShaderCode.FindOptionalData('n');
+#endif
 
 	return Shader;
 }
@@ -214,11 +242,13 @@ FD3D11BoundShaderState::FD3D11BoundShaderState(
 	D3D11_INPUT_ELEMENT_DESC NullInputElement;
 	FMemory::Memzero(&NullInputElement,sizeof(D3D11_INPUT_ELEMENT_DESC));
 
+	FShaderCodeReader VertexShaderCode(InVertexShader->Code);
+
 	VERIFYD3D11RESULT(Direct3DDevice->CreateInputLayout(
 		InVertexDeclaration ? InVertexDeclaration->VertexElements.GetData() : &NullInputElement,
 		InVertexDeclaration ? InVertexDeclaration->VertexElements.Num() : 0,
 		&InVertexShader->Code[ InVertexShader->Offset ],			// TEMP ugly
-		InVertexShader->Code.Num() - 5 - InVertexShader->Offset,
+		VertexShaderCode.GetActualShaderCodeSize() - InVertexShader->Offset,
 		InputLayout.GetInitReference()
 		));
 
diff --git a/Engine/Source/Runtime/Windows/D3D11RHI/Public/D3D11Resources.h b/Engine/Source/Runtime/Windows/D3D11RHI/Public/D3D11Resources.h
--- a/Engine/Source/Runtime/Windows/D3D11RHI/Public/D3D11Resources.h
+++ b/Engine/Source/Runtime/Windows/D3D11RHI/Public/D3D11Resources.h
@@ -43,7 +43,7 @@ public:
 
 	FD3D11ShaderResourceTable ShaderResourceTable;
 
-	/** The vertex shader's bytecode, with custom data in the last byte. */
+	/** The vertex shader's bytecode, with custom data attached. */
 	TArray<uint8> Code;
 
 	// TEMP remove with removal of bound shader state
diff --git a/Engine/Source/Runtime/Windows/D3D12RHI/Private/D3D12Commands.cpp b/Engine/Source/Runtime/Windows/D3D12RHI/Private/D3D12Commands.cpp
--- a/Engine/Source/Runtime/Windows/D3D12RHI/Private/D3D12Commands.cpp
+++ b/Engine/Source/Runtime/Windows/D3D12RHI/Private/D3D12Commands.cpp
@@ -193,7 +193,7 @@ void FD3D12CommandContext::RHIDispatchComputeShader(uint32 ThreadGroupCountX, ui
 
 	OwningRHI.RegisterGPUWork(1);
 
-	if (ComputeShader->bShaderNeedsGlobalConstantBuffer)
+	if (ComputeShader->ResourceCounts.bGlobalUniformBufferUsed)
 	{
 		CommitComputeShaderConstants();
 	}
@@ -220,7 +220,7 @@ void FD3D12CommandContext::RHIDispatchIndirectComputeShader(FVertexBufferRHIPara
 
 	StateCache.SetComputeShader(ComputeShader);
 
-	if (ComputeShader->bShaderNeedsGlobalConstantBuffer)
+	if (ComputeShader->ResourceCounts.bGlobalUniformBufferUsed)
 	{
 		CommitComputeShaderConstants();
 	}
diff --git a/Engine/Source/Runtime/Windows/D3D12RHI/Private/D3D12Shaders.cpp b/Engine/Source/Runtime/Windows/D3D12RHI/Private/D3D12Shaders.cpp
--- a/Engine/Source/Runtime/Windows/D3D12RHI/Private/D3D12Shaders.cpp
+++ b/Engine/Source/Runtime/Windows/D3D12RHI/Private/D3D12Shaders.cpp
@@ -8,23 +8,23 @@
 
 FVertexShaderRHIRef FD3D12DynamicRHI::RHICreateVertexShader(const TArray<uint8>& Code)
 {
-	check(Code.Num());
+	FShaderCodeReader ShaderCode(Code);
+
 	FD3D12VertexShader* Shader = new FD3D12VertexShader;
 
 	FMemoryReader Ar(Code, true);
 	Ar << Shader->ShaderResourceTable;
 	int32 Offset = Ar.Tell();
 	const uint8* CodePtr = Code.GetData() + Offset;
-	const SIZE_T CodeSize = Code.Num() - Offset - 5;
-
-	// bGlobalUniformBufferUsed and resource counts are packed in the last few bytes, see CompileD3D11Shader
-	Shader->bShaderNeedsGlobalConstantBuffer = Code[Code.Num() - 5] != 0;
-	Shader->SamplerCount                     = Code[Code.Num() - 4];
-	Shader->SRVCount                         = Code[Code.Num() - 3];
-	Shader->UAVCount                         = Code[Code.Num() - 1];
+	const SIZE_T CodeSize = ShaderCode.GetActualShaderCodeSize() - Offset;
 
+	Shader->ResourceCounts = *ShaderCode.FindOptionalData<FShaderCodePackedResourceCounts>();
 	//TODO Special case for Constant Buffer count as it currently needs all buffers set if it needs any set.
-	Shader->CBCount = (Code[Code.Num() - 2]) ? MAX_CBS : 0;
+	Shader->ResourceCounts.NumCBs = (Shader->ResourceCounts.NumCBs > 0) ? MAX_CBS : 0;
+
+#if !(UE_BUILD_SHIPPING || UE_BUILD_TEST)
+	Shader->ShaderName = ShaderCode.FindOptionalData('n');
+#endif
 
 	Shader->Code = Code;
 	Shader->Offset = Offset;
@@ -39,7 +39,7 @@ FVertexShaderRHIRef FD3D12DynamicRHI::RHICreateVertexShader(const TArray<uint8>&
 
 FPixelShaderRHIRef FD3D12DynamicRHI::RHICreatePixelShader(const TArray<uint8>& Code)
 {
-	check(Code.Num());
+	FShaderCodeReader ShaderCode(Code);
 
 	FD3D12PixelShader* Shader = new FD3D12PixelShader;
 
@@ -47,16 +47,15 @@ FPixelShaderRHIRef FD3D12DynamicRHI::RHICreatePixelShader(const TArray<uint8>& C
 	Ar << Shader->ShaderResourceTable;
 	int32 Offset = Ar.Tell();
 	const uint8* CodePtr = Code.GetData() + Offset;
-	const SIZE_T CodeSize = Code.Num() - Offset - 5;
-
-	// bGlobalUniformBufferUsed and resource counts are packed in the last few bytes, see CompileD3D11Shader
-	Shader->bShaderNeedsGlobalConstantBuffer = Code[Code.Num() - 5] != 0;
-	Shader->SamplerCount                     = Code[Code.Num() - 4];
-	Shader->SRVCount                         = Code[Code.Num() - 3];
-	Shader->UAVCount                         = Code[Code.Num() - 1];
+	const SIZE_T CodeSize = ShaderCode.GetActualShaderCodeSize() - Offset;
 
+	Shader->ResourceCounts = *ShaderCode.FindOptionalData<FShaderCodePackedResourceCounts>();
 	//TODO Special case for Constant Buffer count as it currently needs all buffers set if it needs any set.
-	Shader->CBCount = (Code[Code.Num() - 2]) ? MAX_CBS : 0;
+	Shader->ResourceCounts.NumCBs = (Shader->ResourceCounts.NumCBs > 0) ? MAX_CBS : 0;
+
+#if !(UE_BUILD_SHIPPING || UE_BUILD_TEST)
+	Shader->ShaderName = ShaderCode.FindOptionalData('n');
+#endif
 
 	Shader->Code = Code;
 
@@ -70,7 +69,7 @@ FPixelShaderRHIRef FD3D12DynamicRHI::RHICreatePixelShader(const TArray<uint8>& C
 
 FHullShaderRHIRef FD3D12DynamicRHI::RHICreateHullShader(const TArray<uint8>& Code)
 {
-	check(Code.Num());
+	FShaderCodeReader ShaderCode(Code);
 
 	FD3D12HullShader* Shader = new FD3D12HullShader;
 
@@ -78,16 +77,15 @@ FHullShaderRHIRef FD3D12DynamicRHI::RHICreateHullShader(const TArray<uint8>& Cod
 	Ar << Shader->ShaderResourceTable;
 	int32 Offset = Ar.Tell();
 	const uint8* CodePtr = Code.GetData() + Offset;
-	const SIZE_T CodeSize = Code.Num() - Offset - 5;
-
-	// bGlobalUniformBufferUsed and resource counts are packed in the last few bytes, see CompileD3D11Shader
-	Shader->bShaderNeedsGlobalConstantBuffer = Code[Code.Num() - 5] != 0;
-	Shader->SamplerCount                     = Code[Code.Num() - 4];
-	Shader->SRVCount                         = Code[Code.Num() - 3];
-	Shader->UAVCount                         = Code[Code.Num() - 1];
+	const SIZE_T CodeSize = ShaderCode.GetActualShaderCodeSize() - Offset;
 
+	Shader->ResourceCounts = *ShaderCode.FindOptionalData<FShaderCodePackedResourceCounts>();
 	//TODO Special case for Constant Buffer count as it currently needs all buffers set if it needs any set.
-	Shader->CBCount = (Code[Code.Num() - 2]) ? MAX_CBS : 0;
+	Shader->ResourceCounts.NumCBs = (Shader->ResourceCounts.NumCBs > 0) ? MAX_CBS : 0;
+
+#if !(UE_BUILD_SHIPPING || UE_BUILD_TEST)
+	Shader->ShaderName = ShaderCode.FindOptionalData('n');
+#endif
 
 	Shader->Code = Code;
 
@@ -101,7 +99,7 @@ FHullShaderRHIRef FD3D12DynamicRHI::RHICreateHullShader(const TArray<uint8>& Cod
 
 FDomainShaderRHIRef FD3D12DynamicRHI::RHICreateDomainShader(const TArray<uint8>& Code)
 {
-	check(Code.Num());
+	FShaderCodeReader ShaderCode(Code);
 
 	FD3D12DomainShader* Shader = new FD3D12DomainShader;
 
@@ -109,16 +107,15 @@ FDomainShaderRHIRef FD3D12DynamicRHI::RHICreateDomainShader(const TArray<uint8>&
 	Ar << Shader->ShaderResourceTable;
 	int32 Offset = Ar.Tell();
 	const uint8* CodePtr = Code.GetData() + Offset;
-	const SIZE_T CodeSize = Code.Num() - Offset - 5;
-
-	// bGlobalUniformBufferUsed and resource counts are packed in the last few bytes, see CompileD3D11Shader
-	Shader->bShaderNeedsGlobalConstantBuffer = Code[Code.Num() - 5] != 0;
-	Shader->SamplerCount                     = Code[Code.Num() - 4];
-	Shader->SRVCount                         = Code[Code.Num() - 3];
-	Shader->UAVCount                         = Code[Code.Num() - 1];
+	const SIZE_T CodeSize = ShaderCode.GetActualShaderCodeSize() - Offset;
 
+	Shader->ResourceCounts = *ShaderCode.FindOptionalData<FShaderCodePackedResourceCounts>();
 	//TODO Special case for Constant Buffer count as it currently needs all buffers set if it needs any set.
-	Shader->CBCount = (Code[Code.Num() - 2] > 0) ? MAX_CBS : 0;
+	Shader->ResourceCounts.NumCBs = (Shader->ResourceCounts.NumCBs > 0) ? MAX_CBS : 0;
+
+#if !(UE_BUILD_SHIPPING || UE_BUILD_TEST)
+	Shader->ShaderName = ShaderCode.FindOptionalData('n');
+#endif
 
 	Shader->Code = Code;
 
@@ -132,7 +129,7 @@ FDomainShaderRHIRef FD3D12DynamicRHI::RHICreateDomainShader(const TArray<uint8>&
 
 FGeometryShaderRHIRef FD3D12DynamicRHI::RHICreateGeometryShader(const TArray<uint8>& Code)
 {
-	check(Code.Num());
+	FShaderCodeReader ShaderCode(Code);
 
 	FD3D12GeometryShader* Shader = new FD3D12GeometryShader;
 
@@ -140,16 +137,15 @@ FGeometryShaderRHIRef FD3D12DynamicRHI::RHICreateGeometryShader(const TArray<uin
 	Ar << Shader->ShaderResourceTable;
 	int32 Offset = Ar.Tell();
 	const uint8* CodePtr = Code.GetData() + Offset;
-	const SIZE_T CodeSize = Code.Num() - Offset - 5;
-
-	// bGlobalUniformBufferUsed and resource counts are packed in the last few bytes, see CompileD3D11Shader
-	Shader->bShaderNeedsGlobalConstantBuffer = Code[Code.Num() - 5] != 0;
-	Shader->SamplerCount                     = Code[Code.Num() - 4];
-	Shader->SRVCount                         = Code[Code.Num() - 3];
-	Shader->UAVCount                         = Code[Code.Num() - 1];
+	const SIZE_T CodeSize = ShaderCode.GetActualShaderCodeSize() - Offset;
 
+	Shader->ResourceCounts = *ShaderCode.FindOptionalData<FShaderCodePackedResourceCounts>();
 	//TODO Special case for Constant Buffer count as it currently needs all buffers set if it needs any set.
-	Shader->CBCount = (Code[Code.Num() - 2] > 0) ? MAX_CBS : 0;
+	Shader->ResourceCounts.NumCBs = (Shader->ResourceCounts.NumCBs > 0) ? MAX_CBS : 0;
+
+#if !(UE_BUILD_SHIPPING || UE_BUILD_TEST)
+	Shader->ShaderName = ShaderCode.FindOptionalData('n');
+#endif
 
 	Shader->Code = Code;
 
@@ -163,7 +159,7 @@ FGeometryShaderRHIRef FD3D12DynamicRHI::RHICreateGeometryShader(const TArray<uin
 
 FGeometryShaderRHIRef FD3D12DynamicRHI::RHICreateGeometryShaderWithStreamOutput(const TArray<uint8>& Code, const FStreamOutElementList& ElementList, uint32 NumStrides, const uint32* Strides, int32 RasterizedStream)
 {
-	check(Code.Num());
+	FShaderCodeReader ShaderCode(Code);
 
 	FD3D12GeometryShader* Shader = new FD3D12GeometryShader;
 
@@ -171,7 +167,7 @@ FGeometryShaderRHIRef FD3D12DynamicRHI::RHICreateGeometryShaderWithStreamOutput(
 	Ar << Shader->ShaderResourceTable;
 	int32 Offset = Ar.Tell();
 	const uint8* CodePtr = Code.GetData() + Offset;
-	const SIZE_T CodeSize = Code.Num() - Offset - 5;
+	const SIZE_T CodeSize = ShaderCode.GetActualShaderCodeSize() - Offset;
 
 	Shader->StreamOutput.RasterizedStream = RasterizedStream;
 	if (RasterizedStream == -1)
@@ -196,15 +192,13 @@ FGeometryShaderRHIRef FD3D12DynamicRHI::RHICreateGeometryShaderWithStreamOutput(
 		Shader->pStreamOutEntries[EntryIndex].OutputSlot = ElementList[EntryIndex].OutputSlot;
 	}
 
-
-	// bGlobalUniformBufferUsed and resource counts are packed in the last few bytes, see CompileD3D11Shader
-	Shader->bShaderNeedsGlobalConstantBuffer = Code[Code.Num() - 5] != 0;
-	Shader->SamplerCount                     = Code[Code.Num() - 4];
-	Shader->SRVCount                         = Code[Code.Num() - 3];
-	Shader->UAVCount                         = Code[Code.Num() - 1];
-
+	Shader->ResourceCounts = *ShaderCode.FindOptionalData<FShaderCodePackedResourceCounts>();
 	//TODO Special case for Constant Buffer count as it currently needs all buffers set if it needs any set.
-	Shader->CBCount = (Code[Code.Num() - 2] > 0) ? MAX_CBS : 0;
+	Shader->ResourceCounts.NumCBs = (Shader->ResourceCounts.NumCBs > 0) ? MAX_CBS : 0;
+
+#if !(UE_BUILD_SHIPPING || UE_BUILD_TEST)
+	Shader->ShaderName = ShaderCode.FindOptionalData('n');
+#endif
 
 	// Indicate this shader uses stream output
 	Shader->bShaderNeedsStreamOutput = true;
@@ -231,7 +225,7 @@ FGeometryShaderRHIRef FD3D12DynamicRHI::RHICreateGeometryShaderWithStreamOutput(
 
 FComputeShaderRHIRef FD3D12DynamicRHI::RHICreateComputeShader(const TArray<uint8>& Code)
 {
-	check(Code.Num());
+	FShaderCodeReader ShaderCode(Code);
 
 	FD3D12ComputeShader* Shader = new FD3D12ComputeShader;
 
@@ -239,16 +233,15 @@ FComputeShaderRHIRef FD3D12DynamicRHI::RHICreateComputeShader(const TArray<uint8
 	Ar << Shader->ShaderResourceTable;
 	int32 Offset = Ar.Tell();
 	const uint8* CodePtr = Code.GetData() + Offset;
-	const SIZE_T CodeSize = Code.Num() - Offset - 5;
-
-	// bGlobalUniformBufferUsed and resource counts are packed in the last few bytes, see CompileD3D11Shader
-	Shader->bShaderNeedsGlobalConstantBuffer = Code[Code.Num() - 5] != 0;
-	Shader->SamplerCount                     = Code[Code.Num() - 4];
-	Shader->SRVCount                         = Code[Code.Num() - 3];
-	Shader->UAVCount                         = Code[Code.Num() - 1];
+	const SIZE_T CodeSize = ShaderCode.GetActualShaderCodeSize() - Offset;
 
+	Shader->ResourceCounts = *ShaderCode.FindOptionalData<FShaderCodePackedResourceCounts>();
 	//TODO Special case for Constant Buffer count as it currently needs all buffers set if it needs any set.
-	Shader->CBCount = (Code[Code.Num() - 2] > 0) ? MAX_CBS : 0;
+	Shader->ResourceCounts.NumCBs = (Shader->ResourceCounts.NumCBs > 0) ? MAX_CBS : 0;
+
+#if !(UE_BUILD_SHIPPING || UE_BUILD_TEST)
+	Shader->ShaderName = ShaderCode.FindOptionalData('n');
+#endif
 
 	Shader->Code = Code;
 
@@ -297,11 +290,11 @@ FD3D12BoundShaderState::FD3D12BoundShaderState(
 	InputLayout.NumElements = (InVertexDeclaration ? InVertexDeclaration->VertexElements.Num() : 0);
 	InputLayout.pInputElementDescs = (InVertexDeclaration ? InVertexDeclaration->VertexElements.GetData() : nullptr);
 
-	bShaderNeedsGlobalConstantBuffer[SF_Vertex] = InVertexShader->bShaderNeedsGlobalConstantBuffer;
-	bShaderNeedsGlobalConstantBuffer[SF_Hull] = InHullShader ? InHullShader->bShaderNeedsGlobalConstantBuffer : false;
-	bShaderNeedsGlobalConstantBuffer[SF_Domain] = InDomainShader ? InDomainShader->bShaderNeedsGlobalConstantBuffer : false;
-	bShaderNeedsGlobalConstantBuffer[SF_Pixel] = InPixelShader ? InPixelShader->bShaderNeedsGlobalConstantBuffer : false;
-	bShaderNeedsGlobalConstantBuffer[SF_Geometry] = InGeometryShader ? InGeometryShader->bShaderNeedsGlobalConstantBuffer : false;
+	bShaderNeedsGlobalConstantBuffer[SF_Vertex] = InVertexShader->ResourceCounts.bGlobalUniformBufferUsed;
+	bShaderNeedsGlobalConstantBuffer[SF_Hull] = InHullShader ? InHullShader->ResourceCounts.bGlobalUniformBufferUsed : false;
+	bShaderNeedsGlobalConstantBuffer[SF_Domain] = InDomainShader ? InDomainShader->ResourceCounts.bGlobalUniformBufferUsed : false;
+	bShaderNeedsGlobalConstantBuffer[SF_Pixel] = InPixelShader ? InPixelShader->ResourceCounts.bGlobalUniformBufferUsed : false;
+	bShaderNeedsGlobalConstantBuffer[SF_Geometry] = InGeometryShader ? InGeometryShader->ResourceCounts.bGlobalUniformBufferUsed : false;
 
 	static_assert(ARRAY_COUNT(bShaderNeedsGlobalConstantBuffer) == SF_NumFrequencies, "EShaderFrequency size should match with array count of bShaderNeedsGlobalConstantBuffer.");
 
diff --git a/Engine/Source/Runtime/Windows/D3D12RHI/Private/D3D12StateCachePrivate.h b/Engine/Source/Runtime/Windows/D3D12RHI/Private/D3D12StateCachePrivate.h
--- a/Engine/Source/Runtime/Windows/D3D12RHI/Private/D3D12StateCachePrivate.h
+++ b/Engine/Source/Runtime/Windows/D3D12RHI/Private/D3D12StateCachePrivate.h
@@ -1053,10 +1053,10 @@ protected:
 		TShader* OldShader = Traits::GetShader(BSS());
 		if (OldShader != Shader)
 		{
-			PipelineState.Common.CurrentShaderSamplerCounts[Traits::Frequency] = (Shader) ? Shader->SamplerCount : 0;
-			PipelineState.Common.CurrentShaderSRVCounts[Traits::Frequency]     = (Shader) ? Shader->SRVCount     : 0;
-			PipelineState.Common.CurrentShaderCBCounts[Traits::Frequency]      = (Shader) ? Shader->CBCount      : 0;
-			PipelineState.Common.CurrentShaderUAVCounts[Traits::Frequency]     = (Shader) ? Shader->UAVCount     : 0;
+			PipelineState.Common.CurrentShaderSamplerCounts[Traits::Frequency] = (Shader) ? Shader->ResourceCounts.NumSamplers : 0;
+			PipelineState.Common.CurrentShaderSRVCounts[Traits::Frequency]     = (Shader) ? Shader->ResourceCounts.NumSRVs     : 0;
+			PipelineState.Common.CurrentShaderCBCounts[Traits::Frequency]      = (Shader) ? Shader->ResourceCounts.NumCBs      : 0;
+			PipelineState.Common.CurrentShaderUAVCounts[Traits::Frequency]     = (Shader) ? Shader->ResourceCounts.NumUAVs     : 0;
 		}
 	}
 
@@ -1312,10 +1312,10 @@ public:
 		{
 			PipelineState.Compute.CurrentComputeShader                  = Shader;
 			PipelineState.Compute.bNeedRebuildPSO = true;
-			PipelineState.Common.CurrentShaderSamplerCounts[SF_Compute] = (Shader) ? Shader->SamplerCount : 0;
-			PipelineState.Common.CurrentShaderSRVCounts[SF_Compute]     = (Shader) ? Shader->SRVCount     : 0;
-			PipelineState.Common.CurrentShaderCBCounts[SF_Compute]      = (Shader) ? Shader->CBCount      : 0;
-			PipelineState.Common.CurrentShaderUAVCounts[SF_Compute]     = (Shader) ? Shader->UAVCount     : 0;
+			PipelineState.Common.CurrentShaderSamplerCounts[SF_Compute] = (Shader) ? Shader->ResourceCounts.NumSamplers : 0;
+			PipelineState.Common.CurrentShaderSRVCounts[SF_Compute]     = (Shader) ? Shader->ResourceCounts.NumSRVs     : 0;
+			PipelineState.Common.CurrentShaderCBCounts[SF_Compute]      = (Shader) ? Shader->ResourceCounts.NumCBs      : 0;
+			PipelineState.Common.CurrentShaderUAVCounts[SF_Compute]     = (Shader) ? Shader->ResourceCounts.NumUAVs     : 0;
 		}
 	}
 
diff --git a/Engine/Source/Runtime/Windows/D3D12RHI/Public/D3D12Resources.h b/Engine/Source/Runtime/Windows/D3D12RHI/Public/D3D12Resources.h
--- a/Engine/Source/Runtime/Windows/D3D12RHI/Public/D3D12Resources.h
+++ b/Engine/Source/Runtime/Windows/D3D12RHI/Public/D3D12Resources.h
@@ -931,17 +931,13 @@ public:
 
 	FD3D12ShaderResourceTable ShaderResourceTable;
 
-	/** The vertex shader's bytecode, with custom data in the last byte. */
+	/** The vertex shader's bytecode, with custom data attached. */
 	TArray<uint8> Code;
 
 	// TEMP remove with removal of bound shader state
 	int32 Offset;
 
-	bool bShaderNeedsGlobalConstantBuffer;
-	uint32 SRVCount;
-	uint32 CBCount;
-	uint32 SamplerCount;
-	uint32 UAVCount;
+	FShaderCodePackedResourceCounts ResourceCounts;
 };
 
 class FD3D12GeometryShader : public FRHIGeometryShader
@@ -954,7 +950,7 @@ public:
 
 	FD3D12ShaderResourceTable ShaderResourceTable;
 
-	/** The shader's bytecode, with custom data in the last byte. */
+	/** The shader's bytecode, with custom data attached. */
 	TArray<uint8> Code;
 
 	/** The shader's stream output description. */
@@ -962,12 +958,8 @@ public:
 	D3D12_SO_DECLARATION_ENTRY *pStreamOutEntries;
 	uint32 *pStreamOutStrides;
 
-	bool bShaderNeedsGlobalConstantBuffer;
+	FShaderCodePackedResourceCounts ResourceCounts;
 	bool bShaderNeedsStreamOutput;
-	uint32 SRVCount;
-	uint32 CBCount;
-	uint32 SamplerCount;
-	uint32 UAVCount;
 
 	FD3D12GeometryShader()
 		: bShaderNeedsStreamOutput(false)
@@ -1003,14 +995,10 @@ public:
 
 	FD3D12ShaderResourceTable ShaderResourceTable;
 
-	/** The shader's bytecode, with custom data in the last byte. */
+	/** The shader's bytecode, with custom data attached. */
 	TArray<uint8> Code;
 
-	bool bShaderNeedsGlobalConstantBuffer;
-	uint32 SRVCount;
-	uint32 CBCount;
-	uint32 SamplerCount;
-	uint32 UAVCount;
+	FShaderCodePackedResourceCounts ResourceCounts;
 };
 
 class FD3D12DomainShader : public FRHIDomainShader
@@ -1023,14 +1011,10 @@ public:
 
 	FD3D12ShaderResourceTable ShaderResourceTable;
 
-	/** The shader's bytecode, with custom data in the last byte. */
+	/** The shader's bytecode, with custom data attached. */
 	TArray<uint8> Code;
 
-	bool bShaderNeedsGlobalConstantBuffer;
-	uint32 SRVCount;
-	uint32 CBCount;
-	uint32 SamplerCount;
-	uint32 UAVCount;
+	FShaderCodePackedResourceCounts ResourceCounts;
 };
 
 class FD3D12PixelShader : public FRHIPixelShader
@@ -1041,16 +1025,12 @@ public:
 	/** The shader's bytecode. */
 	FD3D12ShaderBytecode ShaderBytecode;
 
-	/** The shader's bytecode, with custom data in the last byte. */
+	/** The shader's bytecode, with custom data attached. */
 	TArray<uint8> Code;
 
 	FD3D12ShaderResourceTable ShaderResourceTable;
 
-	bool bShaderNeedsGlobalConstantBuffer;
-	uint32 SRVCount;
-	uint32 CBCount;
-	uint32 SamplerCount;
-	uint32 UAVCount;
+	FShaderCodePackedResourceCounts ResourceCounts;
 };
 
 class FD3D12ComputeShader : public FRHIComputeShader
@@ -1061,16 +1041,12 @@ public:
 	/** The shader's bytecode. */
 	FD3D12ShaderBytecode ShaderBytecode;
 
-	/** The shader's bytecode, with custom data in the last byte. */
+	/** The shader's bytecode, with custom data attached. */
 	TArray<uint8> Code;
 
 	FD3D12ShaderResourceTable ShaderResourceTable;
 
-	bool bShaderNeedsGlobalConstantBuffer;
-	uint32 SRVCount;
-	uint32 CBCount;
-	uint32 SamplerCount;
-	uint32 UAVCount;
+	FShaderCodePackedResourceCounts ResourceCounts;
 };
 
 /**
