OR-5944 Implementing stereo spatialization for PS4.

- Refactored the PC stereo spatialization code to move shared data and calculations to platform-independent code
- Implemented stereo spatialization on PS4 using the non A3D api
- More work needs to be done to determine A3D support and stereo spatialization

diff --git a/Engine/Source/Runtime/Engine/Private/Audio.cpp b/Engine/Source/Runtime/Engine/Private/Audio.cpp
--- a/Engine/Source/Runtime/Engine/Private/Audio.cpp
+++ b/Engine/Source/Runtime/Engine/Private/Audio.cpp
@@ -292,6 +292,106 @@ void FSoundSource::SetHighFrequencyGain( void )
 	}
 }
 
+void FSoundSource::UpdateStereoEmitterPositions()
+{
+	// Only call this function if we're told to use spatialization
+	check(WaveInstance->bUseSpatialization);
+	check(Buffer->NumChannels == 2);
+
+	if (WaveInstance->StereoSpread > 0.0f)
+	{
+		// We need to compute the stereo left/right channel positions using the audio component position and the spread 
+		const FVector& ListenerPosition = AudioDevice->Listeners[0].Transform.GetLocation();
+		FVector ListenerToSourceDir = WaveInstance->Location - ListenerPosition;
+		ListenerToSourceDir.Normalize();
+
+		float HalfSpread = 0.5f * WaveInstance->StereoSpread;
+
+		// Get direction of left emitter from true emitter position (left hand rule)
+		FVector LeftEmitterDir = FVector::CrossProduct(ListenerToSourceDir, FVector::UpVector);
+		FVector LeftEmitterOffset = LeftEmitterDir * HalfSpread;
+
+		// Get position vector of left emitter by adding to true emitter the dir scaled by half the spread
+		LeftChannelSourceLocation = WaveInstance->Location + LeftEmitterOffset;
+
+		// Right emitter position is same as right but opposite direction
+		RightChannelSourceLocation = WaveInstance->Location - LeftEmitterOffset;
+	}
+	else
+	{
+		LeftChannelSourceLocation = WaveInstance->Location;
+		RightChannelSourceLocation = WaveInstance->Location;
+	}
+}
+
+void FSoundSource::DrawDebugInfo()
+{
+	// Draw 3d Debug information about this source, if enabled
+	// TODO - Audio Threading. This won't work cross thread.
+	FAudioDeviceManager* DeviceManager = GEngine->GetAudioDeviceManager();
+
+	if (DeviceManager && DeviceManager->IsVisualizeDebug3dEnabled())
+	{
+		UAudioComponent* AudioComponent = (WaveInstance->ActiveSound ? WaveInstance->ActiveSound->GetAudioComponent() : nullptr);
+		if (AudioComponent)
+		{
+			UWorld* SoundWorld = AudioComponent->GetWorld();
+			if (SoundWorld)
+			{
+				FRotator SoundRotation = AudioComponent->GetComponentRotation();
+				DrawDebugCrosshairs(SoundWorld, WaveInstance->Location, SoundRotation, 20.0f, FColor::White, false, -1.0f, SDPG_Foreground);
+
+				FString Name;
+				WaveInstance->ActiveSound->Sound->GetName(Name);
+
+				if (Buffer->NumChannels == 2 && WaveInstance->bUseSpatialization)
+				{
+					DrawDebugCrosshairs(SoundWorld, LeftChannelSourceLocation, SoundRotation, 20.0f, FColor::Red, false, -1.0f, SDPG_Foreground);
+					DrawDebugCrosshairs(SoundWorld, RightChannelSourceLocation, SoundRotation, 20.0f, FColor::Green, false, -1.0f, SDPG_Foreground);
+				}
+
+				DrawDebugString(SoundWorld, AudioComponent->GetComponentLocation() + FVector(0, 0, 32), *Name, nullptr, FColor::White, 0.033, false);
+			}
+		}
+	}
+}
+
+FSpatializationParams FSoundSource::GetSpatializationParams()
+{
+	FSpatializationParams Params;
+
+		// Calculate direction from listener to sound, where the sound is at the origin if unspatialised.
+	if (WaveInstance->bUseSpatialization)
+	{
+		FVector EmitterPosition = AudioDevice->GetListenerTransformedDirection(WaveInstance->Location, &Params.Distance);
+		Params.NormalizedOmniRadius = (Params.Distance > 0) ? (WaveInstance->OmniRadius / Params.Distance) : FLT_MAX;
+
+		if (Buffer->NumChannels == 2)
+		{
+			Params.LeftChannelPosition = AudioDevice->GetListenerTransformedDirection(LeftChannelSourceLocation, nullptr);
+			Params.RightChannelPosition = AudioDevice->GetListenerTransformedDirection(RightChannelSourceLocation, nullptr);
+			Params.EmitterPosition = FVector::ZeroVector;
+		}
+		else
+		{
+			Params.EmitterPosition = EmitterPosition;
+		}
+	}
+	else
+	{
+		Params.NormalizedOmniRadius = 0.0f;
+		Params.Distance = 0.0f;
+		Params.EmitterPosition = FVector::ZeroVector;
+	}
+
+	// We are currently always computing spatialization for XAudio2 relative to the listener!
+	Params.ListenerOrientation = FVector::UpVector;
+	Params.ListenerPosition = FVector::ZeroVector;
+
+	return Params;
+}
+
+
 /*-----------------------------------------------------------------------------
 	FNotifyBufferFinishedHooks implementation.
 -----------------------------------------------------------------------------*/
diff --git a/Engine/Source/Runtime/Engine/Private/AudioDevice.cpp b/Engine/Source/Runtime/Engine/Private/AudioDevice.cpp
--- a/Engine/Source/Runtime/Engine/Private/AudioDevice.cpp
+++ b/Engine/Source/Runtime/Engine/Private/AudioDevice.cpp
@@ -2068,6 +2068,8 @@ void FAudioDevice::Update( bool bGameTicking )
 		INC_DWORD_STAT_BY( STAT_ActiveSounds, ActiveSounds.Num() );
 	}
 
+	UpdateListenerTransform();
+
 	// now let the platform perform anything it needs to handle
 	UpdateHardware();
 
@@ -2662,6 +2664,27 @@ bool FAudioDevice::IsAudioDeviceMuted() const
 	return bIsDeviceMuted;
 }
 
+void FAudioDevice::UpdateListenerTransform()
+{
+	if (Listeners.Num() > 0)
+	{
+		// Caches the matrix used to transform a sounds position into local space so we can just look
+		// at the Y component after normalization to determine spatialization.
+		const FVector Up = Listeners[0].GetUp();
+		const FVector Right = Listeners[0].GetFront();
+		InverseListenerTransform = FMatrix(Up, Right, Up ^ Right, Listeners[0].Transform.GetTranslation()).InverseFast();
+	}
+}
+
+FVector FAudioDevice::GetListenerTransformedDirection(const FVector& Position, float* OutDistance)
+{
+	FVector UnnormalizedDirection = InverseListenerTransform.TransformPosition(Position);
+	if (OutDistance)
+	{
+		*OutDistance = UnnormalizedDirection.Size();
+	}
+	return UnnormalizedDirection.GetSafeNormal();
+}
 
 #if WITH_EDITOR
 void FAudioDevice::OnBeginPIE(const bool bIsSimulating)
diff --git a/Engine/Source/Runtime/Engine/Public/Audio.h b/Engine/Source/Runtime/Engine/Public/Audio.h
--- a/Engine/Source/Runtime/Engine/Public/Audio.h
+++ b/Engine/Source/Runtime/Engine/Public/Audio.h
@@ -358,6 +358,21 @@ public:
 	class FAudioDevice * AudioDevice;
 };
 
+/**
+* FSpatializationParams
+* Struct for retrieving parameters needed for computing 3d spatialization
+*/
+struct FSpatializationParams
+{
+	FVector ListenerPosition;
+	FVector ListenerOrientation;
+	FVector EmitterPosition;
+	FVector LeftChannelPosition;
+	FVector RightChannelPosition;
+	float Distance;
+	float NormalizedOmniRadius;
+};
+
 /*-----------------------------------------------------------------------------
 	FSoundSource.
 -----------------------------------------------------------------------------*/
@@ -377,6 +392,8 @@ public:
 		, LFEBleed(0.5f)
 		, HighFrequencyGain(1.0f)
 		, LastUpdate(0)
+		, LeftChannelSourceLocation(0)
+		, RightChannelSourceLocation(0)
 	{
 	}
 
@@ -479,6 +496,18 @@ public:
 	 */
 	ENGINE_API void SetHighFrequencyGain( void );
 
+	/** Updates the stereo emitter positions of this voice */
+	ENGINE_API void UpdateStereoEmitterPositions();
+
+	/** Draws debug info about this source voice if enabled. */
+	ENGINE_API void DrawDebugInfo();
+
+	/**
+	* Gets parameters necessary for computing 3d spatialization of sources
+	*/
+	ENGINE_API FSpatializationParams GetSpatializationParams();
+
+
 	const FSoundBuffer* GetBuffer() const {return Buffer;}
 
 	/**
@@ -489,6 +518,7 @@ public:
 	}
 
 protected:
+
 	// Variables.	
 	class FAudioDevice*		AudioDevice;
 	struct FWaveInstance*	WaveInstance;
@@ -515,6 +545,11 @@ protected:
 	/** Last tick when this source was active *and* had a hearable volume */
 	int32					LastHeardUpdate;
 
+	/** The location of the left-channel source for stereo spatialization. */
+	FVector						LeftChannelSourceLocation;
+	/** The location of the right-channel source for stereo spatialization. */
+	FVector						RightChannelSourceLocation;
+
 	friend class FAudioDevice;
 };
 
diff --git a/Engine/Source/Runtime/Engine/Public/AudioDevice.h b/Engine/Source/Runtime/Engine/Public/AudioDevice.h
--- a/Engine/Source/Runtime/Engine/Public/AudioDevice.h
+++ b/Engine/Source/Runtime/Engine/Public/AudioDevice.h
@@ -746,6 +746,17 @@ protected:
 
 	void AddSoundToStop(struct FActiveSound* SoundToStop);
 
+	/** Updates the listener transform. */
+	void UpdateListenerTransform();
+
+	/**   
+	* Gets the direction of the given position vector transformed relative to listener.   
+	* @param Position				Input position vector to transform relative to listener
+	* @param OutDistance			Optional output of distance from position to listener
+	* @return The input position relative to the listener.
+	*/
+	FVector GetListenerTransformedDirection(const FVector& Position, float* OutDistance);
+
 public:
 
 	/** The maximum number of concurrent audible sounds */
@@ -858,6 +869,9 @@ private:
 
 	friend class FSoundConcurrencyManager;
 	FSoundConcurrencyManager ConcurrencyManager;
+
+	/** Inverse listener transformation, used for spatialization */
+	FMatrix InverseListenerTransform;
 };
 
 
diff --git a/Engine/Source/Runtime/Windows/XAudio2/Private/XAudio2Device.cpp b/Engine/Source/Runtime/Windows/XAudio2/Private/XAudio2Device.cpp
--- a/Engine/Source/Runtime/Windows/XAudio2/Private/XAudio2Device.cpp
+++ b/Engine/Source/Runtime/Windows/XAudio2/Private/XAudio2Device.cpp
@@ -218,14 +218,6 @@ void FXAudio2Device::TeardownHardware()
 
 void FXAudio2Device::UpdateHardware()
 {
-	if (Listeners.Num() > 0)
-	{
-		// Caches the matrix used to transform a sounds position into local space so we can just look
-		// at the Y component after normalization to determine spatialization.
-		const FVector Up = Listeners[0].GetUp();
-		const FVector Right = Listeners[0].GetFront();
-		InverseListenerTransform = FMatrix(Up, Right, Up ^ Right, Listeners[0].Transform.GetTranslation()).InverseFast();
-	}
 }
 
 FAudioEffectsManager* FXAudio2Device::CreateEffectsManager()
diff --git a/Engine/Source/Runtime/Windows/XAudio2/Private/XAudio2Source.cpp b/Engine/Source/Runtime/Windows/XAudio2/Private/XAudio2Source.cpp
--- a/Engine/Source/Runtime/Windows/XAudio2/Private/XAudio2Source.cpp
+++ b/Engine/Source/Runtime/Windows/XAudio2/Private/XAudio2Source.cpp
@@ -45,8 +45,6 @@ FXAudio2SoundSource::FXAudio2SoundSource(FAudioDevice* InAudioDevice)
 	, bResourcesNeedFreeing(false)
 	, VoiceId(-1)
 	, bUsingHRTFSpatialization(false)
-	, LeftChannelSourceLocation(0)
-	, RightChannelSourceLocation(0)
 {
 	AudioDevice = ( FXAudio2Device* )InAudioDevice;
 	check( AudioDevice );
@@ -584,54 +582,13 @@ FVector FXAudio2SoundSource::ConvertToXAudio2Orientation(const FVector& InputVec
 
 }
 
-FVector FXAudio2SoundSource::GetListenerTransformedDirection(const FVector& Position, float* OutDistance)
-{
-	FVector UnnormalizedDirection = AudioDevice->InverseListenerTransform.TransformPosition(Position);
-	if (OutDistance)
-	{
-		*OutDistance = UnnormalizedDirection.Size();
-	}
-	return UnnormalizedDirection.GetSafeNormal();
-}
-
-void FXAudio2SoundSource::GetSpatializationParams(FSpatializationParams& Params)
+void FXAudio2SoundSource::GetMonoChannelVolumes(float ChannelVolumes[CHANNEL_MATRIX_COUNT], float AttenuatedVolume)
 {
-	// Calculate direction from listener to sound, where the sound is at the origin if unspatialised.
-	if (WaveInstance->bUseSpatialization)
-	{
-		FVector EmitterDirection = GetListenerTransformedDirection(WaveInstance->Location, &Params.Distance);
-		Params.NormalizedOmniRadius = (Params.Distance > 0) ? (WaveInstance->OmniRadius / Params.Distance) : FLT_MAX;
-
-		if (Buffer->NumChannels == 2)
-		{
-			FVector EmitterDirection = GetListenerTransformedDirection(LeftChannelSourceLocation, nullptr);
-			Params.LeftChannelPosition = ConvertToXAudio2Orientation(EmitterDirection);
+	FSpatializationParams SpatializationParams = GetSpatializationParams();
 
-			EmitterDirection = GetListenerTransformedDirection(RightChannelSourceLocation, nullptr);
-			Params.RightChannelPosition = ConvertToXAudio2Orientation(EmitterDirection);
+	// Convert to xaudio2 coordinates
+	SpatializationParams.EmitterPosition = ConvertToXAudio2Orientation(SpatializationParams.EmitterPosition);
 
-			Params.EmitterPosition = FVector::ZeroVector;
-		}
-		else
-		{
-			Params.EmitterPosition = ConvertToXAudio2Orientation(EmitterDirection);
-		}
-	}
-	else
-	{
-		Params.NormalizedOmniRadius = 0.0f;
-		Params.Distance = 0.0f;
-		Params.EmitterPosition = FVector::ZeroVector;
-	}
-
-	// We are currently always computing spatialization for XAudio2 relative to the listener!
-	Params.ListenerOrientation = FVector::UpVector;
-	Params.ListenerPosition = FVector::ZeroVector;
-}
-
-
-void FXAudio2SoundSource::GetMonoChannelVolumes(float ChannelVolumes[CHANNEL_MATRIX_COUNT], float AttenuatedVolume)
-{
 	if (IsUsingHrtfSpatializer())
 	{
 		// If we are using a HRTF spatializer, we are going to be using an XAPO effect that takes a mono stream and splits it into stereo
@@ -640,19 +597,11 @@ void FXAudio2SoundSource::GetMonoChannelVolumes(float ChannelVolumes[CHANNEL_MAT
 		check(WaveInstance->SpatializationAlgorithm == SPATIALIZATION_HRTF);
 		check(AudioDevice->SpatializeProcessor != nullptr);
 
-		// If we're using an HRTF spatialization algorithm, we need to find the position of the emitter and set it as a parameter
-		FVector EmitterPosition = FVector::ZeroVector;
-		FVector UnnormalizedDirection = AudioDevice->InverseListenerTransform.TransformPosition(WaveInstance->Location);
-		EmitterPosition = UnnormalizedDirection.GetSafeNormal();
-
-		AudioDevice->SpatializeProcessor->SetSpatializationParameters(VoiceId, FAudioSpatializationParams(EmitterPosition, (ESpatializationEffectType)WaveInstance->SpatializationAlgorithm));
+		AudioDevice->SpatializeProcessor->SetSpatializationParameters(VoiceId, FAudioSpatializationParams(SpatializationParams.EmitterPosition, (ESpatializationEffectType)WaveInstance->SpatializationAlgorithm));
 		GetStereoChannelVolumes(ChannelVolumes, AttenuatedVolume);
 	}
 	else // Spatialize the mono stream using the normal 3d audio algorithm
 	{
-		FSpatializationParams SpatializationParams;
-		GetSpatializationParams(SpatializationParams);
-
 		// Calculate 5.1 channel dolby surround rate/multipliers.
 		ChannelVolumes[CHANNELOUT_FRONTLEFT] = AttenuatedVolume;
 		ChannelVolumes[CHANNELOUT_FRONTRIGHT] = AttenuatedVolume;
@@ -667,7 +616,7 @@ void FXAudio2SoundSource::GetMonoChannelVolumes(float ChannelVolumes[CHANNEL_MAT
 
 		ChannelVolumes[CHANNELOUT_RADIO] = 0.0f;
 
-		// Call the spatialisation magic
+		// Call the spatialization magic
 		AudioDevice->SpatializationHelper.CalculateDolbySurroundRate(SpatializationParams.ListenerOrientation, SpatializationParams.ListenerPosition, SpatializationParams.EmitterPosition, SpatializationParams.NormalizedOmniRadius, ChannelVolumes);
 
 		// Handle any special post volume processing
@@ -709,26 +658,6 @@ void FXAudio2SoundSource::GetMonoChannelVolumes(float ChannelVolumes[CHANNEL_MAT
 	}
 }
 
-void FXAudio2SoundSource::UpdateStereoEmitterPositions()
-{
-	// We need to compute the stereo left/right channel positions using the audio component position and the spread 
-	const FVector& ListenerPosition = AudioDevice->Listeners[0].Transform.GetLocation();
-	FVector ListenerToSourceDir = WaveInstance->Location - ListenerPosition;
-	ListenerToSourceDir.Normalize();
-
-	float HalfSpread = 0.5f * WaveInstance->StereoSpread;
-
-	// Get direction of left emitter from true emitter position (left hand rule)
-	FVector LeftEmitterDir = FVector::CrossProduct(ListenerToSourceDir, FVector::UpVector);
-	FVector LeftEmitterOffset = LeftEmitterDir * HalfSpread;
-
-	// Get position vector of left emitter by adding to true emitter the dir scaled by half the spread
-	LeftChannelSourceLocation = WaveInstance->Location + LeftEmitterOffset;
-
-	// Right emitter position is same as right but opposite direction
-	RightChannelSourceLocation = WaveInstance->Location - LeftEmitterOffset;
-}
-
 void FXAudio2SoundSource::GetStereoChannelVolumes(float ChannelVolumes[CHANNEL_MATRIX_COUNT], float AttenuatedVolume)
 {
 	// If we're doing 3d spatializaton of stereo sounds
@@ -765,8 +694,11 @@ void FXAudio2SoundSource::GetStereoChannelVolumes(float ChannelVolumes[CHANNEL_M
 		UpdateStereoEmitterPositions();
 
 		// Now get the spatialization params transformed into listener-space
-		FSpatializationParams SpatializationParams;
-		GetSpatializationParams(SpatializationParams);
+		FSpatializationParams SpatializationParams = GetSpatializationParams();
+		
+		// Convert to Xaudio2 coordinates
+		SpatializationParams.LeftChannelPosition = ConvertToXAudio2Orientation(SpatializationParams.LeftChannelPosition);
+		SpatializationParams.RightChannelPosition = ConvertToXAudio2Orientation(SpatializationParams.RightChannelPosition);
 
 		// Compute the speaker mappings for the left channel
 		float* ChannelMap = ChannelVolumes;
@@ -1288,34 +1220,8 @@ void FXAudio2SoundSource::Update( void )
 		RouteToRadio( ChannelVolumes );
 	}
 
-	// Draw 3d Debug information about this source, if enabled
-	// TODO - Audio Threading. This won't work cross thread.
-	FAudioDeviceManager* DeviceManager = GEngine->GetAudioDeviceManager();
-
-	if (DeviceManager && DeviceManager->IsVisualizeDebug3dEnabled())
-	{
-		UAudioComponent* AudioComponent = (WaveInstance->ActiveSound ? WaveInstance->ActiveSound->GetAudioComponent() : nullptr);
-		if (AudioComponent)
-		{
-			UWorld* SoundWorld = AudioComponent->GetWorld();
-			if (SoundWorld)
-			{
-				FRotator SoundRotation = AudioComponent->GetComponentRotation();
-				DrawDebugCrosshairs(SoundWorld, WaveInstance->Location, SoundRotation, 20.0f, FColor::White, false, -1.0f, SDPG_Foreground);
-
-				FString Name;
-				WaveInstance->ActiveSound->Sound->GetName(Name);
+	FSoundSource::DrawDebugInfo();
 
-				if (Buffer->NumChannels == 2 && WaveInstance->bUseSpatialization)
-				{
-					DrawDebugCrosshairs(SoundWorld, LeftChannelSourceLocation, SoundRotation, 20.0f, FColor::Red, false, -1.0f, SDPG_Foreground);
-					DrawDebugCrosshairs(SoundWorld, RightChannelSourceLocation, SoundRotation, 20.0f, FColor::Green, false, -1.0f, SDPG_Foreground);
-				}
-
-				DrawDebugString(SoundWorld, AudioComponent->GetComponentLocation() + FVector(0, 0, 32), *Name, nullptr, FColor::White, 0.033, false);
-			}
-		}
-	}
 }
 
 
diff --git a/Engine/Source/Runtime/Windows/XAudio2/Private/XAudio2Support.h b/Engine/Source/Runtime/Windows/XAudio2/Private/XAudio2Support.h
--- a/Engine/Source/Runtime/Windows/XAudio2/Private/XAudio2Support.h
+++ b/Engine/Source/Runtime/Windows/XAudio2/Private/XAudio2Support.h
@@ -423,37 +423,11 @@ protected:
 	 */
 	int32 GetDestinationVoiceIndexForEffect( SourceDestinations Effect );
 
-	/**
-	* FSpatializationParams
-	* Struct for retrieving parameters needed for computing 3d spatialization
-	*/
-	struct FSpatializationParams
-	{
-		FVector ListenerPosition;
-		FVector ListenerOrientation;
-		FVector EmitterPosition;
-		FVector LeftChannelPosition;
-		FVector RightChannelPosition;
-		float Distance;
-		float NormalizedOmniRadius;
-	};
-
 	/**
 	* Converts a vector orientation from UE4 coordinates to XAudio2 coordinates
 	*/
 	inline FVector ConvertToXAudio2Orientation(const FVector& InputVector);
 
-	/**
-	* Transform the given position in absolute world-space to listener-relative space. Optionally returns distance.
-	*/
-	FVector GetListenerTransformedDirection(const FVector& Position, float* OutDistance = nullptr);
-
-	/**
-	* Gets parameters necessary for computing 3d spatialization of sources
-	*/
-	void GetSpatializationParams(FSpatializationParams& Params);
-	void UpdateStereoEmitterPositions();
-
 	/**
 	* Calculates the channel volumes for various input channel configurations.
 	*/
@@ -503,10 +477,6 @@ protected:
 	uint32						VoiceId;
 	/** Whether or not this sound is spatializing using an HRTF spatialization algorithm. */
 	bool						bUsingHRTFSpatialization;
-	/** The location of the left-channel source for stereo spatialization. */
-	FVector						LeftChannelSourceLocation;
-	/** The location of the right-channel source for stereo spatialization. */
-	FVector						RightChannelSourceLocation;
 
 	friend class FXAudio2Device;
 	friend class FXAudio2SoundSourceCallback;
diff --git a/Engine/Source/Runtime/Windows/XAudio2/Public/XAudio2Device.h b/Engine/Source/Runtime/Windows/XAudio2/Public/XAudio2Device.h
--- a/Engine/Source/Runtime/Windows/XAudio2/Public/XAudio2Device.h
+++ b/Engine/Source/Runtime/Windows/XAudio2/Public/XAudio2Device.h
@@ -166,9 +166,6 @@ protected:
 	/** Decompress a wav a number of times for profiling purposes */
 	void TimeTest( FOutputDevice& Ar, const TCHAR* WaveAssetName );
 
-	/** Inverse listener transformation, used for spatialization */
-	FMatrix InverseListenerTransform;
-
 	// For calculating spatialised volumes
 	static FSpatializationHelper SpatializationHelper;
 
