Merging CL 2715246 and 2715172 from UE4-Orion to UE4-Main

diff --git a/Engine/Source/Runtime/Windows/XAudio2/Private/XAudio2Device.cpp b/Engine/Source/Runtime/Windows/XAudio2/Private/XAudio2Device.cpp
--- a/Engine/Source/Runtime/Windows/XAudio2/Private/XAudio2Device.cpp
+++ b/Engine/Source/Runtime/Windows/XAudio2/Private/XAudio2Device.cpp
@@ -57,7 +57,6 @@ XAUDIO2_DEVICE_DETAILS FXAudioDeviceProperties::DeviceDetails;
 
 #define DEBUG_XAUDIO2 0
 
-FSpatializationHelper FXAudio2Device::SpatializationHelper;
 
 bool FXAudio2Device::InitializeHardware()
 {
@@ -182,7 +181,7 @@ bool FXAudio2Device::InitializeHardware()
 	}
 #endif	//XAUDIO_SUPPORTS_DEVICE_DETAILS
 
-	SpatializationHelper.Init();
+	DeviceProperties->SpatializationHelper.Init();
 
 	// Initialize permanent memory stack for initial & always loaded sound allocations.
 	if( CommonAudioPoolSize )
diff --git a/Engine/Source/Runtime/Windows/XAudio2/Private/XAudio2Source.cpp b/Engine/Source/Runtime/Windows/XAudio2/Private/XAudio2Source.cpp
--- a/Engine/Source/Runtime/Windows/XAudio2/Private/XAudio2Source.cpp
+++ b/Engine/Source/Runtime/Windows/XAudio2/Private/XAudio2Source.cpp
@@ -616,8 +616,8 @@ void FXAudio2SoundSource::GetMonoChannelVolumes(float ChannelVolumes[CHANNEL_MAT
 
 		ChannelVolumes[CHANNELOUT_RADIO] = 0.0f;
 
-		// Call the spatialization magic
-		AudioDevice->SpatializationHelper.CalculateDolbySurroundRate(SpatializationParams.ListenerOrientation, SpatializationParams.ListenerPosition, SpatializationParams.EmitterPosition, SpatializationParams.NormalizedOmniRadius, ChannelVolumes);
+		AudioDevice->DeviceProperties->SpatializationHelper.CalculateDolbySurroundRate(SpatializationParams.ListenerOrientation, SpatializationParams.ListenerPosition, SpatializationParams.EmitterPosition, SpatializationParams.NormalizedOmniRadius, ChannelVolumes);
+
 
 		// Handle any special post volume processing
 		if (WaveInstance->bApplyRadioFilter)
@@ -702,11 +702,11 @@ void FXAudio2SoundSource::GetStereoChannelVolumes(float ChannelVolumes[CHANNEL_M
 
 		// Compute the speaker mappings for the left channel
 		float* ChannelMap = ChannelVolumes;
-		AudioDevice->SpatializationHelper.CalculateDolbySurroundRate(SpatializationParams.ListenerOrientation, SpatializationParams.ListenerPosition, SpatializationParams.LeftChannelPosition, SpatializationParams.NormalizedOmniRadius, ChannelMap);
+		AudioDevice->DeviceProperties->SpatializationHelper.CalculateDolbySurroundRate(SpatializationParams.ListenerOrientation, SpatializationParams.ListenerPosition, SpatializationParams.LeftChannelPosition, SpatializationParams.NormalizedOmniRadius, ChannelMap);
 		
 		// Now compute the speaker mappings for the right channel
 		ChannelMap = &ChannelVolumes[CHANNELOUT_COUNT];
-		AudioDevice->SpatializationHelper.CalculateDolbySurroundRate(SpatializationParams.ListenerOrientation, SpatializationParams.ListenerPosition, SpatializationParams.RightChannelPosition, SpatializationParams.NormalizedOmniRadius, ChannelMap);
+		AudioDevice->DeviceProperties->SpatializationHelper.CalculateDolbySurroundRate(SpatializationParams.ListenerOrientation, SpatializationParams.ListenerPosition, SpatializationParams.RightChannelPosition, SpatializationParams.NormalizedOmniRadius, ChannelMap);
 	}
 	else
 	{
@@ -1501,6 +1501,7 @@ void FSpatializationHelper::Init()
 	Listener.Position.y = 0.0f;
 	Listener.Position.z = 0.0f;
 	Listener.Velocity = ZeroVector;
+	Listener.pCone = nullptr;
 
 	// Set up emitter parameters
 	Emitter.OrientFront.x = 0.0f;
@@ -1523,7 +1524,8 @@ void FSpatializationHelper::Init()
 
 	Emitter.ChannelCount = UE4_XAUDIO3D_INPUTCHANNELS;
 	Emitter.ChannelRadius = 0.0f;
-	Emitter.pChannelAzimuths = EmitterAzimuths;
+	// we aren't using the helper to spatialize multichannel files so we can set this nullptr
+	Emitter.pChannelAzimuths = nullptr;
 
 	// real volume -> 5.1-ch rate
 	VolumeCurvePoint[0].Distance = 0.0f;
@@ -1551,6 +1553,7 @@ void FSpatializationHelper::Init()
 	DSPSettings.SrcChannelCount = UE4_XAUDIO3D_INPUTCHANNELS;
 	DSPSettings.DstChannelCount = SPEAKER_COUNT;
 	DSPSettings.pMatrixCoefficients = MatrixCoefficients;
+	DSPSettings.pDelayTimes = nullptr;
 }
 
 void FSpatializationHelper::DumpSpatializationState() const
@@ -1633,7 +1636,6 @@ void FSpatializationHelper::DumpSpatializationState() const
 	// DSPSettings
 	UE_LOG(LogXAudio2, Log, TEXT("  DSPSettings"));
 	FLocal::DumpChannelArray(TEXT("    "), TEXT("pMatrixCoefficients"), ARRAY_COUNT(MatrixCoefficients), DSPSettings.pMatrixCoefficients);
-	FLocal::DumpChannelArray(TEXT("    "), TEXT("pDelayTimes"), ARRAY_COUNT(MatrixCoefficients), DSPSettings.pDelayTimes);
 	UE_LOG(LogXAudio2, Log, TEXT("    SrcChannelCount: %u"), DSPSettings.SrcChannelCount);
 	UE_LOG(LogXAudio2, Log, TEXT("    DstChannelCount: %u"), DSPSettings.DstChannelCount);
 	UE_LOG(LogXAudio2, Log, TEXT("    LPFDirectCoefficient: %f"), DSPSettings.LPFDirectCoefficient);
@@ -1737,6 +1739,12 @@ void FSpatializationHelper::CalculateDolbySurroundRate( const FVector& OrientFro
 			UE_LOG(LogXAudio2, Warning, TEXT("CalculateDolbySurroundRate generated a %s in channel %d. OmniRadius:%f MatrixCoefficient:%f"),
 				*NaNorINF, SpeakerIndex, OmniRadius, DSPSettings.pMatrixCoefficients[SpeakerIndex]);
 			//DumpSpatializationState();
+
+#if ENABLE_NAN_DIAGNOSTIC
+			DumpSpatializationState();
+			ensureMsgf(!GEnsureOnNANDiagnostic, TEXT("CalculateDolbySurroundRate generated a %s in channel %d. OmniRadius:%f MatrixCoefficient:%f"),
+				*NaNorINF, SpeakerIndex, OmniRadius, DSPSettings.pMatrixCoefficients[SpeakerIndex]);
+#endif
 		}
 #endif
 	}
diff --git a/Engine/Source/Runtime/Windows/XAudio2/Private/XAudio2Support.h b/Engine/Source/Runtime/Windows/XAudio2/Private/XAudio2Support.h
--- a/Engine/Source/Runtime/Windows/XAudio2/Private/XAudio2Support.h
+++ b/Engine/Source/Runtime/Windows/XAudio2/Private/XAudio2Support.h
@@ -579,6 +579,9 @@ struct FXAudioDeviceProperties
 	static XAUDIO2_DEVICE_DETAILS		DeviceDetails;
 #endif	//XAUDIO_SUPPORTS_DEVICE_DETAILS
 
+	// For calculating speaker maps for 3d audio
+	FSpatializationHelper				SpatializationHelper;
+
 	/** Source callback to handle looping sound callbacks */
 	FXAudio2SoundSourceCallback	SourceCallback;
 	
diff --git a/Engine/Source/Runtime/Windows/XAudio2/Public/XAudio2Device.h b/Engine/Source/Runtime/Windows/XAudio2/Public/XAudio2Device.h
--- a/Engine/Source/Runtime/Windows/XAudio2/Public/XAudio2Device.h
+++ b/Engine/Source/Runtime/Windows/XAudio2/Public/XAudio2Device.h
@@ -166,13 +166,9 @@ protected:
 	/** Decompress a wav a number of times for profiling purposes */
 	void TimeTest( FOutputDevice& Ar, const TCHAR* WaveAssetName );
 
-	// For calculating spatialised volumes
-	static FSpatializationHelper SpatializationHelper;
-
 	friend class FXAudio2SoundBuffer;
 	friend class FXAudio2SoundSource;
 	friend class FXAudio2EffectsManager;
-	friend class FAsyncXAudio2SourceDestroyer;
 
 private:
 	struct FXAudioDeviceProperties* DeviceProperties;
