Disable DX11 TransitionResource checks by default

diff --git a/Engine/Source/Runtime/Windows/D3D11RHI/Private/D3D11Commands.cpp b/Engine/Source/Runtime/Windows/D3D11RHI/Private/D3D11Commands.cpp
--- a/Engine/Source/Runtime/Windows/D3D11RHI/Private/D3D11Commands.cpp
+++ b/Engine/Source/Runtime/Windows/D3D11RHI/Private/D3D11Commands.cpp
@@ -61,6 +61,24 @@ DECLARE_ISBOUNDSHADER(ComputeShader)
 #define VALIDATE_BOUND_SHADER(s)
 #endif
 
+int32 GEnableDX11TransitionChecks = 0;
+static FAutoConsoleVariableRef CVarDX11TransitionChecks(
+	TEXT("r.TransitionChecksEnableDX11"),
+	GEnableDX11TransitionChecks,
+	TEXT("Enables transition checks in the DX11 RHI."),
+	ECVF_Default
+	);
+
+void FD3D11BaseShaderResource::SetDirty(bool bInDirty, uint32 CurrentFrame)
+{
+	bDirty = bInDirty;
+	if (bDirty)
+	{
+		LastFrameWritten = CurrentFrame;
+	}
+	ensureMsgf((GEnableDX11TransitionChecks == 0) || !(CurrentGPUAccess == EResourceTransitionAccess::EReadable && bDirty), TEXT("ShaderResource is dirty, but set to Readable."));
+}
+
 void FD3D11DynamicRHI::RHIBeginAsyncComputeJob_DrawThread(EAsyncComputePriority Priority) 
 {  
 }
@@ -354,7 +372,7 @@ void FD3D11DynamicRHI::RHISetUAVParameter(FComputeShaderRHIParamRef ComputeShade
 		//check it's safe for r/w for this UAV
 		const EResourceTransitionAccess CurrentUAVAccess = UAV->Resource->GetCurrentGPUAccess();
 		const bool UAVDirty = UAV->Resource->IsDirty();
-		ensureMsgf(!UAVDirty || (CurrentUAVAccess == EResourceTransitionAccess::ERWNoBarrier), TEXT("UAV: %i is in unsafe state for GPU R/W: %s, Dirty: %i"), UAVIndex, *FResourceTransitionUtility::ResourceTransitionAccessStrings[(int32)CurrentUAVAccess], (int32)UAVDirty);
+		ensureMsgf((GEnableDX11TransitionChecks == 0) || !UAVDirty || (CurrentUAVAccess == EResourceTransitionAccess::ERWNoBarrier), TEXT("UAV: %i is in unsafe state for GPU R/W: %s, Dirty: %i"), UAVIndex, *FResourceTransitionUtility::ResourceTransitionAccessStrings[(int32)CurrentUAVAccess], (int32)UAVDirty);
 
 		//UAVs always dirty themselves. If a shader wanted to just read, it should use an SRV.
 		UAV->Resource->SetDirty(true, PresentCounter);
@@ -379,7 +397,7 @@ void FD3D11DynamicRHI::RHISetUAVParameter(FComputeShaderRHIParamRef ComputeShade
 		//check it's safe for r/w for this UAV
 		const EResourceTransitionAccess CurrentUAVAccess = UAV->Resource->GetCurrentGPUAccess();
 		const bool UAVDirty = UAV->Resource->IsDirty();
-		ensureMsgf(!UAVDirty || (CurrentUAVAccess == EResourceTransitionAccess::ERWNoBarrier), TEXT("UAV: %i is in unsafe state for GPU R/W: %s, Dirty: %i"), UAVIndex, *FResourceTransitionUtility::ResourceTransitionAccessStrings[(int32)CurrentUAVAccess], (int32)UAVDirty);
+		ensureMsgf((GEnableDX11TransitionChecks == 0) || !UAVDirty || (CurrentUAVAccess == EResourceTransitionAccess::ERWNoBarrier), TEXT("UAV: %i is in unsafe state for GPU R/W: %s, Dirty: %i"), UAVIndex, *FResourceTransitionUtility::ResourceTransitionAccessStrings[(int32)CurrentUAVAccess], (int32)UAVDirty);
 
 		//UAVs always dirty themselves. If a shader wanted to just read, it should use an SRV.
 		UAV->Resource->SetDirty(true, PresentCounter);
@@ -869,7 +887,7 @@ void FD3D11DynamicRHI::RHISetRenderTargets(
 									LastFrameWritten != CurrentFrame || 
 									!bDepthWrite;
 
-		ensureMsgf(bAccessValid, TEXT("DepthTarget '%s' is not GPU writable."), *NewDepthStencilTargetRHI->Texture->GetName().ToString());		
+		ensureMsgf((GEnableDX11TransitionChecks == 0) || bAccessValid, TEXT("DepthTarget '%s' is not GPU writable."), *NewDepthStencilTargetRHI->Texture->GetName().ToString());
 
 		//switch to writable state if this is the first render of the frame.  Don't switch if it's a later render and this is a depth test only situation
 		if (!bAccessValid || (bReadable && bDepthWrite))
@@ -903,7 +921,7 @@ void FD3D11DynamicRHI::RHISetRenderTargets(
 				const uint32 LastFrameWritten = NewRenderTarget->GetLastFrameWritten();
 				const bool bReadable = CurrentAccess == EResourceTransitionAccess::EReadable;
 				const bool bAccessValid = !bReadable || LastFrameWritten != CurrentFrame;
-				ensureMsgf(bAccessValid, TEXT("RenderTarget '%s' is not GPU writable."), *NewRenderTargetsRHI[RenderTargetIndex].Texture->GetName().ToString());
+				ensureMsgf((GEnableDX11TransitionChecks == 0) || bAccessValid, TEXT("RenderTarget '%s' is not GPU writable."), *NewRenderTargetsRHI[RenderTargetIndex].Texture->GetName().ToString());
 								
 				if (!bAccessValid || bReadable)
 				{
@@ -985,7 +1003,7 @@ void FD3D11DynamicRHI::RHISetRenderTargets(
 				const EResourceTransitionAccess CurrentUAVAccess = RHIUAV->Resource->GetCurrentGPUAccess();
 				const bool UAVDirty = RHIUAV->Resource->IsDirty();
 				const bool bAccessPass = (CurrentUAVAccess == EResourceTransitionAccess::ERWBarrier && !UAVDirty) || (CurrentUAVAccess == EResourceTransitionAccess::ERWNoBarrier);
-				ensureMsgf(bAccessPass, TEXT("UAV: %i is in unsafe state for GPU R/W: %s"), UAVIndex, *FResourceTransitionUtility::ResourceTransitionAccessStrings[(int32)CurrentUAVAccess]);
+				ensureMsgf((GEnableDX11TransitionChecks == 0) || bAccessPass, TEXT("UAV: %i is in unsafe state for GPU R/W: %s"), UAVIndex, *FResourceTransitionUtility::ResourceTransitionAccessStrings[(int32)CurrentUAVAccess]);
 
 				//UAVs get set to dirty.  If the shader just wanted to read it should have used an SRV.
 				RHIUAV->Resource->SetDirty(true, PresentCounter);
diff --git a/Engine/Source/Runtime/Windows/D3D11RHI/Private/D3D11RHI.cpp b/Engine/Source/Runtime/Windows/D3D11RHI/Private/D3D11RHI.cpp
--- a/Engine/Source/Runtime/Windows/D3D11RHI/Private/D3D11RHI.cpp
+++ b/Engine/Source/Runtime/Windows/D3D11RHI/Private/D3D11RHI.cpp
@@ -175,6 +175,7 @@ void FD3D11DynamicRHI::CheckIfSRVIsResolved(ID3D11ShaderResourceView* SRV)
 #endif
 }
 
+extern int32 GEnableDX11TransitionChecks;
 template <EShaderFrequency ShaderFrequency>
 void FD3D11DynamicRHI::InternalSetShaderResourceView(FD3D11BaseShaderResource* Resource, ID3D11ShaderResourceView* SRV, int32 ResourceIndex, FName SRVName, FD3D11StateCache::ESRV_Type SrvType)
 {
@@ -186,7 +187,7 @@ void FD3D11DynamicRHI::InternalSetShaderResourceView(FD3D11BaseShaderResource* R
 	{		
 		const EResourceTransitionAccess CurrentAccess = Resource->GetCurrentGPUAccess();
 		const bool bAccessPass = CurrentAccess == EResourceTransitionAccess::EReadable || (CurrentAccess == EResourceTransitionAccess::ERWBarrier && !Resource->IsDirty()) || CurrentAccess == EResourceTransitionAccess::ERWSubResBarrier;
-		ensureMsgf(bAccessPass || Resource->GetLastFrameWritten() != PresentCounter, TEXT("Shader resource %s is not GPU readable.  Missing a call to RHITransitionResources()"), *SRVName.ToString());
+		ensureMsgf((GEnableDX11TransitionChecks == 0) || bAccessPass || Resource->GetLastFrameWritten() != PresentCounter, TEXT("Shader resource %s is not GPU readable.  Missing a call to RHITransitionResources()"), *SRVName.ToString());
 	}
 
 	FD3D11BaseShaderResource*& ResourceSlot = CurrentResourcesBoundAsSRVs[ShaderFrequency][ResourceIndex];
diff --git a/Engine/Source/Runtime/Windows/D3D11RHI/Public/D3D11Resources.h b/Engine/Source/Runtime/Windows/D3D11RHI/Public/D3D11Resources.h
--- a/Engine/Source/Runtime/Windows/D3D11RHI/Public/D3D11Resources.h
+++ b/Engine/Source/Runtime/Windows/D3D11RHI/Public/D3D11Resources.h
@@ -170,17 +170,9 @@ public:
 	uint32 GetLastFrameWritten() const
 	{
 		return LastFrameWritten;
-	}
+	}	
 
-	void SetDirty(bool bInDirty, uint32 CurrentFrame)
-	{
-		bDirty = bInDirty;
-		if (bDirty)
-		{
-			LastFrameWritten = CurrentFrame;
-		}
-		ensureMsgf(!(CurrentGPUAccess == EResourceTransitionAccess::EReadable && bDirty), TEXT("UAV is dirty, but set to Readable."));
-	}
+	void SetDirty(bool bInDirty, uint32 CurrentFrame);
 
 	bool IsDirty() const
 	{
