Merging //UE4/Release-4.11 to Dev-Main (//UE4/Dev-Main)

Change 2781072 on 2015/11/25 by [REDACTED_PERSON][REDACTED_HANDLE]

        Set SampleRate to UE4_XAUDIO2_SAMPLERATE in FXAudio2Device::InitializeHardware()

diff --git a/Engine/Source/Runtime/Windows/XAudio2/Private/XAudio2Device.cpp b/Engine/Source/Runtime/Windows/XAudio2/Private/XAudio2Device.cpp
--- a/Engine/Source/Runtime/Windows/XAudio2/Private/XAudio2Device.cpp
+++ b/Engine/Source/Runtime/Windows/XAudio2/Private/XAudio2Device.cpp
@@ -74,7 +74,7 @@ bool FXAudio2Device::InitializeHardware()
 	// Load ogg and vorbis dlls if they haven't been loaded yet
 	LoadVorbisLibraries();
 
-	SampleRate = 0;
+	SampleRate = UE4_XAUDIO2_SAMPLERATE;
 
 #if PLATFORM_WINDOWS
 	bComInitialized = FWindowsPlatformMisc::CoInitialize();
