Copying //UE4/Orion-Staging to //UE4/Main (Origin: //Orion/Dev-General [REDACTED_HANDLE])

diff --git a/Engine/Source/Runtime/SlateRHIRenderer/Private/SlateRHIRenderer.h b/Engine/Source/Runtime/SlateRHIRenderer/Private/SlateRHIRenderer.h
--- a/Engine/Source/Runtime/SlateRHIRenderer/Private/SlateRHIRenderer.h
+++ b/Engine/Source/Runtime/SlateRHIRenderer/Private/SlateRHIRenderer.h
@@ -188,7 +188,6 @@ public:
 	virtual void ReleaseAccessedResources(bool bImmediatelyFlush) override;
 	virtual TSharedRef<FSlateRenderDataHandle, ESPMode::ThreadSafe> CacheElementRenderData(const ILayoutCache* Cacher, FSlateWindowElementList& ElementList) override;
 	virtual void ReleaseCachingResourcesFor(const ILayoutCache* Cacher) override;
-	virtual void ReleaseCachedRenderData(FSlateRenderDataHandle* RenderHandle) override;
 
 	/** Draws windows from a FSlateDrawBuffer on the render thread */
 	void DrawWindow_RenderThread(FRHICommandListImmediate& RHICmdList, const FSlateRHIRenderer::FViewportInfo& ViewportInfo, FSlateWindowElementList& WindowElementList, bool bLockToVsync, bool bClear);
diff --git a/Engine/Source/Runtime/SlateRHIRenderer/Private/SlateRHIRenderingPolicy.cpp b/Engine/Source/Runtime/SlateRHIRenderer/Private/SlateRHIRenderingPolicy.cpp
--- a/Engine/Source/Runtime/SlateRHIRenderer/Private/SlateRHIRenderingPolicy.cpp
+++ b/Engine/Source/Runtime/SlateRHIRenderer/Private/SlateRHIRenderingPolicy.cpp
@@ -130,11 +130,6 @@ void FSlateRHIRenderingPolicy::UpdateVertexAndIndexBuffers(FRHICommandListImmedi
 	UpdateVertexAndIndexBuffers(RHICmdList, InBatchData, Buffers->VertexBuffer, Buffers->IndexBuffer);
 }
 
-void FSlateRHIRenderingPolicy::ReleaseCachedRenderData(FSlateRenderDataHandle* InRenderHandle)
-{
-	ResourceManager->ReleaseCachedRenderData(InRenderHandle);
-}
-
 void FSlateRHIRenderingPolicy::ReleaseCachingResourcesFor(const ILayoutCache* Cacher)
 {
 	ResourceManager->ReleaseCachingResourcesFor(Cacher);
diff --git a/Engine/Source/Runtime/SlateRHIRenderer/Private/SlateRHIRenderingPolicy.h b/Engine/Source/Runtime/SlateRHIRenderer/Private/SlateRHIRenderingPolicy.h
--- a/Engine/Source/Runtime/SlateRHIRenderer/Private/SlateRHIRenderingPolicy.h
+++ b/Engine/Source/Runtime/SlateRHIRenderer/Private/SlateRHIRenderingPolicy.h
@@ -16,7 +16,6 @@ public:
 	void UpdateVertexAndIndexBuffers(FRHICommandListImmediate& RHICmdList, FSlateBatchData& BatchData);
 	void UpdateVertexAndIndexBuffers(FRHICommandListImmediate& RHICmdList, FSlateBatchData& BatchData, const TSharedRef<FSlateRenderDataHandle, ESPMode::ThreadSafe>& RenderHandle);
 
-	void ReleaseCachedRenderData(FSlateRenderDataHandle* InRenderHandle);
 	void ReleaseCachingResourcesFor(const ILayoutCache* Cacher);
 
 	void SetDefaultBlendMode(const FBlendStateInitializerRHI& BlendState);
diff --git a/Engine/Source/Runtime/SlateRHIRenderer/Private/SlateRHIResourceManager.cpp b/Engine/Source/Runtime/SlateRHIRenderer/Private/SlateRHIResourceManager.cpp
--- a/Engine/Source/Runtime/SlateRHIRenderer/Private/SlateRHIResourceManager.cpp
+++ b/Engine/Source/Runtime/SlateRHIRenderer/Private/SlateRHIResourceManager.cpp
@@ -699,6 +699,8 @@ FSlateMaterialResource* FSlateRHIResourceManager::GetMaterialResource(const UObj
 
 void FSlateRHIResourceManager::OnAppExit()
 {
+	FlushRenderingCommands();
+
 	ReleaseResources();
 
 	FlushRenderingCommands();
@@ -851,6 +853,8 @@ void FSlateRHIResourceManager::UpdateTextureAtlases()
 
 FCachedRenderBuffers* FSlateRHIResourceManager::FindOrCreateCachedBuffersForHandle(const TSharedRef<FSlateRenderDataHandle, ESPMode::ThreadSafe>& RenderHandle)
 {
+	check(IsInRenderingThread());
+
 	FCachedRenderBuffers* Buffers = CachedBuffers.FindRef(&RenderHandle.Get());
 	if ( Buffers == nullptr )
 	{
@@ -865,8 +869,8 @@ FCachedRenderBuffers* FSlateRHIResourceManager::FindOrCreateCachedBuffersForHand
 		if ( Pool.Num() == 0 )
 		{
 			Buffers = new FCachedRenderBuffers();
-			Buffers->VertexBuffer.Init(200);
-			Buffers->IndexBuffer.Init(200);
+			Buffers->VertexBuffer.Init(100);
+			Buffers->IndexBuffer.Init(100);
 		}
 		else
 		{
@@ -881,16 +885,36 @@ FCachedRenderBuffers* FSlateRHIResourceManager::FindOrCreateCachedBuffersForHand
 	return Buffers;
 }
 
-void FSlateRHIResourceManager::ReleaseCachedRenderData(FSlateRenderDataHandle* InRenderHandle)
+void FSlateRHIResourceManager::BeginReleasingRenderData(const FSlateRenderDataHandle* RenderHandle)
+{
+	struct FReleaseCachedRenderDataContext
+	{
+		FSlateRHIResourceManager* ResourceManager;
+		const FSlateRenderDataHandle* RenderDataHandle;
+		const ILayoutCache* LayoutCacher;
+	};
+	FReleaseCachedRenderDataContext ReleaseCachedRenderDataContext =
+	{
+		this,
+		RenderHandle,
+		RenderHandle->GetCacher()
+	};
+	ENQUEUE_UNIQUE_RENDER_COMMAND_ONEPARAMETER(
+		ReleaseCachedRenderData,
+		FReleaseCachedRenderDataContext, Context, ReleaseCachedRenderDataContext,
+		{
+			Context.ResourceManager->ReleaseCachedRenderData(Context.RenderDataHandle, Context.LayoutCacher);
+		});
+}
+
+void FSlateRHIResourceManager::ReleaseCachedRenderData(const FSlateRenderDataHandle* RenderHandle, const ILayoutCache* LayoutCacher)
 {
 	// Should only be called by the rendering thread
 	check(IsInRenderingThread());
-	check(InRenderHandle);
 
-	FCachedRenderBuffers* PooledBuffer = CachedBuffers.FindRef(InRenderHandle);
+	FCachedRenderBuffers* PooledBuffer = CachedBuffers.FindRef(RenderHandle);
 	if ( ensure(PooledBuffer != nullptr) )
 	{
-		const ILayoutCache* LayoutCacher = InRenderHandle->GetCacher();
 		TArray< FCachedRenderBuffers* >* Pool = CachedBufferPool.Find(LayoutCacher);
 		if ( Pool )
 		{
@@ -904,12 +928,14 @@ void FSlateRHIResourceManager::ReleaseCachedRenderData(FSlateRenderDataHandle* I
 			delete PooledBuffer;
 		}
 
-		CachedBuffers.Remove(InRenderHandle);
+		CachedBuffers.Remove(RenderHandle);
 	}
 }
 
 void FSlateRHIResourceManager::ReleaseCachingResourcesFor(const ILayoutCache* Cacher)
 {
+	check(IsInRenderingThread());
+
 	TArray< FCachedRenderBuffers* >* Pool = CachedBufferPool.Find(Cacher);
 	if ( Pool )
 	{
diff --git a/Engine/Source/Runtime/SlateRHIRenderer/Private/SlateRHIResourceManager.h b/Engine/Source/Runtime/SlateRHIRenderer/Private/SlateRHIResourceManager.h
--- a/Engine/Source/Runtime/SlateRHIRenderer/Private/SlateRHIResourceManager.h
+++ b/Engine/Source/Runtime/SlateRHIRenderer/Private/SlateRHIResourceManager.h
@@ -73,7 +73,7 @@ struct FCachedRenderBuffers
 /**
  * Stores a mapping of texture names to their RHI texture resource               
  */
-class FSlateRHIResourceManager : public ISlateAtlasProvider, public FSlateShaderResourceManager, public FGCObject
+class FSlateRHIResourceManager : public ISlateAtlasProvider, public ISlateRenderDataManager, public FSlateShaderResourceManager, public FGCObject
 {
 public:
 	FSlateRHIResourceManager();
@@ -85,6 +85,10 @@ public:
 	virtual FSlateShaderResource* GetAtlasPageResource(const int32 InIndex) const override;
 	virtual bool IsAtlasPageResourceAlphaOnly() const override;
 
+	/** ISlateRenderDataManager */
+	virtual void BeginReleasingRenderData(const FSlateRenderDataHandle* RenderHandle) override;
+
+	/** FGCObject */
 	virtual void AddReferencedObjects(FReferenceCollector& Collector) override;
 
 	/**
@@ -152,14 +156,24 @@ public:
 	virtual bool LoadTexture( const FName& TextureName, const FString& ResourcePath, uint32& Width, uint32& Height, TArray<uint8>& DecodedImage );
 	virtual bool LoadTexture( const FSlateBrush& InBrush, uint32& Width, uint32& Height, TArray<uint8>& DecodedImage );
 
+	/**
+	 * 
+	 */
 	FCachedRenderBuffers* FindCachedBuffersForHandle(const FSlateRenderDataHandle* RenderHandle) const
 	{
 		return CachedBuffers.FindRef(RenderHandle);
 	}
 
+	/**
+	 * Creates a vertex and index buffer for a given render handle that it can use to store cached widget geometry to.
+	 */
 	FCachedRenderBuffers* FindOrCreateCachedBuffersForHandle(const TSharedRef<FSlateRenderDataHandle, ESPMode::ThreadSafe>& RenderHandle);
-	void ReleaseCachedRenderData(FSlateRenderDataHandle* InRenderHandle);
+
+	/**
+	 * Releases all cached buffers allocated by a given layout cacher.  This would happen when an invalidation panel is destroyed.
+	 */
 	void ReleaseCachingResourcesFor(const ILayoutCache* Cacher);
+
 	/**
 	 * Releases rendering resources
 	 */
@@ -172,6 +186,13 @@ public:
 	 */
 	void ReloadTextures();
 
+private:
+	/**
+	* Releases the cached render data for a given render handle.  The layout cacher that owned the handle must also be
+	* provided, as render handle is likely no longer a valid object.
+	*/
+	void ReleaseCachedRenderData(const FSlateRenderDataHandle* RenderHandle, const ILayoutCache* LayoutCacher);
+
 private:
 	/**
 	 * Gets the current accessed UObject tracking set.
diff --git a/Engine/Source/Runtime/UMG/Private/Animation/UMGSequencePlayer.cpp b/Engine/Source/Runtime/UMG/Private/Animation/UMGSequencePlayer.cpp
--- a/Engine/Source/Runtime/UMG/Private/Animation/UMGSequencePlayer.cpp
+++ b/Engine/Source/Runtime/UMG/Private/Animation/UMGSequencePlayer.cpp
@@ -146,6 +146,14 @@ void UUMGSequencePlayer::Pause()
 	PlayerStatus = EMovieScenePlayerStatus::Stopped;
 }
 
+void UUMGSequencePlayer::Reverse()
+{
+	if (PlayerStatus == EMovieScenePlayerStatus::Playing)
+	{
+		bIsPlayingForward = !bIsPlayingForward;
+	}
+}
+
 void UUMGSequencePlayer::Stop()
 {
 	PlayerStatus = EMovieScenePlayerStatus::Stopped;
diff --git a/Engine/Source/Runtime/UMG/Private/Components/RetainerBox.cpp b/Engine/Source/Runtime/UMG/Private/Components/RetainerBox.cpp
--- a/Engine/Source/Runtime/UMG/Private/Components/RetainerBox.cpp
+++ b/Engine/Source/Runtime/UMG/Private/Components/RetainerBox.cpp
@@ -117,6 +117,17 @@ const FText URetainerBox::GetPaletteCategory()
 
 #endif
 
+const FGeometry& URetainerBox::GetCachedAllottedGeometry() const
+{
+	if (MyRetainerWidget.IsValid())
+	{
+		return MyRetainerWidget->GetCachedAllottedGeometry();
+	}
+
+	static const FGeometry TempGeo;
+	return TempGeo;
+}
+
 /////////////////////////////////////////////////////
 
 #undef LOCTEXT_NAMESPACE
diff --git a/Engine/Source/Runtime/UMG/Private/Slate/SMeshWidget.cpp b/Engine/Source/Runtime/UMG/Private/Slate/SMeshWidget.cpp
--- a/Engine/Source/Runtime/UMG/Private/Slate/SMeshWidget.cpp
+++ b/Engine/Source/Runtime/UMG/Private/Slate/SMeshWidget.cpp
@@ -98,6 +98,13 @@ uint32 SMeshWidget::AddMesh(USlateVectorArtData& InMeshData)
 	return RenderData.Num()-1;
 }
 
+uint32 SMeshWidget::AddMeshWithInstancing(USlateVectorArtData& InMeshData, int32 InitialBufferSize)
+{
+	const uint32 NewMeshId = AddMesh(InMeshData);
+	EnableInstancing(NewMeshId, InitialBufferSize);
+	return NewMeshId;
+}
+
 
 UMaterialInstanceDynamic* SMeshWidget::ConvertToMID(uint32 MeshId)
 {
diff --git a/Engine/Source/Runtime/UMG/Private/UserWidget.cpp b/Engine/Source/Runtime/UMG/Private/UserWidget.cpp
--- a/Engine/Source/Runtime/UMG/Private/UserWidget.cpp
+++ b/Engine/Source/Runtime/UMG/Private/UserWidget.cpp
@@ -63,9 +63,18 @@ bool UUserWidget::Initialize()
 
 		// Only do this if this widget is of a blueprint class
 		UWidgetBlueprintGeneratedClass* BGClass = Cast<UWidgetBlueprintGeneratedClass>(GetClass());
-		if ( BGClass != nullptr )
+		if (BGClass != nullptr)
 		{
-			BGClass->InitializeWidget(this);
+			//TODO NickD: This is a hack, and should be undone later
+			UWidgetBlueprintGeneratedClass* SuperBGClass = Cast<UWidgetBlueprintGeneratedClass>(BGClass->GetSuperClass());
+			if (BGClass->WidgetTree->RootWidget == nullptr && SuperBGClass )
+			{
+				SuperBGClass->InitializeWidget(this);
+			}
+			else
+			{
+				BGClass->InitializeWidget(this);
+			}
 		}
 		else
 		{
@@ -343,6 +352,19 @@ void UUserWidget::SetNumLoopsToPlay(const UWidgetAnimation* InAnimation, int32 I
 	}
 }
 
+void UUserWidget::ReverseAnimation(const UWidgetAnimation* InAnimation)
+{
+	if (InAnimation)
+	{
+		UUMGSequencePlayer** FoundPlayer = ActiveSequencePlayers.FindByPredicate([&](const UUMGSequencePlayer* Player) { return Player->GetAnimation() == InAnimation; });
+
+		if (FoundPlayer)
+		{
+			(*FoundPlayer)->Reverse();
+		}
+	}
+}
+
 void UUserWidget::OnAnimationFinishedPlaying(UUMGSequencePlayer& Player)
 {
 	OnAnimationFinished( Player.GetAnimation() );
@@ -924,6 +946,11 @@ void UUserWidget::StopListeningForAllInputActions()
 {
 	if ( InputComponent )
 	{
+		for ( int32 ExistingIndex = InputComponent->GetNumActionBindings() - 1; ExistingIndex >= 0; --ExistingIndex )
+		{
+			InputComponent->RemoveActionBinding( ExistingIndex );
+		}
+
 		if ( APlayerController* Controller = GetOwningPlayer() )
 		{
 			Controller->PopInputComponent( InputComponent );
diff --git a/Engine/Source/Runtime/UMG/Public/Animation/UMGSequencePlayer.h b/Engine/Source/Runtime/UMG/Public/Animation/UMGSequencePlayer.h
--- a/Engine/Source/Runtime/UMG/Public/Animation/UMGSequencePlayer.h
+++ b/Engine/Source/Runtime/UMG/Public/Animation/UMGSequencePlayer.h
@@ -28,6 +28,9 @@ public:
 	/** Pauses a running animation */
 	void Pause();
 
+	/** Reverses a running animation */
+	void Reverse();
+
 	/** Gets the current time position in the player (in seconds). */
 	double GetTimeCursorPosition() const { return TimeCursorPosition; }
 
diff --git a/Engine/Source/Runtime/UMG/Public/Blueprint/UserWidget.h b/Engine/Source/Runtime/UMG/Public/Blueprint/UserWidget.h
--- a/Engine/Source/Runtime/UMG/Public/Blueprint/UserWidget.h
+++ b/Engine/Source/Runtime/UMG/Public/Blueprint/UserWidget.h
@@ -741,6 +741,14 @@ public:
 	UFUNCTION(BlueprintCallable, BlueprintCosmetic, Category = "User Interface|Animation")
 	void SetNumLoopsToPlay(const UWidgetAnimation* InAnimation, int32 NumLoopsToPlay);
 
+	/**
+	* If an animation is playing, this function will reverse the playback.
+	*
+	* @param InAnimation The playing animation that we want to reverse
+	*/
+	UFUNCTION(BlueprintCallable, BlueprintCosmetic, Category = "User Interface|Animation")
+	void ReverseAnimation(const UWidgetAnimation* InAnimation);
+
 	/** Called when a sequence player is finished playing an animation */
 	void OnAnimationFinishedPlaying(UUMGSequencePlayer& Player );
 
diff --git a/Engine/Source/Runtime/UMG/Public/Components/RetainerBox.h b/Engine/Source/Runtime/UMG/Public/Components/RetainerBox.h
--- a/Engine/Source/Runtime/UMG/Public/Components/RetainerBox.h
+++ b/Engine/Source/Runtime/UMG/Public/Components/RetainerBox.h
@@ -68,6 +68,8 @@ public:
 	virtual const FText GetPaletteCategory() override;
 #endif
 
+	const FGeometry& GetCachedAllottedGeometry() const;
+
 protected:
 
 	/**
diff --git a/Engine/Source/Runtime/UMG/Public/Slate/SMeshWidget.h b/Engine/Source/Runtime/UMG/Public/Slate/SMeshWidget.h
--- a/Engine/Source/Runtime/UMG/Public/Slate/SMeshWidget.h
+++ b/Engine/Source/Runtime/UMG/Public/Slate/SMeshWidget.h
@@ -29,6 +29,9 @@ public:
 	 */
 	uint32 AddMesh(USlateVectorArtData& InMeshData);
 
+	/** Much like AddMesh, but also enables instancing support for this MeshId. */
+	uint32 AddMeshWithInstancing(USlateVectorArtData& InMeshData, int32 InitialBufferSize = 1);
+
 	/**
 	 * Switch from static material to material instance dynamic.
 	 * 
diff --git a/Engine/Source/Runtime/UMG/Public/Slate/SRetainerWidget.h b/Engine/Source/Runtime/UMG/Public/Slate/SRetainerWidget.h
--- a/Engine/Source/Runtime/UMG/Public/Slate/SRetainerWidget.h
+++ b/Engine/Source/Runtime/UMG/Public/Slate/SRetainerWidget.h
@@ -63,6 +63,11 @@ public:
 	virtual TSharedPtr<struct FVirtualPointerPosition> TranslateMouseCoordinateFor3DChild(const TSharedRef<SWidget>& ChildWidget, const FGeometry& ViewportGeometry, const FVector2D& ScreenSpaceMouseCoordinate, const FVector2D& LastScreenSpaceMouseCoordinate) const override;
 	// ICustomHitTestPath
 
+	FORCEINLINE const FGeometry& GetCachedAllottedGeometry() const
+	{
+		return CachedAllottedGeometry;
+	}
+
 protected:
 	// BEGIN SLeafWidget interface
 	virtual int32 OnPaint(const FPaintArgs& Args, const FGeometry& AllottedGeometry, const FSlateRect& MyClippingRect, FSlateWindowElementList& OutDrawElements, int32 LayerId, const FWidgetStyle& InWidgetStyle, bool bParentEnabled) const override;
diff --git a/Engine/Source/Runtime/UMG/Public/Slate/SlateVectorArtInstanceData.h b/Engine/Source/Runtime/UMG/Public/Slate/SlateVectorArtInstanceData.h
--- a/Engine/Source/Runtime/UMG/Public/Slate/SlateVectorArtInstanceData.h
+++ b/Engine/Source/Runtime/UMG/Public/Slate/SlateVectorArtInstanceData.h
@@ -15,6 +15,34 @@ public:
 	void SetScale(float Scale);
 	void SetBaseAddress(float Address);
 
+	template<int32 Component, int32 ByteIndex>
+	void PackFloatIntoByte(float InValue)
+	{
+		PackByteIntoByte<Component, ByteIndex>(static_cast<uint8>(FMath::RoundToInt(InValue * 0xFFu)));
+	}
+
+	template<int32 Component, int32 ByteIndex>
+	void PackByteIntoByte(uint8 InValue)
+	{
+		// Each float has 24 usable bits of mantissa, but we cannot access the bits directly.
+		// We will not respect the IEEE "normalized mantissa" rules, so let 
+		// the compiler/FPU do conversions from byte to float and vice versa for us.
+
+		//Produces a value like 0xFFFF00FF; where the 00s are shifted by `Position` bytes.
+		static const uint32 Mask = ~(0xFF << ByteIndex * 8);
+
+		// Clear out a byte to OR with while maintaining the rest of the data intact.
+		uint32 CurrentData = static_cast<uint32>(Data[Component]) & Mask;
+		// OR in the value
+		Data[Component] = CurrentData | (InValue << ByteIndex * 8);
+	}
+
+	template<int32 Component, int32 ByteIndex>
+	void PackByteIntoByte(float InValue)
+	{
+		checkf(false, TEXT("PackByteIntoByte() being used with a float value. This is almost definitely an error. Use PackFloatIntoByte()."));
+	}
+
 protected:
 	FVector4 Data;
 };
\ No newline at end of file
diff --git a/Engine/Source/Runtime/Windows/D3D11RHI/Private/D3D11RHI.cpp b/Engine/Source/Runtime/Windows/D3D11RHI/Private/D3D11RHI.cpp
--- a/Engine/Source/Runtime/Windows/D3D11RHI/Private/D3D11RHI.cpp
+++ b/Engine/Source/Runtime/Windows/D3D11RHI/Private/D3D11RHI.cpp
@@ -185,6 +185,13 @@ void FD3D11DynamicRHI::InternalSetShaderResourceView(FD3D11BaseShaderResource* R
 	check((Resource && SRV) || (!Resource && !SRV));
 	CheckIfSRVIsResolved(SRV);
 
+	//avoid state cache crash
+	if (!((Resource && SRV) || (!Resource && !SRV)))
+	{
+		//UE_LOG(LogRHI, Warning, TEXT("Bailing on InternalSetShaderResourceView on resource: %i, %s"), ResourceIndex, *SRVName.ToString());
+		return;
+	}
+
 	if (Resource)
 	{		
 		const EResourceTransitionAccess CurrentAccess = Resource->GetCurrentGPUAccess();
diff --git a/Engine/Source/Runtime/Windows/D3D11RHI/Private/D3D11RHIPrivate.h b/Engine/Source/Runtime/Windows/D3D11RHI/Private/D3D11RHIPrivate.h
--- a/Engine/Source/Runtime/Windows/D3D11RHI/Private/D3D11RHIPrivate.h
+++ b/Engine/Source/Runtime/Windows/D3D11RHI/Private/D3D11RHIPrivate.h
@@ -4,8 +4,7 @@
 	D3D11RHIPrivate.h: Private D3D RHI definitions.
 =============================================================================*/
 
-#ifndef __D3D11RHIPRIVATE_H__
-#define __D3D11RHIPRIVATE_H__
+#pragma once
 
 #include "D3D11RHI.h"
 // Dependencies.
@@ -301,6 +300,8 @@ public:
 	virtual void Init() override;
 	virtual void Shutdown() override;
 
+	virtual void FlushPendingLogs() override;
+
 	template<typename TRHIType>
 	static FORCEINLINE typename TD3D11ResourceTraits<TRHIType>::TConcreteType* ResourceCast(TRHIType* Resource)
 	{
@@ -1319,6 +1320,3 @@ private:
 	// ordered from small to large (for efficient compactening)
 	TArray<FRange> Entries;
 };
-
-
-#endif
diff --git a/Engine/Source/Runtime/Windows/D3D11RHI/Private/Windows/WindowsD3D11Device.cpp b/Engine/Source/Runtime/Windows/D3D11RHI/Private/Windows/WindowsD3D11Device.cpp
--- a/Engine/Source/Runtime/Windows/D3D11RHI/Private/Windows/WindowsD3D11Device.cpp
+++ b/Engine/Source/Runtime/Windows/D3D11RHI/Private/Windows/WindowsD3D11Device.cpp
@@ -1,7 +1,7 @@
 // Copyright 1998-2016 Epic Games, Inc. All Rights Reserved.
 
 /*=============================================================================
-	D3D11Device.cpp: D3D device RHI implementation.
+	WindowsD3D11Device.cpp: Windows D3D device RHI implementation.
 =============================================================================*/
 
 #include "D3D11RHIPrivate.h"
@@ -269,6 +269,8 @@ void FD3D11DynamicRHIModule::FindAdapter()
 	bool bIsAnyIntel = false;
 	bool bIsAnyNVIDIA = false;
 
+	UE_LOG(LogD3D11RHI, Log, TEXT("D3D11 adapters:"));
+
 	// Enumerate the DXGIFactory's adapters.
 	for(uint32 AdapterIndex = 0; DXGIFactory1->EnumAdapters(AdapterIndex,TempAdapter.GetInitReference()) != DXGI_ERROR_NOT_FOUND; ++AdapterIndex)
 	{
@@ -284,13 +286,13 @@ void FD3D11DynamicRHIModule::FindAdapter()
 				uint32 OutputCount = CountAdapterOutputs(TempAdapter);
 
 				UE_LOG(LogD3D11RHI, Log,
-					TEXT("Found D3D11 adapter %u: %s (Feature Level %s)"),
+					TEXT("  %2u. '%s' (Feature Level %s)"),
 					AdapterIndex,
 					AdapterDesc.Description,
 					GetFeatureLevelString(ActualFeatureLevel)
 					);
 				UE_LOG(LogD3D11RHI, Log,
-					TEXT("Adapter has %uMB of dedicated video memory, %uMB of dedicated system memory, and %uMB of shared system memory, %d output[s]"),
+					TEXT("      %uMB of dedicated video memory, %uMB of dedicated system memory, and %uMB of shared system memory, %d output[s]"),
 					(uint32)(AdapterDesc.DedicatedVideoMemory / (1024*1024)),
 					(uint32)(AdapterDesc.DedicatedSystemMemory / (1024*1024)),
 					(uint32)(AdapterDesc.SharedSystemMemory / (1024*1024)),
@@ -361,7 +363,7 @@ void FD3D11DynamicRHIModule::FindAdapter()
 
 	if(ChosenAdapter.IsValid())
 	{
-		UE_LOG(LogD3D11RHI, Log, TEXT("Chosen D3D11 Adapter Id = %u"), ChosenAdapter.AdapterIndex);
+		UE_LOG(LogD3D11RHI, Log, TEXT("Chosen D3D11 Adapter: %u"), ChosenAdapter.AdapterIndex);
 	}
 	else
 	{
@@ -382,6 +384,43 @@ void FD3D11DynamicRHI::Init()
 	InitD3DDevice();
 }
 
+void FD3D11DynamicRHI::FlushPendingLogs()
+{
+#if !(UE_BUILD_SHIPPING && WITH_EDITOR)
+	if (D3D11RHI_ShouldCreateWithD3DDebug())
+	{
+		TRefCountPtr<ID3D11InfoQueue> InfoQueue = nullptr;
+		VERIFYD3D11RESULT(Direct3DDevice->QueryInterface(IID_ID3D11InfoQueue, (void**)InfoQueue.GetInitReference()));
+		if (InfoQueue)
+		{
+			FString FullMessage;
+			uint64 NumMessages = InfoQueue->GetNumStoredMessagesAllowedByRetrievalFilter();
+			for (uint64 Index = 0; Index < NumMessages; ++Index)
+			{
+				SIZE_T Length = 0;
+				if (SUCCEEDED(InfoQueue->GetMessage(Index, nullptr, &Length)))
+				{
+					TArray<uint8> Bytes;
+					Bytes.AddUninitialized((int32)Length);
+					D3D11_MESSAGE* Message = (D3D11_MESSAGE*)Bytes.GetData();
+					if (SUCCEEDED(InfoQueue->GetMessage(Index, Message, &Length)))
+					{
+						FullMessage += TEXT("\n\t");
+						FullMessage += Message->pDescription;
+					}
+				}
+			}
+
+			if (FullMessage.Len() > 0)
+			{
+				UE_LOG(LogD3D11RHI, Warning, TEXT("d3debug warnings/errors found:%s"), *FullMessage);
+			}
+			InfoQueue->ClearStoredMessages();
+		}
+	}
+#endif
+}
+
 void FD3D11DynamicRHI::InitD3DDevice()
 {
 	check( IsInGameThread() );
@@ -461,6 +500,15 @@ void FD3D11DynamicRHI::InitD3DDevice()
 
 					GRHIAdapterName = AdapterDesc.Description;
 					GRHIVendorId = AdapterDesc.VendorId;
+					
+					// get driver version (todo: share with other RHIs)
+					{
+						FPlatformMisc::GetGPUDriverInfo(GRHIAdapterName, GRHIAdapterInternalDriverVersion, GRHIAdapterUserDriverVersion, GRHIAdapterDriverDate);
+
+						UE_LOG(LogD3D11RHI, Log, TEXT("    Adapter Name: %s"), *GRHIAdapterName);
+						UE_LOG(LogD3D11RHI, Log, TEXT("  Driver Version: %s (internal %s)"), *GRHIAdapterUserDriverVersion, *GRHIAdapterInternalDriverVersion);
+						UE_LOG(LogD3D11RHI, Log, TEXT("     Driver Date: %s"), *GRHIAdapterDriverDate);
+					}
 
 					// Issue: 32bit windows doesn't report 64bit value, we take what we get.
 					FD3D11GlobalStats::GDedicatedVideoMemory = int64(AdapterDesc.DedicatedVideoMemory);
@@ -615,9 +663,9 @@ void FD3D11DynamicRHI::InitD3DDevice()
 		// Add some filter outs for known debug spew messages (that we don't care about)
 		if(DeviceFlags & D3D11_CREATE_DEVICE_DEBUG)
 		{
-			ID3D11InfoQueue *pd3dInfoQueue = NULL;
-			VERIFYD3D11RESULT(Direct3DDevice->QueryInterface( IID_ID3D11InfoQueue, (void**)&pd3dInfoQueue));
-			if(pd3dInfoQueue)
+			TRefCountPtr<ID3D11InfoQueue> InfoQueue;
+			VERIFYD3D11RESULT(Direct3DDevice->QueryInterface( IID_ID3D11InfoQueue, (void**)InfoQueue.GetInitReference()));
+			if (InfoQueue)
 			{
 				D3D11_INFO_QUEUE_FILTER NewFilter;
 				FMemory::Memzero(&NewFilter,sizeof(NewFilter));
@@ -658,20 +706,18 @@ void FD3D11DynamicRHI::InitD3DDevice()
 				NewFilter.DenyList.NumIDs = sizeof(DenyIds)/sizeof(D3D11_MESSAGE_ID);
 				NewFilter.DenyList.pIDList = (D3D11_MESSAGE_ID*)&DenyIds;
 
-				pd3dInfoQueue->PushStorageFilter(&NewFilter);
+				InfoQueue->PushStorageFilter(&NewFilter);
 
 				// Break on D3D debug errors.
-				pd3dInfoQueue->SetBreakOnSeverity(D3D11_MESSAGE_SEVERITY_ERROR,true);
+				InfoQueue->SetBreakOnSeverity(D3D11_MESSAGE_SEVERITY_ERROR,true);
 
 				// Enable this to break on a specific id in order to quickly get a callstack
-				//pd3dInfoQueue->SetBreakOnID(D3D11_MESSAGE_ID_DEVICE_DRAW_CONSTANT_BUFFER_TOO_SMALL, true);
+				//InfoQueue->SetBreakOnID(D3D11_MESSAGE_ID_DEVICE_DRAW_CONSTANT_BUFFER_TOO_SMALL, true);
 
-				if(FParse::Param(FCommandLine::Get(),TEXT("d3dbreakonwarning")))
+				if (FParse::Param(FCommandLine::Get(),TEXT("d3dbreakonwarning")))
 				{
-					pd3dInfoQueue->SetBreakOnSeverity(D3D11_MESSAGE_SEVERITY_WARNING,true);
+					InfoQueue->SetBreakOnSeverity(D3D11_MESSAGE_SEVERITY_WARNING,true);
 				}
-
-				pd3dInfoQueue->Release();
 			}
 		}
 #endif
diff --git a/Engine/Source/Runtime/Windows/XAudio2/Private/XAudio2Source.cpp b/Engine/Source/Runtime/Windows/XAudio2/Private/XAudio2Source.cpp
--- a/Engine/Source/Runtime/Windows/XAudio2/Private/XAudio2Source.cpp
+++ b/Engine/Source/Runtime/Windows/XAudio2/Private/XAudio2Source.cpp
@@ -373,6 +373,9 @@ bool FXAudio2SoundSource::CreateSource( void )
 	bUsingHRTFSpatialization = false;
 	bool bCreatedWithSpatializationEffect = false;
 	MaxEffectChainChannels = 0;
+	
+	// Set to nullptr in case the voice is not successfully created, the source won't be garbage
+	Source = nullptr;
 
 	if (CreateWithSpatializationEffect())
 	{
@@ -1729,8 +1732,10 @@ void FSpatializationHelper::CalculateDolbySurroundRate( const FVector& OrientFro
 	OrientFront.DiagnosticCheckNaN(TEXT("FSpatializationHelper: OrientFront"));
 	ListenerPosition.DiagnosticCheckNaN(TEXT("FSpatializationHelper: ListenerPosition"));
 	EmitterPosition.DiagnosticCheckNaN(TEXT("FSpatializationHelper: EmitterPosition"));
-	if (!FMath::IsFinite(OmniRadius))
+	static bool bLoggedOmniRadius = false;
+	if (!FMath::IsFinite(OmniRadius) && !bLoggedOmniRadius)
 	{
+		bLoggedOmniRadius = true;
 		const FString NaNorINF = FMath::IsNaN(OmniRadius) ? TEXT("NaN") : TEXT("INF");
 		UE_LOG(LogXAudio2, Warning, TEXT("OmniRadius generated a %s: %f"), *NaNorINF, OmniRadius);
 	}
@@ -1757,9 +1762,11 @@ void FSpatializationHelper::CalculateDolbySurroundRate( const FVector& OrientFro
 		OutVolumes[SpeakerIndex] *= DSPSettings.pMatrixCoefficients[SpeakerIndex];
 
 #if !UE_BUILD_SHIPPING && !UE_BUILD_TEST
+		static bool bLoggedDSPSettings = false;
 		// Detect and warn about NaN and INF volumes. XAudio does not do this internally and behavior is undefined.
-		if (!FMath::IsFinite(OutVolumes[SpeakerIndex]))
+		if (!FMath::IsFinite(OutVolumes[SpeakerIndex]) && !bLoggedDSPSettings)
 		{
+			bLoggedDSPSettings = true;
 			const FString NaNorINF = FMath::IsNaN(OutVolumes[SpeakerIndex]) ? TEXT("NaN") : TEXT("INF");
 			UE_LOG(LogXAudio2, Warning, TEXT("CalculateDolbySurroundRate generated a %s in channel %d. OmniRadius:%f MatrixCoefficient:%f"),
 				*NaNorINF, SpeakerIndex, OmniRadius, DSPSettings.pMatrixCoefficients[SpeakerIndex]);
diff --git a/Engine/Source/Runtime/Windows/XAudio2/Private/XAudio2Support.h b/Engine/Source/Runtime/Windows/XAudio2/Private/XAudio2Support.h
--- a/Engine/Source/Runtime/Windows/XAudio2/Private/XAudio2Support.h
+++ b/Engine/Source/Runtime/Windows/XAudio2/Private/XAudio2Support.h
@@ -695,9 +695,18 @@ struct FXAudioDeviceProperties
 		// And make sure there's no audio remaining the voice so when it's re-used it's fresh.
 		Voice->FlushSourceBuffers();
 
+#if XAUDIO2_SUPPORTS_SENDLIST
+		// Clear out the send effects (OutputVoices). When the voice gets reused, the old internal state might be invalid 
+		// when the new send effects are applied to the voice.
+		Voice->SetOutputVoices(nullptr);
+#endif
+
 		// Release the effect chain
 		Voice->SetEffectChain(nullptr);
-
+#if PLATFORM_WINDOWS
+		Voice->DestroyVoice();
+		return;
+#else
 		// See if there is an existing pool for this source voice
 		FSourceVoicePoolEntry* VoicePoolEntry = nullptr;
 		for (int32 i = 0; i < VoicePool.Num(); ++i)
@@ -723,6 +732,7 @@ struct FXAudioDeviceProperties
 			VoicePoolEntry->MaxEffectChainChannels = MaxEffectChainChannels;
 			VoicePool.Add(VoicePoolEntry);
 		}
+#endif
 	}
 };
 
