Copying //UE4/Orion-Staging to //UE4/Main (originated from //Orion/Dev-General @ 2831630)

diff --git a/Engine/Source/Runtime/Windows/D3D11RHI/Private/D3D11Shaders.cpp b/Engine/Source/Runtime/Windows/D3D11RHI/Private/D3D11Shaders.cpp
--- a/Engine/Source/Runtime/Windows/D3D11RHI/Private/D3D11Shaders.cpp
+++ b/Engine/Source/Runtime/Windows/D3D11RHI/Private/D3D11Shaders.cpp
@@ -43,9 +43,8 @@ FVertexShaderRHIRef FD3D11DynamicRHI::RHICreateVertexShader(const TArray<uint8>&
 	const uint8* CodePtr = Code.GetData() + Offset;
 	const size_t CodeSize = ShaderCode.GetActualShaderCodeSize() - Offset;
 
-	VERIFYD3D11RESULT( Direct3DDevice->CreateVertexShader( (void*)CodePtr, CodeSize, NULL, Shader->Resource.GetInitReference() ) );
-
 	ReadShaderOptionalData(ShaderCode, *Shader);
+	VERIFYD3D11SHADERRESULT(Direct3DDevice->CreateVertexShader( (void*)CodePtr, CodeSize, NULL, Shader->Resource.GetInitReference() ), Shader, Direct3DDevice);
 	
 	// TEMP
 	Shader->Code = Code;
@@ -66,9 +65,9 @@ FGeometryShaderRHIRef FD3D11DynamicRHI::RHICreateGeometryShader(const TArray<uin
 	const uint8* CodePtr = Code.GetData() + Offset;
 	const size_t CodeSize = ShaderCode.GetActualShaderCodeSize() - Offset;
 
-	VERIFYD3D11RESULT( Direct3DDevice->CreateGeometryShader( (void*)CodePtr, CodeSize, NULL, Shader->Resource.GetInitReference() ) );
-	
 	ReadShaderOptionalData(ShaderCode, *Shader);
+	VERIFYD3D11SHADERRESULT(Direct3DDevice->CreateGeometryShader( (void*)CodePtr, CodeSize, NULL, Shader->Resource.GetInitReference() ), Shader, Direct3DDevice);
+
 
 	return Shader;
 }
@@ -103,17 +102,19 @@ FGeometryShaderRHIRef FD3D11DynamicRHI::RHICreateGeometryShaderWithStreamOutput(
 		StreamOutEntries[EntryIndex].OutputSlot = ElementList[EntryIndex].OutputSlot;
 	}
 
-	VERIFYD3D11RESULT( Direct3DDevice->CreateGeometryShaderWithStreamOutput(
-		(void*)CodePtr,
-		CodeSize,
-		StreamOutEntries,
-		ElementList.Num(),
-		Strides,
-		NumStrides,
-		D3DRasterizedStream,
-		NULL,
-		Shader->Resource.GetInitReference()
-	) );
+	VERIFYD3D11SHADERRESULT( Direct3DDevice->CreateGeometryShaderWithStreamOutput(
+			(void*)CodePtr,
+			CodeSize,
+			StreamOutEntries,
+			ElementList.Num(),
+			Strides,
+			NumStrides,
+			D3DRasterizedStream,
+			NULL,
+			Shader->Resource.GetInitReference()
+			),
+		Shader,
+		Direct3DDevice);
 	
 	auto PackedResourceCounts = ShaderCode.FindOptionalData<FShaderCodePackedResourceCounts>();
 	check(PackedResourceCounts);
@@ -137,9 +138,8 @@ FHullShaderRHIRef FD3D11DynamicRHI::RHICreateHullShader(const TArray<uint8>& Cod
 	const uint8* CodePtr = Code.GetData() + Offset;
 	const size_t CodeSize = ShaderCode.GetActualShaderCodeSize() - Offset;
 
-	VERIFYD3D11RESULT( Direct3DDevice->CreateHullShader( (void*)CodePtr, CodeSize, NULL, Shader->Resource.GetInitReference() ) );
-	
 	ReadShaderOptionalData(ShaderCode, *Shader);
+	VERIFYD3D11SHADERRESULT( Direct3DDevice->CreateHullShader( (void*)CodePtr, CodeSize, NULL, Shader->Resource.GetInitReference() ), Shader, Direct3DDevice);
 
 	return Shader;
 }
@@ -156,9 +156,8 @@ FDomainShaderRHIRef FD3D11DynamicRHI::RHICreateDomainShader(const TArray<uint8>&
 	const uint8* CodePtr = Code.GetData() + Offset;
 	const size_t CodeSize = ShaderCode.GetActualShaderCodeSize() - Offset;
 
-	VERIFYD3D11RESULT( Direct3DDevice->CreateDomainShader( (void*)CodePtr, CodeSize, NULL, Shader->Resource.GetInitReference() ) );
-	
 	ReadShaderOptionalData(ShaderCode, *Shader);
+	VERIFYD3D11SHADERRESULT( Direct3DDevice->CreateDomainShader( (void*)CodePtr, CodeSize, NULL, Shader->Resource.GetInitReference() ), Shader, Direct3DDevice);
 
 	return Shader;
 }
@@ -175,9 +174,8 @@ FPixelShaderRHIRef FD3D11DynamicRHI::RHICreatePixelShader(const TArray<uint8>& C
 	const uint8* CodePtr = Code.GetData() + Offset;
 	const size_t CodeSize = ShaderCode.GetActualShaderCodeSize() - Offset;
 
-	VERIFYD3D11RESULT( Direct3DDevice->CreatePixelShader( (void*)CodePtr, CodeSize, NULL, Shader->Resource.GetInitReference() ) );
-	
 	ReadShaderOptionalData(ShaderCode, *Shader);
+	VERIFYD3D11SHADERRESULT(Direct3DDevice->CreatePixelShader( (void*)CodePtr, CodeSize, NULL, Shader->Resource.GetInitReference() ), Shader, Direct3DDevice);
 
 	return Shader;
 }
@@ -194,9 +192,8 @@ FComputeShaderRHIRef FD3D11DynamicRHI::RHICreateComputeShader(const TArray<uint8
 	const uint8* CodePtr = Code.GetData() + Offset;
 	const size_t CodeSize = ShaderCode.GetActualShaderCodeSize() - Offset;
 
-	VERIFYD3D11RESULT( Direct3DDevice->CreateComputeShader( (void*)CodePtr, CodeSize, NULL, Shader->Resource.GetInitReference() ) );
-	
 	ReadShaderOptionalData(ShaderCode, *Shader);
+	VERIFYD3D11SHADERRESULT( Direct3DDevice->CreateComputeShader( (void*)CodePtr, CodeSize, NULL, Shader->Resource.GetInitReference() ), Shader, Direct3DDevice);
 
 	return Shader;
 }
diff --git a/Engine/Source/Runtime/Windows/D3D11RHI/Private/D3D11Util.cpp b/Engine/Source/Runtime/Windows/D3D11RHI/Private/D3D11Util.cpp
--- a/Engine/Source/Runtime/Windows/D3D11RHI/Private/D3D11Util.cpp
+++ b/Engine/Source/Runtime/Windows/D3D11RHI/Private/D3D11Util.cpp
@@ -200,6 +200,28 @@ void VerifyD3D11Result(HRESULT D3DResult,const ANSICHAR* Code,const ANSICHAR* Fi
 	UE_LOG(LogD3D11RHI, Fatal,TEXT("%s failed \n at %s:%u \n with error %s"),ANSI_TO_TCHAR(Code),ANSI_TO_TCHAR(Filename),Line,*ErrorString);
 }
 
+void VerifyD3D11ShaderResult(FRHIShader* Shader, HRESULT D3DResult, const ANSICHAR* Code, const ANSICHAR* Filename, uint32 Line, ID3D11Device* Device)
+{
+	check(FAILED(D3DResult));
+
+	const FString& ErrorString = GetD3D11ErrorString(D3DResult, Device);
+
+#if !(UE_BUILD_SHIPPING || UE_BUILD_TEST)
+	if (Shader->ShaderName.Len())
+	{
+		UE_LOG(LogD3D11RHI, Error, TEXT("%s failed trying to create shader %s\n at %s:%u \n with error %s"), ANSI_TO_TCHAR(Code), *Shader->ShaderName, ANSI_TO_TCHAR(Filename), Line, *ErrorString);
+		TerminateOnDeviceRemoved(D3DResult);
+		TerminateOnOutOfMemory(D3DResult, false);
+
+		UE_LOG(LogD3D11RHI, Fatal, TEXT("%s failed trying to create shader %s \n at %s:%u \n with error %s"), ANSI_TO_TCHAR(Code), *Shader->ShaderName, ANSI_TO_TCHAR(Filename), Line, *ErrorString);
+	}
+	else
+#endif
+	{
+		VerifyD3D11Result(D3DResult, Code, Filename, Line, Device);
+	}
+}
+
 void VerifyD3D11CreateTextureResult(HRESULT D3DResult,const ANSICHAR* Code,const ANSICHAR* Filename,uint32 Line,uint32 SizeX,uint32 SizeY,uint32 SizeZ,uint8 Format,uint32 NumMips,uint32 Flags)
 {
 	check(FAILED(D3DResult));
diff --git a/Engine/Source/Runtime/Windows/D3D11RHI/Private/Windows/WindowsD3D11Device.cpp b/Engine/Source/Runtime/Windows/D3D11RHI/Private/Windows/WindowsD3D11Device.cpp
--- a/Engine/Source/Runtime/Windows/D3D11RHI/Private/Windows/WindowsD3D11Device.cpp
+++ b/Engine/Source/Runtime/Windows/D3D11RHI/Private/Windows/WindowsD3D11Device.cpp
@@ -448,12 +448,12 @@ void FD3D11DynamicRHI::InitD3DDevice()
 
 		bDeviceRemoved = false;
 
-		// Cleanup the D3D device.
-		CleanupD3DDevice();
-
 		// We currently don't support removed devices because FTexture2DResource can't recreate its RHI resources from scratch.
 		// We would also need to recreate the viewport swap chains from scratch.
 		UE_LOG(LogD3D11RHI, Fatal, TEXT("The Direct3D 11 device that was being used has been removed (Error: %d '%s').  Please restart the game."), hRes, Reason);
+
+		// UE4 no longer supports clean-up and recovery.
+		//CleanupD3DDevice();
 	}
 
 	// If we don't have a device yet, either because this is the first viewport, or the old device was removed, create a device.
@@ -500,6 +500,7 @@ void FD3D11DynamicRHI::InitD3DDevice()
 
 					GRHIAdapterName = AdapterDesc.Description;
 					GRHIVendorId = AdapterDesc.VendorId;
+					GRHIDeviceId = AdapterDesc.DeviceId;
 					
 					// get driver version (todo: share with other RHIs)
 					{
diff --git a/Engine/Source/Runtime/Windows/D3D11RHI/Public/D3D11Util.h b/Engine/Source/Runtime/Windows/D3D11RHI/Public/D3D11Util.h
--- a/Engine/Source/Runtime/Windows/D3D11RHI/Public/D3D11Util.h
+++ b/Engine/Source/Runtime/Windows/D3D11RHI/Public/D3D11Util.h
@@ -29,13 +29,24 @@
 
 /**
  * Checks that the given result isn't a failure.  If it is, the application exits with an appropriate error message.
- * @param	Result - The result code to check
+ * @param	Result - The result code to check.
  * @param	Code - The code which yielded the result.
  * @param	Filename - The filename of the source file containing Code.
  * @param	Line - The line number of Code within Filename.
  */
 extern D3D11RHI_API void VerifyD3D11Result(HRESULT Result,const ANSICHAR* Code,const ANSICHAR* Filename,uint32 Line, ID3D11Device* Device);
 
+/**
+ * Checks that the given result isn't a failure.  If it is, the application exits with an appropriate error message.
+ * @param	Shader - The shader we are trying to create.
+ * @param	Result - The result code to check.
+ * @param	Code - The code which yielded the result.
+ * @param	Filename - The filename of the source file containing Code.
+ * @param	Line - The line number of Code within Filename.
+ * @param	Device - The D3D device used to create the shadr.
+ */
+extern D3D11RHI_API void VerifyD3D11ShaderResult(class FRHIShader* Shader, HRESULT Result, const ANSICHAR* Code, const ANSICHAR* Filename, uint32 Line, ID3D11Device* Device);
+
 /**
 * Checks that the given result isn't a failure.  If it is, the application exits with an appropriate error message.
 * @param	Result - The result code to check
@@ -51,6 +62,7 @@ extern D3D11RHI_API void VerifyD3D11CreateTextureResult(HRESULT D3DResult,const
  */
 #define VERIFYD3D11RESULT_EX(x, Device)	{HRESULT hr = x; if (FAILED(hr)) { VerifyD3D11Result(hr,#x,__FILE__,__LINE__, Device); }}
 #define VERIFYD3D11RESULT(x)			{HRESULT hr = x; if (FAILED(hr)) { VerifyD3D11Result(hr,#x,__FILE__,__LINE__, 0); }}
+#define VERIFYD3D11SHADERRESULT(Result, Shader, Device) {HRESULT hr = (Result); if (FAILED(hr)) { VerifyD3D11ShaderResult(Shader, hr, #Result,__FILE__,__LINE__, Device); }}
 #define VERIFYD3D11CREATETEXTURERESULT(x,SizeX,SizeY,SizeZ,Format,NumMips,Flags) {HRESULT hr = x; if (FAILED(hr)) { VerifyD3D11CreateTextureResult(hr,#x,__FILE__,__LINE__,SizeX,SizeY,SizeZ,Format,NumMips,Flags); }}
 
 /**
diff --git a/Engine/Source/Runtime/Windows/D3D12RHI/Private/Windows/WindowsD3D12Device.cpp b/Engine/Source/Runtime/Windows/D3D12RHI/Private/Windows/WindowsD3D12Device.cpp
--- a/Engine/Source/Runtime/Windows/D3D12RHI/Private/Windows/WindowsD3D12Device.cpp
+++ b/Engine/Source/Runtime/Windows/D3D12RHI/Private/Windows/WindowsD3D12Device.cpp
@@ -402,6 +402,7 @@ void FD3D12DynamicRHI::PerRHISetup(FD3D12Device* MainDevice)
     // TODO MS: We need to do this after CreateDevice due to a bug in QueryVideoMemoryInfo
     GRHIAdapterName = AdapterDesc->Description;
     GRHIVendorId = AdapterDesc->VendorId;
+	GRHIDeviceId = AdapterDesc->DeviceId;
 
     // Issue: 32bit windows doesn't report 64bit value, we take what we get.
     FD3D12GlobalStats::GDedicatedVideoMemory = int64(AdapterDesc->DedicatedVideoMemory);
