Copying //UE4/Orion-Staging to //UE4/Main (//Orion/Dev-General @ 2855324)

diff --git a/Engine/Source/Runtime/Windows/D3D11RHI/Private/Windows/WindowsD3D11Device.cpp b/Engine/Source/Runtime/Windows/D3D11RHI/Private/Windows/WindowsD3D11Device.cpp
--- a/Engine/Source/Runtime/Windows/D3D11RHI/Private/Windows/WindowsD3D11Device.cpp
+++ b/Engine/Source/Runtime/Windows/D3D11RHI/Private/Windows/WindowsD3D11Device.cpp
@@ -11,7 +11,7 @@
 
 #include "HardwareInfo.h"
 #include "Runtime/HeadMountedDisplay/Public/IHeadMountedDisplayModule.h"
-
+#include "GenericPlatformDriver.h"			// FGPUDriverInfo
 
 extern bool D3D11RHI_ShouldCreateWithD3DDebug();
 extern bool D3D11RHI_ShouldAllowAsyncResourceCreation();
@@ -19,7 +19,7 @@ extern bool D3D11RHI_ShouldAllowAsyncResourceCreation();
 static TAutoConsoleVariable<int32> CVarGraphicsAdapter(
 	TEXT("r.GraphicsAdapter"),
 	-1,
-	TEXT("User request to pick a specific graphics adapter (e.g. when using a integrated graphics card with a descrete one)\n")
+	TEXT("User request to pick a specific graphics adapter (e.g. when using a integrated graphics card with a discrete one)\n")
 	TEXT("At the moment this only works on Direct3D 11.\n")
 	TEXT(" -2: Take the first one that fulfills the criteria\n")
 	TEXT(" -1: Favour non integrated because there are usually faster (default)\n")
@@ -264,6 +264,8 @@ void FD3D11DynamicRHIModule::FindAdapter()
 
 	FD3D11Adapter FirstWithoutIntegratedAdapter;
 	FD3D11Adapter FirstAdapter;
+	// indexed by AdapterIndex, we store it instead of query it later from the created device to prevent some Optimus bug reporting the data/name of the wrong adapter
+	TArray<DXGI_ADAPTER_DESC> AdapterDescription;
 
 	bool bIsAnyAMD = false;
 	bool bIsAnyIntel = false;
@@ -274,6 +276,10 @@ void FD3D11DynamicRHIModule::FindAdapter()
 	// Enumerate the DXGIFactory's adapters.
 	for(uint32 AdapterIndex = 0; DXGIFactory1->EnumAdapters(AdapterIndex,TempAdapter.GetInitReference()) != DXGI_ERROR_NOT_FOUND; ++AdapterIndex)
 	{
+		// to make sure the array elements can be indexed with AdapterIndex
+		DXGI_ADAPTER_DESC AdapterDesc;
+		ZeroMemory(&AdapterDesc, sizeof(DXGI_ADAPTER_DESC));
+
 		// Check that if adapter supports D3D11.
 		if(TempAdapter)
 		{
@@ -281,7 +287,6 @@ void FD3D11DynamicRHIModule::FindAdapter()
 			if(SafeTestD3D11CreateDevice(TempAdapter,MaxAllowedFeatureLevel,&ActualFeatureLevel))
 			{
 				// Log some information about the available D3D11 adapters.
-				DXGI_ADAPTER_DESC AdapterDesc;
 				VERIFYD3D11RESULT(TempAdapter->GetDesc(&AdapterDesc));
 				uint32 OutputCount = CountAdapterOutputs(TempAdapter);
 
@@ -344,6 +349,8 @@ void FD3D11DynamicRHIModule::FindAdapter()
 				}
 			}
 		}
+		
+		AdapterDescription.Add(AdapterDesc);
 	}
 
 	if(bFavorNonIntegrated && (bIsAnyAMD || bIsAnyNVIDIA))
@@ -363,6 +370,7 @@ void FD3D11DynamicRHIModule::FindAdapter()
 
 	if(ChosenAdapter.IsValid())
 	{
+		ChosenDescription = AdapterDescription[ChosenAdapter.AdapterIndex];
 		UE_LOG(LogD3D11RHI, Log, TEXT("Chosen D3D11 Adapter: %u"), ChosenAdapter.AdapterIndex);
 	}
 	else
@@ -376,7 +384,7 @@ FDynamicRHI* FD3D11DynamicRHIModule::CreateRHI()
 	TRefCountPtr<IDXGIFactory1> DXGIFactory1;
 	SafeCreateDXGIFactory(DXGIFactory1.GetInitReference());
 	check(DXGIFactory1);
-	return new FD3D11DynamicRHI(DXGIFactory1,ChosenAdapter.MaxSupportedFeatureLevel,ChosenAdapter.AdapterIndex);
+	return new FD3D11DynamicRHI(DXGIFactory1,ChosenAdapter.MaxSupportedFeatureLevel,ChosenAdapter.AdapterIndex,ChosenDescription);
 }
 
 void FD3D11DynamicRHI::Init()
@@ -493,101 +501,99 @@ void FD3D11DynamicRHI::InitD3DDevice()
 		{
 			if (EnumAdapter)// && EnumAdapter->CheckInterfaceSupport(__uuidof(ID3D11Device),NULL) == S_OK)
 			{
-				DXGI_ADAPTER_DESC AdapterDesc;
-				if (SUCCEEDED(EnumAdapter->GetDesc(&AdapterDesc)))
-				{
-					Adapter = EnumAdapter;
-
-					GRHIAdapterName = AdapterDesc.Description;
-					GRHIVendorId = AdapterDesc.VendorId;
-					GRHIDeviceId = AdapterDesc.DeviceId;
+				// we don't use AdapterDesc.Description as there is a bug with Optimus where it can report the wrong name
+				DXGI_ADAPTER_DESC AdapterDesc = ChosenDescription;
+				Adapter = EnumAdapter;
 
-					UE_LOG(LogD3D11RHI, Log, TEXT("    GPU DeviceId: 0x%x (for the marketing name, search the web for \"GPU Device Id\")"), 
-						AdapterDesc.DeviceId);
+				GRHIAdapterName = AdapterDesc.Description;
+				GRHIVendorId = AdapterDesc.VendorId;
+				GRHIDeviceId = AdapterDesc.DeviceId;
 
-					// get driver version (todo: share with other RHIs)
-					{
-						FPlatformMisc::GetGPUDriverInfo(GRHIAdapterName, GRHIAdapterInternalDriverVersion, GRHIAdapterUserDriverVersion, GRHIAdapterDriverDate);
+				UE_LOG(LogD3D11RHI, Log, TEXT("    GPU DeviceId: 0x%x (for the marketing name, search the web for \"GPU Device Id\")"), 
+					AdapterDesc.DeviceId);
 
-						UE_LOG(LogD3D11RHI, Log, TEXT("    Adapter Name: %s"), *GRHIAdapterName);
-						UE_LOG(LogD3D11RHI, Log, TEXT("  Driver Version: %s (internal %s)"), *GRHIAdapterUserDriverVersion, *GRHIAdapterInternalDriverVersion);
-						UE_LOG(LogD3D11RHI, Log, TEXT("     Driver Date: %s"), *GRHIAdapterDriverDate);
-					}
+				// get driver version (todo: share with other RHIs)
+				{
+					FGPUDriverInfo GPUDriverInfo = FPlatformMisc::GetGPUDriverInfo(GRHIAdapterName);
 
-					// Issue: 32bit windows doesn't report 64bit value, we take what we get.
-					FD3D11GlobalStats::GDedicatedVideoMemory = int64(AdapterDesc.DedicatedVideoMemory);
-					FD3D11GlobalStats::GDedicatedSystemMemory = int64(AdapterDesc.DedicatedSystemMemory);
-					FD3D11GlobalStats::GSharedSystemMemory = int64(AdapterDesc.SharedSystemMemory);
+					GRHIAdapterUserDriverVersion = GPUDriverInfo.UserDriverVersion;
+					GRHIAdapterInternalDriverVersion = GPUDriverInfo.InternalDriverVersion;
+					GRHIAdapterDriverDate = GPUDriverInfo.DriverDate;
 
-					// Total amount of system memory, clamped to 8 GB
-					int64 TotalPhysicalMemory = FMath::Min(int64(FPlatformMemory::GetConstants().TotalPhysicalGB), 8ll) * (1024ll * 1024ll * 1024ll);
+					UE_LOG(LogD3D11RHI, Log, TEXT("    Adapter Name: %s"), *GRHIAdapterName);
+					UE_LOG(LogD3D11RHI, Log, TEXT("  Driver Version: %s (internal:%s, unified:%s)"), *GRHIAdapterUserDriverVersion, *GRHIAdapterInternalDriverVersion, *GPUDriverInfo.GetUnifiedDriverVersion());
+					UE_LOG(LogD3D11RHI, Log, TEXT("     Driver Date: %s"), *GRHIAdapterDriverDate);
+				}
 
-					// Consider 50% of the shared memory but max 25% of total system memory.
-					int64 ConsideredSharedSystemMemory = FMath::Min( FD3D11GlobalStats::GSharedSystemMemory / 2ll, TotalPhysicalMemory / 4ll );
+				// Issue: 32bit windows doesn't report 64bit value, we take what we get.
+				FD3D11GlobalStats::GDedicatedVideoMemory = int64(AdapterDesc.DedicatedVideoMemory);
+				FD3D11GlobalStats::GDedicatedSystemMemory = int64(AdapterDesc.DedicatedSystemMemory);
+				FD3D11GlobalStats::GSharedSystemMemory = int64(AdapterDesc.SharedSystemMemory);
 
-					FD3D11GlobalStats::GTotalGraphicsMemory = 0;
-					if ( IsRHIDeviceIntel() )
-					{
-						// It's all system memory.
-						FD3D11GlobalStats::GTotalGraphicsMemory = FD3D11GlobalStats::GDedicatedVideoMemory;
-						FD3D11GlobalStats::GTotalGraphicsMemory += FD3D11GlobalStats::GDedicatedSystemMemory;
-						FD3D11GlobalStats::GTotalGraphicsMemory += ConsideredSharedSystemMemory;
-					}
-					else if ( FD3D11GlobalStats::GDedicatedVideoMemory >= 200*1024*1024 )
-					{
-						// Use dedicated video memory, if it's more than 200 MB
-						FD3D11GlobalStats::GTotalGraphicsMemory = FD3D11GlobalStats::GDedicatedVideoMemory;
-					}
-					else if ( FD3D11GlobalStats::GDedicatedSystemMemory >= 200*1024*1024 )
-					{
-						// Use dedicated system memory, if it's more than 200 MB
-						FD3D11GlobalStats::GTotalGraphicsMemory = FD3D11GlobalStats::GDedicatedSystemMemory;
-					}
-					else if ( FD3D11GlobalStats::GSharedSystemMemory >= 400*1024*1024 )
-					{
-						// Use some shared system memory, if it's more than 400 MB
-						FD3D11GlobalStats::GTotalGraphicsMemory = ConsideredSharedSystemMemory;
-					}
-					else
-					{
-						// Otherwise consider 25% of total system memory for graphics.
-						FD3D11GlobalStats::GTotalGraphicsMemory = TotalPhysicalMemory / 4ll;
-					}
+				// Total amount of system memory, clamped to 8 GB
+				int64 TotalPhysicalMemory = FMath::Min(int64(FPlatformMemory::GetConstants().TotalPhysicalGB), 8ll) * (1024ll * 1024ll * 1024ll);
 
-					if ( sizeof(SIZE_T) < 8 )
-					{
-						// Clamp to 1 GB if we're less than 64-bit
-						FD3D11GlobalStats::GTotalGraphicsMemory = FMath::Min( FD3D11GlobalStats::GTotalGraphicsMemory, 1024ll * 1024ll * 1024ll );
-					}
-					else
-					{
-						// Clamp to 1.9 GB if we're 64-bit
-						FD3D11GlobalStats::GTotalGraphicsMemory = FMath::Min( FD3D11GlobalStats::GTotalGraphicsMemory, 1945ll * 1024ll * 1024ll );
-					}
+				// Consider 50% of the shared memory but max 25% of total system memory.
+				int64 ConsideredSharedSystemMemory = FMath::Min( FD3D11GlobalStats::GSharedSystemMemory / 2ll, TotalPhysicalMemory / 4ll );
 
-					if ( GPoolSizeVRAMPercentage > 0 )
-					{
-						float PoolSize = float(GPoolSizeVRAMPercentage) * 0.01f * float(FD3D11GlobalStats::GTotalGraphicsMemory);
+				FD3D11GlobalStats::GTotalGraphicsMemory = 0;
+				if ( IsRHIDeviceIntel() )
+				{
+					// It's all system memory.
+					FD3D11GlobalStats::GTotalGraphicsMemory = FD3D11GlobalStats::GDedicatedVideoMemory;
+					FD3D11GlobalStats::GTotalGraphicsMemory += FD3D11GlobalStats::GDedicatedSystemMemory;
+					FD3D11GlobalStats::GTotalGraphicsMemory += ConsideredSharedSystemMemory;
+				}
+				else if ( FD3D11GlobalStats::GDedicatedVideoMemory >= 200*1024*1024 )
+				{
+					// Use dedicated video memory, if it's more than 200 MB
+					FD3D11GlobalStats::GTotalGraphicsMemory = FD3D11GlobalStats::GDedicatedVideoMemory;
+				}
+				else if ( FD3D11GlobalStats::GDedicatedSystemMemory >= 200*1024*1024 )
+				{
+					// Use dedicated system memory, if it's more than 200 MB
+					FD3D11GlobalStats::GTotalGraphicsMemory = FD3D11GlobalStats::GDedicatedSystemMemory;
+				}
+				else if ( FD3D11GlobalStats::GSharedSystemMemory >= 400*1024*1024 )
+				{
+					// Use some shared system memory, if it's more than 400 MB
+					FD3D11GlobalStats::GTotalGraphicsMemory = ConsideredSharedSystemMemory;
+				}
+				else
+				{
+					// Otherwise consider 25% of total system memory for graphics.
+					FD3D11GlobalStats::GTotalGraphicsMemory = TotalPhysicalMemory / 4ll;
+				}
 
-						// Truncate GTexturePoolSize to MB (but still counted in bytes)
-						GTexturePoolSize = int64(FGenericPlatformMath::TruncToFloat(PoolSize / 1024.0f / 1024.0f)) * 1024 * 1024;
+				if ( sizeof(SIZE_T) < 8 )
+				{
+					// Clamp to 1 GB if we're less than 64-bit
+					FD3D11GlobalStats::GTotalGraphicsMemory = FMath::Min( FD3D11GlobalStats::GTotalGraphicsMemory, 1024ll * 1024ll * 1024ll );
+				}
+				else
+				{
+					// Clamp to 1.9 GB if we're 64-bit
+					FD3D11GlobalStats::GTotalGraphicsMemory = FMath::Min( FD3D11GlobalStats::GTotalGraphicsMemory, 1945ll * 1024ll * 1024ll );
+				}
 
-						UE_LOG(LogRHI,Log,TEXT("Texture pool is %llu MB (%d%% of %llu MB)"),
-							GTexturePoolSize / 1024 / 1024,
-							GPoolSizeVRAMPercentage,
-							FD3D11GlobalStats::GTotalGraphicsMemory / 1024 / 1024);
-					}
+				if ( GPoolSizeVRAMPercentage > 0 )
+				{
+					float PoolSize = float(GPoolSizeVRAMPercentage) * 0.01f * float(FD3D11GlobalStats::GTotalGraphicsMemory);
 
-					const bool bIsPerfHUD = !FCString::Stricmp(AdapterDesc.Description,TEXT("NVIDIA PerfHUD"));
+					// Truncate GTexturePoolSize to MB (but still counted in bytes)
+					GTexturePoolSize = int64(FGenericPlatformMath::TruncToFloat(PoolSize / 1024.0f / 1024.0f)) * 1024 * 1024;
 
-					if(bIsPerfHUD)
-					{
-						DriverType =  D3D_DRIVER_TYPE_REFERENCE;
-					}
+					UE_LOG(LogRHI,Log,TEXT("Texture pool is %llu MB (%d%% of %llu MB)"),
+						GTexturePoolSize / 1024 / 1024,
+						GPoolSizeVRAMPercentage,
+						FD3D11GlobalStats::GTotalGraphicsMemory / 1024 / 1024);
 				}
-				else
+
+				const bool bIsPerfHUD = !FCString::Stricmp(AdapterDesc.Description,TEXT("NVIDIA PerfHUD"));
+
+				if(bIsPerfHUD)
 				{
-					check(!"Internal error, GetDesc() failed but before it worked")
+					DriverType =  D3D_DRIVER_TYPE_REFERENCE;
 				}
 			}
 		}
diff --git a/Engine/Source/Runtime/Windows/XAudio2/Private/XAudio2Source.cpp b/Engine/Source/Runtime/Windows/XAudio2/Private/XAudio2Source.cpp
--- a/Engine/Source/Runtime/Windows/XAudio2/Private/XAudio2Source.cpp
+++ b/Engine/Source/Runtime/Windows/XAudio2/Private/XAudio2Source.cpp
@@ -93,8 +93,6 @@ void FXAudio2SoundSource::FreeResources( void )
 	// Release voice.
 	if (Source)
 	{
-		// Because XAudio2's DestroyVoice source is blocking and can be slow on some processors (e.g. AMD), we're creating
-		// a task that destroys the voice on a separate thread to avoid blocking or hitching.
 		AudioDevice->DeviceProperties->ReleaseSourceVoice(Source, XAudio2Buffer->PCM, MaxEffectChainChannels);
 		Source = nullptr;
 	}
@@ -440,7 +438,6 @@ bool FXAudio2SoundSource::CreateSource( void )
 #if XAUDIO2_SUPPORTS_SENDLIST
 		Source->SetOutputVoices(&SourceSendList);
 #endif
-		Source->SetEffectChain(nullptr);
 	}
 
 	return true;
diff --git a/Engine/Source/Runtime/Windows/XAudio2/Private/XAudio2Support.h b/Engine/Source/Runtime/Windows/XAudio2/Private/XAudio2Support.h
--- a/Engine/Source/Runtime/Windows/XAudio2/Private/XAudio2Support.h
+++ b/Engine/Source/Runtime/Windows/XAudio2/Private/XAudio2Support.h
@@ -21,6 +21,11 @@
 #ifndef XAUDIO2_SUPPORTS_SENDLIST
 	#define XAUDIO2_SUPPORTS_SENDLIST			1
 #endif	//XAUDIO2_SUPPORTS_SENDLIST
+#ifndef XAUDIO2_SUPPORTS_VOICE_POOL
+	#define XAUDIO2_SUPPORTS_VOICE_POOL			0
+#endif	//XAUDIO2_SUPPORTS_VOICE_POOL
+
+
 
 /*------------------------------------------------------------------------------------
 	XAudio2 system headers
@@ -713,6 +718,7 @@ struct FXAudioDeviceProperties
 	{
 		bool bSuccess = false;
 
+#if XAUDIO2_SUPPORTS_VOICE_POOL
 		// First find the pool for the given format
 		FSourceVoicePoolEntry* VoicePoolEntry = nullptr;
 		for (int32 i = 0; i < VoicePool.Num(); ++i)
@@ -743,6 +749,11 @@ struct FXAudioDeviceProperties
 			bSuccess = Validate(TEXT("GetFreeSourceVoice, XAudio2->CreateSourceVoice"),
 								XAudio2->CreateSourceVoice(Voice, &BufferInfo.PCMFormat, XAUDIO2_VOICE_USEFILTER, MAX_PITCH, &SourceCallback, nullptr, EffectChain));
 		}
+#else // XAUDIO2_SUPPORTS_VOICE_POOL
+		check(XAudio2 != nullptr);
+		bSuccess = Validate(TEXT("GetFreeSourceVoice, XAudio2->CreateSourceVoice"),
+							XAudio2->CreateSourceVoice(Voice, &BufferInfo.PCMFormat, XAUDIO2_VOICE_USEFILTER, MAX_PITCH, &SourceCallback, nullptr, EffectChain));
+#endif // XAUDIO2_SUPPORTS_VOICE_POOL
 
 		if (bSuccess)
 		{
@@ -759,6 +770,8 @@ struct FXAudioDeviceProperties
 	/** Releases the voice into a pool of free voices according to the voice format and the max effect chain channels */
 	void ReleaseSourceVoice(IXAudio2SourceVoice* Voice, const FPCMBufferInfo& BufferInfo, const int32 MaxEffectChainChannels)
 	{
+#if XAUDIO2_SUPPORTS_VOICE_POOL
+
 		check(Voice != nullptr);
 
 		// Make sure the voice is stopped
@@ -809,6 +822,9 @@ struct FXAudioDeviceProperties
 				Voice->DestroyVoice();
 			}
 		}
+#else // XAUDIO2_SUPPORTS_VOICE_POOL
+		Voice->DestroyVoice();
+#endif // XAUDIO2_SUPPORTS_VOICE_POOL
 
 		--NumActiveVoices;
 	}
