Copying //UE4/Dev-Rendering to Dev-Main (//UE4/Dev-Main)

diff --git a/Engine/Source/Runtime/Renderer/Private/ShadowSetup.cpp b/Engine/Source/Runtime/Renderer/Private/ShadowSetup.cpp
--- a/Engine/Source/Runtime/Renderer/Private/ShadowSetup.cpp
+++ b/Engine/Source/Runtime/Renderer/Private/ShadowSetup.cpp
@@ -177,11 +177,9 @@ FAutoConsoleCommand CmdDumpShadowDumpSetup(
  *
  * @return	fade value between 0 and 1
  */
-float CalculateShadowFadeAlpha(float MaxUnclampedResolution, int32 ShadowFadeResolution, int32 MinShadowResolution)
+float CalculateShadowFadeAlpha(const float MaxUnclampedResolution, const uint32 ShadowFadeResolution, const uint32 MinShadowResolution)
 {
-	check(MaxUnclampedResolution >= 0);
-	check(ShadowFadeResolution >= 0);
-	check(MinShadowResolution >= 0);
+	// NB: MaxUnclampedResolution < 0 will return FadeAlpha = 0.0f. 
 
 	float FadeAlpha = 0.0f;
 	// Shadow size is above fading resolution.
@@ -192,11 +190,23 @@ float CalculateShadowFadeAlpha(float MaxUnclampedResolution, int32 ShadowFadeRes
 	// Shadow size is below fading resolution but above min resolution.
 	else if (MaxUnclampedResolution > MinShadowResolution)
 	{
-		const float InverseRange = 1.0f / (ShadowFadeResolution - MinShadowResolution);
-		const float FirstFadeValue = FMath::Pow(InverseRange, CVarShadowFadeExponent.GetValueOnRenderThread());
-		const float SizeRatio = (float)(MaxUnclampedResolution - MinShadowResolution) * InverseRange;
-		// Rescale the fade alpha to reduce the change between no fading and the first value, which reduces popping with small ShadowFadeExponent's
-		FadeAlpha = (FMath::Pow(SizeRatio, CVarShadowFadeExponent.GetValueOnRenderThread()) - FirstFadeValue) / (1.0f - FirstFadeValue);
+		const float Exponent = CVarShadowFadeExponent.GetValueOnRenderThread();
+		
+		// Use the limit case ShadowFadeResolution = MinShadowResolution
+		// to gracefully handle this case.
+		if (MinShadowResolution >= ShadowFadeResolution)
+		{
+			const float SizeRatio = (float)(MaxUnclampedResolution - MinShadowResolution);
+			FadeAlpha = 1.0f - FMath::Pow(SizeRatio, Exponent);
+		} 
+		else
+		{
+			const float InverseRange = 1.0f / (ShadowFadeResolution - MinShadowResolution);
+			const float FirstFadeValue = FMath::Pow(InverseRange, Exponent);
+			const float SizeRatio = (float)(MaxUnclampedResolution - MinShadowResolution) * InverseRange;
+			// Rescale the fade alpha to reduce the change between no fading and the first value, which reduces popping with small ShadowFadeExponent's
+			FadeAlpha = (FMath::Pow(SizeRatio, Exponent) - FirstFadeValue) / (1.0f - FirstFadeValue);
+		}
 	}
 	return FadeAlpha;
 }
@@ -492,7 +502,7 @@ bool FProjectedShadowInfo::SetupPerObjectProjection(
 
 		MinPreSubjectZ = Initializer.MinLightW;
 
-		ResolutionY = FMath::Min<uint32>(FMath::TruncToInt(InResolutionX / AspectRatio), MaxShadowResolutionY);
+		ResolutionY = FMath::Clamp<uint32>(FMath::TruncToInt(InResolutionX / AspectRatio), 1, MaxShadowResolutionY);
 
 		if (ResolutionX == 0 || ResolutionY == 0)
 		{
@@ -978,8 +988,14 @@ void FProjectedShadowInfo::GatherDynamicMeshElements(FSceneRenderer& Renderer, F
 		// Shadow must be visible in at least one view
 		if (bVisibleInAnyView)
 		{		
-			check(FoundView && IsInRenderingThread() 
-				&& !FRHICommandListImmediate::AnyRenderThreadTasksOutstanding()); // this would be bad because they probably rely on the view we are hacking
+			check(FoundView && IsInRenderingThread());
+			bool bMakeViewSnapshot = FRHICommandListImmediate::AnyRenderThreadTasksOutstanding(); // if there are task in flight, we need to clone
+
+			if (bMakeViewSnapshot)
+			{
+				FoundView = FoundView->CreateSnapshot();
+			}
+
 
 		    // Backup properties of the view that we will override
 		    FMatrix OriginalViewMatrix = FoundView->ViewMatrices.ViewMatrix;
@@ -1021,9 +1037,12 @@ void FProjectedShadowInfo::GatherDynamicMeshElements(FSceneRenderer& Renderer, F
 		    FoundView->ViewMatrices.GetDynamicMeshElementsShadowCullFrustum = (Disable & 4) ? &NoCull : &CasterFrustum;
 		    GatherDynamicMeshElementsArray(FoundView, Renderer, SubjectTranslucentPrimitives, DynamicSubjectTranslucentMeshElements, ReusedViewsArray);
     
-		    FoundView->ViewMatrices.GetDynamicMeshElementsShadowCullFrustum = nullptr;
-    
-		    FoundView->ViewMatrices.ViewMatrix = OriginalViewMatrix;
+			if (!bMakeViewSnapshot)
+			{
+				FoundView->ViewMatrices.GetDynamicMeshElementsShadowCullFrustum = nullptr;
+
+				FoundView->ViewMatrices.ViewMatrix = OriginalViewMatrix;
+			}
 
 			Renderer.MeshCollector.ProcessTasks();
 	    }
@@ -1402,13 +1421,16 @@ void FSceneRenderer::CreatePerObjectProjectedShadow(
 	FSceneRenderTargets& SceneContext = FSceneRenderTargets::Get(RHICmdList);
 	
 	// Shadowing constants.
-	const uint32 MinShadowResolution = CVarMinShadowResolution.GetValueOnRenderThread();
+	
 	const uint32 MaxShadowResolutionSetting = GetCachedScalabilityCVars().MaxShadowResolution;
 	const FIntPoint ShadowBufferResolution = SceneContext.GetShadowDepthTextureResolution();
 	const uint32 MaxShadowResolution = FMath::Min<int32>(MaxShadowResolutionSetting, ShadowBufferResolution.X) - SHADOW_BORDER * 2;
 	const uint32 MaxShadowResolutionY = FMath::Min<int32>(MaxShadowResolutionSetting, ShadowBufferResolution.Y) - SHADOW_BORDER * 2;
-	const int32 ShadowFadeResolution = CVarShadowFadeResolution.GetValueOnRenderThread();
-
+	const uint32 MinShadowResolution     = FMath::Max<int32>(0, CVarMinShadowResolution.GetValueOnRenderThread());
+	const uint32 ShadowFadeResolution    = FMath::Max<int32>(0, CVarShadowFadeResolution.GetValueOnRenderThread());
+	const uint32 MinPreShadowResolution  = FMath::Max<int32>(0, CVarMinPreShadowResolution.GetValueOnRenderThread());
+	const uint32 PreShadowFadeResolution = FMath::Max<int32>(0, CVarPreShadowFadeResolution.GetValueOnRenderThread());
+	
 	// Compute the maximum resolution required for the shadow by any view. Also keep track of the unclamped resolution for fading.
 	uint32 MaxDesiredResolution = 0;
 	float MaxUnclampedResolution = 0;
@@ -1457,7 +1479,8 @@ void FSceneRenderer::CreatePerObjectProjectedShadow(
 		MaxResolutionFadeAlpha = FMath::Max(MaxResolutionFadeAlpha, ViewSpecificAlpha);
 		ResolutionFadeAlphas.Add(ViewSpecificAlpha);
 
-		const float ViewSpecificPreShadowAlpha = CalculateShadowFadeAlpha( UnclampedResolution * CVarPreShadowResolutionFactor.GetValueOnRenderThread(), CVarPreShadowFadeResolution.GetValueOnRenderThread(), CVarMinPreShadowResolution.GetValueOnRenderThread() );
+	
+		const float ViewSpecificPreShadowAlpha = CalculateShadowFadeAlpha(UnclampedResolution * CVarPreShadowResolutionFactor.GetValueOnRenderThread(), PreShadowFadeResolution, MinPreShadowResolution);
 		MaxResolutionPreShadowFadeAlpha = FMath::Max(MaxResolutionPreShadowFadeAlpha, ViewSpecificPreShadowAlpha);
 		ResolutionPreShadowFadeAlphas.Add(ViewSpecificPreShadowAlpha);
 	}
@@ -1678,12 +1701,12 @@ void FDeferredShadingSceneRenderer::CreateWholeSceneProjectedShadow(FLightSceneI
 
 		// Shadow resolution constants.
 		const uint32 EffectiveDoubleShadowBorder = ProjectedShadowInitializers[0].CascadeSettings.bOnePassPointLightShadow ? 0 : SHADOW_BORDER * 2;
-		const int32 MinShadowResolution = CVarMinShadowResolution.GetValueOnRenderThread();
+		const uint32 MinShadowResolution = FMath::Max<int32>(0, CVarMinShadowResolution.GetValueOnRenderThread());
 		const int32 MaxShadowResolutionSetting = GetCachedScalabilityCVars().MaxShadowResolution;
 		const FIntPoint ShadowBufferResolution = SceneContext_ConstantsOnly.GetShadowDepthTextureResolution();
 		const uint32 MaxShadowResolution = FMath::Min(MaxShadowResolutionSetting, ShadowBufferResolution.X) - EffectiveDoubleShadowBorder;
 		const uint32 MaxShadowResolutionY = FMath::Min(MaxShadowResolutionSetting, ShadowBufferResolution.Y) - EffectiveDoubleShadowBorder;
-		const int32 ShadowFadeResolution = CVarShadowFadeResolution.GetValueOnRenderThread();
+		const uint32 ShadowFadeResolution = FMath::Max<int32>(0, CVarShadowFadeResolution.GetValueOnRenderThread());
 
 		// Compute the maximum resolution required for the shadow by any view. Also keep track of the unclamped resolution for fading.
 		float MaxDesiredResolution = 0;
diff --git a/Engine/Source/Runtime/Renderer/Private/TexCoordScaleAnalysisRendering.cpp b/Engine/Source/Runtime/Renderer/Private/TexCoordScaleAnalysisRendering.cpp
new file mode 100644
--- /dev/null
+++ b/Engine/Source/Runtime/Renderer/Private/TexCoordScaleAnalysisRendering.cpp
@@ -0,0 +1,72 @@
+// Copyright 1998-2016 Epic Games, Inc. All Rights Reserved.
+
+/*=============================================================================
+WantedMipsAccuracyRendering.cpp: Contains definitions for rendering the viewmode.
+=============================================================================*/
+
+#include "RendererPrivate.h"
+#include "ScenePrivate.h"
+#include "TexCoordScaleAnalysisRendering.h"
+
+static TAutoConsoleVariable<int32> CVarStreamingAnalysisIndex(
+	TEXT("r.Streaming.AnalysisIndex"),
+	0,
+	TEXT("Analysis index of the texture coordinate scale accuracy viewmode."),
+	ECVF_Default);
+
+IMPLEMENT_MATERIAL_SHADER_TYPE(,FTexCoordScaleAnalysisPS,TEXT("TexCoordScaleAnalysisPixelShader"),TEXT("Main"),SF_Pixel);
+
+void FTexCoordScaleAnalysisPS::SetParameters(
+	FRHICommandList& RHICmdList, 
+	const FShader* OriginalVS, 
+	const FShader* OriginalPS, 
+	const FMaterialRenderProxy* MaterialRenderProxy,
+	const FMaterial& Material,
+	const FSceneView& View
+	)
+{
+	const int32 NumEngineColors = FMath::Min<int32>(GEngine->StreamingAccuracyColors.Num(), NumStreamingAccuracyColors);
+	int32 ColorIndex = 0;
+	for (; ColorIndex < NumEngineColors; ++ColorIndex)
+	{
+		SetShaderValue(RHICmdList, FMeshMaterialShader::GetPixelShader(), AccuracyColorsParameter, GEngine->StreamingAccuracyColors[ColorIndex], ColorIndex);
+	}
+	for (; ColorIndex < NumStreamingAccuracyColors; ++ColorIndex)
+	{
+		SetShaderValue(RHICmdList, FMeshMaterialShader::GetPixelShader(), AccuracyColorsParameter, FLinearColor::Black, ColorIndex);
+	}
+
+	FMeshMaterialShader::SetParameters(RHICmdList, FMeshMaterialShader::GetPixelShader(), MaterialRenderProxy, Material, View, ESceneRenderTargetsMode::DontSet);
+}
+
+void FTexCoordScaleAnalysisPS::SetMesh(
+	FRHICommandList& RHICmdList, 
+	const FVertexFactory* VertexFactory,
+	const FSceneView& View,
+	const FPrimitiveSceneProxy* Proxy,
+	const FMeshBatchElement& BatchElement, 
+	const FMeshDrawingRenderState& DrawRenderState
+	)
+{
+	// This is taken from FStreamingHandlerTextureStatic::GetWantedMips
+
+	float CPUWantedMips = -1.f;
+	if (Proxy && Proxy->GetWorldTexelFactor() > 0)
+	{
+		FVector ViewToObject = Proxy->GetBounds().Origin - View.ViewMatrices.ViewOrigin;
+
+		float DistSqMinusRadiusSq = 0;
+		float ViewSize = View.ViewRect.Size().X;
+
+		DistSqMinusRadiusSq = FMath::Max(1.f, DistSqMinusRadiusSq);
+
+		float CoordSize = Proxy->GetWorldTexelFactor() * FMath::InvSqrt(DistSqMinusRadiusSq) * ViewSize;
+
+		CPUWantedMips = FMath::Clamp<float>(FMath::Log2(CoordSize), 0.f, (float)MaxStreamingAccuracyMips);
+	}
+
+	SetShaderValue(RHICmdList, FMeshMaterialShader::GetPixelShader(), TextureAnalysisIndexParameter, (int32)CVarStreamingAnalysisIndex.GetValueOnRenderThread());
+	SetShaderValue(RHICmdList, FMeshMaterialShader::GetPixelShader(), CPUScaleParameter, 1.f);
+
+	FMeshMaterialShader::SetMesh(RHICmdList, FMeshMaterialShader::GetPixelShader(), VertexFactory, View, Proxy, BatchElement, DrawRenderState);
+}
diff --git a/Engine/Source/Runtime/Renderer/Private/TexCoordScaleAnalysisRendering.h b/Engine/Source/Runtime/Renderer/Private/TexCoordScaleAnalysisRendering.h
new file mode 100644
--- /dev/null
+++ b/Engine/Source/Runtime/Renderer/Private/TexCoordScaleAnalysisRendering.h
@@ -0,0 +1,78 @@
+// Copyright 1998-2016 Epic Games, Inc. All Rights Reserved.
+
+/*=============================================================================
+TexCoordScaleAnalysisRendering.h: Declarations used for the viewmode.
+=============================================================================*/
+
+#pragma once
+
+/**
+* Pixel shader that renders texcoord scales.
+* The shader is only compiled with the local vertex factory to prevent multiple compilation.
+* Nothing from the factory is actually used, but the shader must still derive from FMeshMaterialShader.
+* This is required to call FMeshMaterialShader::SetMesh and bind primitive related data.
+*/
+class FTexCoordScaleAnalysisPS : public FMeshMaterialShader, public IDebugViewModePSInterface
+{
+	DECLARE_SHADER_TYPE(FTexCoordScaleAnalysisPS,MeshMaterial);
+
+public:
+
+	static bool ShouldCache(EShaderPlatform Platform, const FMaterial* Material, const FVertexFactoryType* VertexFactoryType)
+	{
+		// Debug viewmode pixel shaders have predefined interpolants, so we only need to compile it for one vertex factory.
+		return Material->GetFriendlyName().Contains(TEXT("FDebugViewModeMaterialProxy")) &&  AllowDebugViewModeShader(Platform) && VertexFactoryType == FindVertexFactoryType(FName(TEXT("FLocalVertexFactory"), FNAME_Find));
+	}
+
+	FTexCoordScaleAnalysisPS(const ShaderMetaType::CompiledShaderInitializerType& Initializer):
+		FMeshMaterialShader(Initializer)
+	{
+		AccuracyColorsParameter.Bind(Initializer.ParameterMap,TEXT("AccuracyColors"));
+		TextureAnalysisIndexParameter.Bind(Initializer.ParameterMap,TEXT("TextureAnalysisIndex"));
+		CPUScaleParameter.Bind(Initializer.ParameterMap,TEXT("CPUScale"));
+	}
+
+	FTexCoordScaleAnalysisPS() {}
+
+	virtual bool Serialize(FArchive& Ar) override
+	{
+		bool bShaderHasOutdatedParameters = FMeshMaterialShader::Serialize(Ar);
+		Ar << AccuracyColorsParameter;
+		Ar << TextureAnalysisIndexParameter;
+		Ar << CPUScaleParameter;
+		return bShaderHasOutdatedParameters;
+	}
+
+	static void ModifyCompilationEnvironment(EShaderPlatform Platform, const FMaterial* Material, FShaderCompilerEnvironment& OutEnvironment)
+	{
+		FMeshMaterialShader::ModifyCompilationEnvironment(Platform, OutEnvironment);
+	}
+
+	virtual void SetParameters(
+		FRHICommandList& RHICmdList, 
+		const FShader* OriginalVS, 
+		const FShader* OriginalPS, 
+		const FMaterialRenderProxy* MaterialRenderProxy,
+		const FMaterial& Material,
+		const FSceneView& View
+		) override;
+
+	virtual void SetMesh(
+		FRHICommandList& RHICmdList, 
+		const FVertexFactory* VertexFactory,
+		const FSceneView& View,
+		const FPrimitiveSceneProxy* Proxy,
+		const FMeshBatchElement& BatchElement, 
+		const FMeshDrawingRenderState& DrawRenderState
+		) override;
+
+	virtual void SetMesh(FRHICommandList& RHICmdList, const FSceneView& View) override { check(false); }
+
+	virtual FShader* GetShader() override { return static_cast<FShader*>(this); }
+
+private:
+
+	FShaderParameter AccuracyColorsParameter;
+	FShaderParameter TextureAnalysisIndexParameter;
+	FShaderParameter CPUScaleParameter;
+};
diff --git a/Engine/Source/Runtime/Renderer/Private/TexelFactorAccuracyRendering.cpp b/Engine/Source/Runtime/Renderer/Private/TexelFactorAccuracyRendering.cpp
--- a/Engine/Source/Runtime/Renderer/Private/TexelFactorAccuracyRendering.cpp
+++ b/Engine/Source/Runtime/Renderer/Private/TexelFactorAccuracyRendering.cpp
@@ -1,4 +1,4 @@
-// Copyright 1998-2015 Epic Games, Inc. All Rights Reserved.
+// Copyright 1998-2016 Epic Games, Inc. All Rights Reserved.
 
 /*=============================================================================
 TexelFactorAccuracyRendering.cpp: Contains definitions for rendering the viewmode.
@@ -6,6 +6,7 @@ TexelFactorAccuracyRendering.cpp: Contains definitions for rendering the viewmod
 
 #include "RendererPrivate.h"
 #include "ScenePrivate.h"
+#include "TexelFactorAccuracyRendering.h"
 
 IMPLEMENT_SHADER_TYPE(,FTexelFactorAccuracyPS,TEXT("TexelFactorAccuracyPixelShader"),TEXT("Main"),SF_Pixel);
 
@@ -13,6 +14,8 @@ void FTexelFactorAccuracyPS::SetParameters(
 	FRHICommandList& RHICmdList, 
 	const FShader* OriginalVS, 
 	const FShader* OriginalPS, 
+	const FMaterialRenderProxy* MaterialRenderProxy,
+	const FMaterial& Material,
 	const FSceneView& View
 	)
 {
@@ -23,11 +26,11 @@ void FTexelFactorAccuracyPS::SetParameters(
 	int32 ColorIndex = 0;
 	for (; ColorIndex < NumEngineColors; ++ColorIndex)
 	{
-		SetShaderValue(RHICmdList, FGlobalShader::GetPixelShader(), StreamingAccuracyColorsParameter, GEngine->StreamingAccuracyColors[ColorIndex], ColorIndex);
+		SetShaderValue(RHICmdList, FGlobalShader::GetPixelShader(), AccuracyColorsParameter, GEngine->StreamingAccuracyColors[ColorIndex], ColorIndex);
 	}
 	for (; ColorIndex < NumStreamingAccuracyColors; ++ColorIndex)
 	{
-		SetShaderValue(RHICmdList, FGlobalShader::GetPixelShader(), StreamingAccuracyColorsParameter, FLinearColor::Black, ColorIndex);
+		SetShaderValue(RHICmdList, FGlobalShader::GetPixelShader(), AccuracyColorsParameter, FLinearColor::Black, ColorIndex);
 	}
 }
 
diff --git a/Engine/Source/Runtime/Renderer/Private/TexelFactorAccuracyRendering.h b/Engine/Source/Runtime/Renderer/Private/TexelFactorAccuracyRendering.h
--- a/Engine/Source/Runtime/Renderer/Private/TexelFactorAccuracyRendering.h
+++ b/Engine/Source/Runtime/Renderer/Private/TexelFactorAccuracyRendering.h
@@ -1,4 +1,4 @@
-// Copyright 1998-2015 Epic Games, Inc. All Rights Reserved.
+// Copyright 1998-2016 Epic Games, Inc. All Rights Reserved.
 
 /*=============================================================================
 TexelFactorAccuracyRendering.h: Declarations used for the viewmode.
@@ -9,7 +9,7 @@ TexelFactorAccuracyRendering.h: Declarations used for the viewmode.
 /**
 * Pixel shader that renders the accuracy of the texel factor.
 */
-class FTexelFactorAccuracyPS : public FDebugViewModePS
+class FTexelFactorAccuracyPS : public FGlobalShader, public IDebugViewModePSInterface
 {
 	DECLARE_SHADER_TYPE(FTexelFactorAccuracyPS,Global);
 
@@ -21,9 +21,9 @@ public:
 	}
 
 	FTexelFactorAccuracyPS(const ShaderMetaType::CompiledShaderInitializerType& Initializer):
-		FDebugViewModePS(Initializer)
+		FGlobalShader(Initializer)
 	{
-		StreamingAccuracyColorsParameter.Bind(Initializer.ParameterMap,TEXT("StreamingAccuracyColors"));
+		AccuracyColorsParameter.Bind(Initializer.ParameterMap,TEXT("AccuracyColors"));
 		CPUTexelFactorParameter.Bind(Initializer.ParameterMap,TEXT("CPUTexelFactor"));
 	}
 
@@ -32,7 +32,7 @@ public:
 	virtual bool Serialize(FArchive& Ar) override
 	{
 		bool bShaderHasOutdatedParameters = FGlobalShader::Serialize(Ar);
-		Ar << StreamingAccuracyColorsParameter;
+		Ar << AccuracyColorsParameter;
 		Ar << CPUTexelFactorParameter;
 		return bShaderHasOutdatedParameters;
 	}
@@ -45,6 +45,8 @@ public:
 		FRHICommandList& RHICmdList, 
 		const FShader* OriginalVS, 
 		const FShader* OriginalPS, 
+		const FMaterialRenderProxy* MaterialRenderProxy,
+		const FMaterial& Material,
 		const FSceneView& View
 		) override;
 
@@ -59,8 +61,10 @@ public:
 
 	virtual void SetMesh(FRHICommandList& RHICmdList, const FSceneView& View) override;
 
+	virtual FShader* GetShader() override { return static_cast<FShader*>(this); }
+
 private:
 
-	FShaderParameter StreamingAccuracyColorsParameter;
+	FShaderParameter AccuracyColorsParameter;
 	FShaderParameter CPUTexelFactorParameter;
 };
diff --git a/Engine/Source/Runtime/Renderer/Private/TranslucentRendering.cpp b/Engine/Source/Runtime/Renderer/Private/TranslucentRendering.cpp
--- a/Engine/Source/Runtime/Renderer/Private/TranslucentRendering.cpp
+++ b/Engine/Source/Runtime/Renderer/Private/TranslucentRendering.cpp
@@ -20,7 +20,8 @@ static void SetTranslucentRenderTargetAndState(FRHICommandList& RHICmdList, cons
 
 	if ((TranslucenyPassType == TPT_SeparateTransluceny) && SceneContext.IsSeparateTranslucencyActive(View))
 	{
-		bSetupTranslucentState = SceneContext.BeginRenderingSeparateTranslucency(RHICmdList, View, bFirstTimeThisFrame);
+		const bool bNeedsClear = (&View == View.Family->Views[0]) && bFirstTimeThisFrame;
+		bSetupTranslucentState = SceneContext.BeginRenderingSeparateTranslucency(RHICmdList, View, bNeedsClear);
 	}
 	else
 	{
@@ -36,6 +37,8 @@ static void SetTranslucentRenderTargetAndState(FRHICommandList& RHICmdList, cons
 
 static void FinishTranslucentRenderTarget(FRHICommandListImmediate& RHICmdList, const FViewInfo& View, ETranslucencyPassType TranslucenyPassType)
 {
+	QUICK_SCOPE_CYCLE_COUNTER(STAT_FinishTranslucentRenderTarget);
+	
 	FSceneRenderTargets& SceneContext = FSceneRenderTargets::Get(RHICmdList);
 	if ((TranslucenyPassType == TPT_SeparateTransluceny) && SceneContext.IsSeparateTranslucencyActive(View))
 	{
@@ -49,6 +52,7 @@ static void FinishTranslucentRenderTarget(FRHICommandListImmediate& RHICmdList,
 
 const FProjectedShadowInfo* FDeferredShadingSceneRenderer::PrepareTranslucentShadowMap(FRHICommandList& RHICmdList, const FViewInfo& View, FPrimitiveSceneInfo* PrimitiveSceneInfo, ETranslucencyPassType TranslucenyPassType)
 {
+	QUICK_SCOPE_CYCLE_COUNTER(STAT_FDeferredShadingSceneRenderer_PrepareTranslucentShadowMap);
 	const FVisibleLightInfo* VisibleLightInfo = NULL;
 	FProjectedShadowInfo* TranslucentSelfShadow = NULL;
 
@@ -970,6 +974,7 @@ void FDeferredShadingSceneRenderer::RenderTranslucencyParallel(FRHICommandListIm
 		{
 			if (SceneContext.IsSeparateTranslucencyActive(View))
 			{
+				QUICK_SCOPE_CYCLE_COUNTER(RenderTranslucencyParallel_Downsample);
 				// we need to allocate this now so it ends up in the snapshot
 				FIntPoint ScaledSize;
 				uint32 NumSamples = 1;
@@ -988,6 +993,7 @@ void FDeferredShadingSceneRenderer::RenderTranslucencyParallel(FRHICommandListIm
 				TPT_NonSeparateTransluceny);
 
 			{
+				QUICK_SCOPE_CYCLE_COUNTER(RenderTranslucencyParallel_Start_FDrawSortedTransAnyThreadTask);
 				int32 NumPrims = View.TranslucentPrimSet.NumPrims() - View.TranslucentPrimSet.NumSeparateTranslucencyPrims();
 				int32 EffectiveThreads = FMath::Min<int32>(FMath::DivideAndRoundUp(NumPrims, ParallelCommandListSet.MinDrawsPerCommandList), ParallelCommandListSet.Width);
 
@@ -1017,10 +1023,15 @@ void FDeferredShadingSceneRenderer::RenderTranslucencyParallel(FRHICommandListIm
 				}
 			}
 			// Draw the view's mesh elements with the translucent drawing policy.
-			DrawViewElementsParallel<FTranslucencyDrawingPolicyFactory>(GParallelTranslucencyContext, SDPG_World, false, ParallelCommandListSet);
+			{
+				QUICK_SCOPE_CYCLE_COUNTER(RenderTranslucencyParallel_SDPG_World);
+				DrawViewElementsParallel<FTranslucencyDrawingPolicyFactory>(GParallelTranslucencyContext, SDPG_World, false, ParallelCommandListSet);
+			}
 			// Draw the view's mesh elements with the translucent drawing policy.
-			DrawViewElementsParallel<FTranslucencyDrawingPolicyFactory>(GParallelTranslucencyContext, SDPG_Foreground, false, ParallelCommandListSet);
-
+			{
+				QUICK_SCOPE_CYCLE_COUNTER(RenderTranslucencyParallel_SDPG_Foreground);
+				DrawViewElementsParallel<FTranslucencyDrawingPolicyFactory>(GParallelTranslucencyContext, SDPG_Foreground, false, ParallelCommandListSet);
+			}
 		}
 		FinishTranslucentRenderTarget(RHICmdList, View, TPT_NonSeparateTransluceny);
 
@@ -1048,6 +1059,7 @@ void FDeferredShadingSceneRenderer::RenderTranslucencyParallel(FRHICommandListIm
 				// Draw only translucent prims that are in the SeparateTranslucency pass
 				if (View.TranslucentPrimSet.NumSeparateTranslucencyPrims() > 0)
 				{
+					QUICK_SCOPE_CYCLE_COUNTER(RenderTranslucencyParallel_Start_FDrawSortedTransAnyThreadTask_SeparateTransluceny);
 					int32 NumPrims = View.TranslucentPrimSet.NumSeparateTranslucencyPrims();
 					int32 EffectiveThreads = FMath::Min<int32>(FMath::DivideAndRoundUp(NumPrims, ParallelCommandListSet.MinDrawsPerCommandList), ParallelCommandListSet.Width);
 
diff --git a/Engine/Source/Runtime/Renderer/Private/VelocityRendering.cpp b/Engine/Source/Runtime/Renderer/Private/VelocityRendering.cpp
--- a/Engine/Source/Runtime/Renderer/Private/VelocityRendering.cpp
+++ b/Engine/Source/Runtime/Renderer/Private/VelocityRendering.cpp
@@ -55,7 +55,10 @@ public:
 
 	bool SupportsVelocity() const
 	{
-		return PreviousLocalToWorld.IsBound() || GPUSkinCachePreviousBuffer.IsBound();
+		return PreviousLocalToWorld.IsBound() || 
+			GPUSkinCachePreviousBuffer.IsBound() || 
+			PrevTransformBuffer.IsBound() || 
+			(PrevTransform0.IsBound() && PrevTransform1.IsBound() && PrevTransform2.IsBound());
 	}
 
 	static bool ShouldCache(EShaderPlatform Platform,const FMaterial* Material,const FVertexFactoryType* VertexFactoryType)
@@ -76,6 +79,10 @@ protected:
 	{
 		PreviousLocalToWorld.Bind(Initializer.ParameterMap,TEXT("PreviousLocalToWorld"));
 		GPUSkinCachePreviousBuffer.Bind(Initializer.ParameterMap, TEXT("GPUSkinCachePreviousBuffer"));
+		PrevTransform0.Bind(Initializer.ParameterMap, TEXT("PrevTransform0"));
+		PrevTransform1.Bind(Initializer.ParameterMap, TEXT("PrevTransform1"));
+		PrevTransform2.Bind(Initializer.ParameterMap, TEXT("PrevTransform2"));
+		PrevTransformBuffer.Bind(Initializer.ParameterMap, TEXT("PrevTransformBuffer"));
 	}
 
 	FVelocityVS() {}
@@ -83,13 +90,24 @@ protected:
 	virtual bool Serialize(FArchive& Ar) override
 	{
 		bool bShaderHasOutdatedParameters = FMeshMaterialShader::Serialize(Ar);
-		Ar << PreviousLocalToWorld << GPUSkinCachePreviousBuffer;
+
+		Ar << PreviousLocalToWorld;
+		Ar << GPUSkinCachePreviousBuffer;
+		Ar << PrevTransform0;
+		Ar << PrevTransform1;
+		Ar << PrevTransform2;
+		Ar << PrevTransformBuffer;
+
 		return bShaderHasOutdatedParameters;
 	}
 
 private:
 	FShaderParameter PreviousLocalToWorld;
 	FShaderResourceParameter GPUSkinCachePreviousBuffer;
+	FShaderParameter PrevTransform0;
+	FShaderParameter PrevTransform1;
+	FShaderParameter PrevTransform2;
+	FShaderResourceParameter PrevTransformBuffer;
 };
 
 
diff --git a/Engine/Source/Runtime/Renderer/Private/WantedMipsAccuracyRendering.cpp b/Engine/Source/Runtime/Renderer/Private/WantedMipsAccuracyRendering.cpp
--- a/Engine/Source/Runtime/Renderer/Private/WantedMipsAccuracyRendering.cpp
+++ b/Engine/Source/Runtime/Renderer/Private/WantedMipsAccuracyRendering.cpp
@@ -1,4 +1,4 @@
-// Copyright 1998-2015 Epic Games, Inc. All Rights Reserved.
+// Copyright 1998-2016 Epic Games, Inc. All Rights Reserved.
 
 /*=============================================================================
 WantedMipsAccuracyRendering.cpp: Contains definitions for rendering the viewmode.
@@ -6,6 +6,7 @@ WantedMipsAccuracyRendering.cpp: Contains definitions for rendering the viewmode
 
 #include "RendererPrivate.h"
 #include "ScenePrivate.h"
+#include "WantedMipsAccuracyRendering.h"
 
 extern ENGINE_API TAutoConsoleVariable<int32> CVarStreamingUseAABB;
 
@@ -15,6 +16,8 @@ void FWantedMipsAccuracyPS::SetParameters(
 	FRHICommandList& RHICmdList, 
 	const FShader* OriginalVS, 
 	const FShader* OriginalPS, 
+	const FMaterialRenderProxy* MaterialRenderProxy,
+	const FMaterial& Material,
 	const FSceneView& View
 	)
 {
@@ -22,11 +25,11 @@ void FWantedMipsAccuracyPS::SetParameters(
 	int32 ColorIndex = 0;
 	for (; ColorIndex < NumEngineColors; ++ColorIndex)
 	{
-		SetShaderValue(RHICmdList, FGlobalShader::GetPixelShader(), StreamingAccuracyColorsParameter, GEngine->StreamingAccuracyColors[ColorIndex], ColorIndex);
+		SetShaderValue(RHICmdList, FGlobalShader::GetPixelShader(), AccuracyColorsParameter, GEngine->StreamingAccuracyColors[ColorIndex], ColorIndex);
 	}
 	for (; ColorIndex < NumStreamingAccuracyColors; ++ColorIndex)
 	{
-		SetShaderValue(RHICmdList, FGlobalShader::GetPixelShader(), StreamingAccuracyColorsParameter, FLinearColor::Black, ColorIndex);
+		SetShaderValue(RHICmdList, FGlobalShader::GetPixelShader(), AccuracyColorsParameter, FLinearColor::Black, ColorIndex);
 	}
 }
 
diff --git a/Engine/Source/Runtime/Renderer/Private/WantedMipsAccuracyRendering.h b/Engine/Source/Runtime/Renderer/Private/WantedMipsAccuracyRendering.h
--- a/Engine/Source/Runtime/Renderer/Private/WantedMipsAccuracyRendering.h
+++ b/Engine/Source/Runtime/Renderer/Private/WantedMipsAccuracyRendering.h
@@ -1,4 +1,4 @@
-// Copyright 1998-2015 Epic Games, Inc. All Rights Reserved.
+// Copyright 1998-2016 Epic Games, Inc. All Rights Reserved.
 
 /*=============================================================================
 WantedMipsAccuracyRendering.h: Declarations used for the viewmode.
@@ -6,13 +6,10 @@ WantedMipsAccuracyRendering.h: Declarations used for the viewmode.
 
 #pragma once
 
-static const int32 NumStreamingAccuracyColors = 5;
-static const int32 MaxStreamingAccuracyMips = 11;
-
 /**
 * Pixel shader that renders texture streamer wanted mips accuracy.
 */
-class FWantedMipsAccuracyPS : public FDebugViewModePS
+class FWantedMipsAccuracyPS : public FGlobalShader, public IDebugViewModePSInterface
 {
 	DECLARE_SHADER_TYPE(FWantedMipsAccuracyPS,Global);
 
@@ -24,9 +21,9 @@ public:
 	}
 
 	FWantedMipsAccuracyPS(const ShaderMetaType::CompiledShaderInitializerType& Initializer):
-		FDebugViewModePS(Initializer)
+		FGlobalShader(Initializer)
 	{
-		StreamingAccuracyColorsParameter.Bind(Initializer.ParameterMap,TEXT("StreamingAccuracyColors"));
+		AccuracyColorsParameter.Bind(Initializer.ParameterMap,TEXT("AccuracyColors"));
 		CPUWantedMipsParameter.Bind(Initializer.ParameterMap,TEXT("CPUWantedMips"));
 	}
 
@@ -35,7 +32,7 @@ public:
 	virtual bool Serialize(FArchive& Ar) override
 	{
 		bool bShaderHasOutdatedParameters = FGlobalShader::Serialize(Ar);
-		Ar << StreamingAccuracyColorsParameter;
+		Ar << AccuracyColorsParameter;
 		Ar << CPUWantedMipsParameter;
 		return bShaderHasOutdatedParameters;
 	}
@@ -49,6 +46,8 @@ public:
 		FRHICommandList& RHICmdList, 
 		const FShader* OriginalVS, 
 		const FShader* OriginalPS, 
+		const FMaterialRenderProxy* MaterialRenderProxy,
+		const FMaterial& Material,
 		const FSceneView& View
 		) override;
 
@@ -63,8 +62,10 @@ public:
 
 	virtual void SetMesh(FRHICommandList& RHICmdList, const FSceneView& View) override;
 
+	virtual FShader* GetShader() override { return static_cast<FShader*>(this); }
+
 private:
 
-	FShaderParameter StreamingAccuracyColorsParameter;
+	FShaderParameter AccuracyColorsParameter;
 	FShaderParameter CPUWantedMipsParameter;
 };
diff --git a/Engine/Source/Runtime/Renderer/Public/AtmosphereTextureParameters.h b/Engine/Source/Runtime/Renderer/Public/AtmosphereTextureParameters.h
--- a/Engine/Source/Runtime/Renderer/Public/AtmosphereTextureParameters.h
+++ b/Engine/Source/Runtime/Renderer/Public/AtmosphereTextureParameters.h
@@ -6,6 +6,8 @@
 
 #pragma once
 
+#include "RHIStaticStates.h"
+
 /** Shader parameters needed for atmosphere passes. */
 class FAtmosphereShaderTextureParameters
 {
@@ -13,7 +15,18 @@ public:
 	void Bind(const FShaderParameterMap& ParameterMap);
 
 	template< typename ShaderRHIParamRef >
-	void Set( FRHICommandList& RHICmdList, const ShaderRHIParamRef ShaderRHI, const FSceneView& View ) const;
+	FORCEINLINE_DEBUGGABLE void Set( FRHICommandList& RHICmdList, const ShaderRHIParamRef ShaderRHI, const FSceneView& View ) const
+	{
+		if (TransmittanceTexture.IsBound() || IrradianceTexture.IsBound() || InscatterTexture.IsBound())
+		{
+			SetTextureParameter(RHICmdList, ShaderRHI, TransmittanceTexture, TransmittanceTextureSampler,
+				TStaticSamplerState<SF_Bilinear>::GetRHI(), View.AtmosphereTransmittanceTexture);
+			SetTextureParameter(RHICmdList, ShaderRHI, IrradianceTexture, IrradianceTextureSampler,
+				TStaticSamplerState<SF_Bilinear>::GetRHI(), View.AtmosphereIrradianceTexture);
+			SetTextureParameter(RHICmdList, ShaderRHI, InscatterTexture, InscatterTextureSampler,
+				TStaticSamplerState<SF_Bilinear>::GetRHI(), View.AtmosphereInscatterTexture);
+		}
+	}
 
 	friend FArchive& operator<<(FArchive& Ar,FAtmosphereShaderTextureParameters& P);
 
diff --git a/Engine/Source/Runtime/Renderer/Public/GlobalDistanceFieldParameters.h b/Engine/Source/Runtime/Renderer/Public/GlobalDistanceFieldParameters.h
--- a/Engine/Source/Runtime/Renderer/Public/GlobalDistanceFieldParameters.h
+++ b/Engine/Source/Runtime/Renderer/Public/GlobalDistanceFieldParameters.h
@@ -64,7 +64,7 @@ public:
 	}
 
 	template<typename ShaderRHIParamRef>
-	void Set(FRHICommandList& RHICmdList, const ShaderRHIParamRef ShaderRHI, const FGlobalDistanceFieldParameterData& ParameterData)
+	FORCEINLINE_DEBUGGABLE void Set(FRHICommandList& RHICmdList, const ShaderRHIParamRef ShaderRHI, const FGlobalDistanceFieldParameterData& ParameterData)
 	{
 		if (GlobalVolumeCenterAndExtent.IsBound() || GlobalVolumeWorldToUVAddAndMul.IsBound())
 		{
diff --git a/Engine/Source/Runtime/Renderer/Public/MaterialShader.h b/Engine/Source/Runtime/Renderer/Public/MaterialShader.h
--- a/Engine/Source/Runtime/Renderer/Public/MaterialShader.h
+++ b/Engine/Source/Runtime/Renderer/Public/MaterialShader.h
@@ -110,7 +110,7 @@ public:
 	FUniformBufferRHIParamRef GetParameterCollectionBuffer(const FGuid& Id, const FSceneInterface* SceneInterface) const;
 
 	template<typename ShaderRHIParamRef>
-	void SetParameters(FRHICommandList& RHICmdList, const ShaderRHIParamRef ShaderRHI, const FSceneView& View)
+	FORCEINLINE_DEBUGGABLE void SetParameters(FRHICommandList& RHICmdList, const ShaderRHIParamRef ShaderRHI, const FSceneView& View)
 	{
 		const auto& ViewUniformBufferParameter = GetUniformBufferParameter<FViewUniformShaderParameters>();
 		const auto& FrameUniformBufferParameter = GetUniformBufferParameter<FFrameUniformShaderParameters>();
@@ -118,7 +118,9 @@ public:
 		CheckShaderIsValid();
 		SetUniformBufferParameter(RHICmdList, ShaderRHI, ViewUniformBufferParameter, View.ViewUniformBuffer);
 		SetUniformBufferParameter(RHICmdList, ShaderRHI, FrameUniformBufferParameter, View.FrameUniformBuffer);
+#if USE_GBuiltinSamplersUniformBuffer
 		SetUniformBufferParameter(RHICmdList, ShaderRHI, BuiltinSamplersUBParameter, GBuiltinSamplersUniformBuffer.GetUniformBufferRHI());
+#endif
 
 		// Skip if instanced stereo is not enabled
 		if (View.bIsInstancedStereoEnabled && View.Family->Views.Num() > 0)
@@ -167,21 +169,9 @@ private:
 	FShaderResourceParameter LightAttenuation;
 	FShaderResourceParameter LightAttenuationSampler;
 
-	// For materials using atmospheric fog color 
-	FAtmosphereShaderTextureParameters AtmosphericFogTextureParameters;
-
 	//Use of the eye adaptation texture here is experimental and potentially dangerous as it can introduce a feedback loop. May be removed.
 	FShaderResourceParameter EyeAdaptation;
 
-	/** The PerlinNoiseGradientTexture parameter for materials that use GradientNoise */
-	FShaderResourceParameter PerlinNoiseGradientTexture;
-	FShaderResourceParameter PerlinNoiseGradientTextureSampler;
-	/** The PerlinNoise3DTexture parameter for materials that use GradientNoise */
-	FShaderResourceParameter PerlinNoise3DTexture;
-	FShaderResourceParameter PerlinNoise3DTextureSampler;
-
-	FGlobalDistanceFieldParameters GlobalDistanceFieldParameters;
-
 	FDebugUniformExpressionSet	DebugUniformExpressionSet;
 	FRHIUniformBufferLayout		DebugUniformExpressionUBLayout;
 	FString						DebugDescription;
diff --git a/Engine/Source/Runtime/ShaderCore/Public/Shader.h b/Engine/Source/Runtime/ShaderCore/Public/Shader.h
--- a/Engine/Source/Runtime/ShaderCore/Public/Shader.h
+++ b/Engine/Source/Runtime/ShaderCore/Public/Shader.h
@@ -623,7 +623,7 @@ public:
 
 	/** Finds an automatically bound uniform buffer matching the given uniform buffer type if one exists, or returns an unbound parameter. */
 	template<typename UniformBufferStructType>
-	const TShaderUniformBufferParameter<UniformBufferStructType>& GetUniformBufferParameter() const
+	FORCEINLINE_DEBUGGABLE const TShaderUniformBufferParameter<UniformBufferStructType>& GetUniformBufferParameter() const
 	{
 		FUniformBufferStruct* SearchStruct = &UniformBufferStructType::StaticStruct;
 		int32 FoundIndex = INDEX_NONE;
diff --git a/Engine/Source/Runtime/ShaderCore/Public/VertexFactory.h b/Engine/Source/Runtime/ShaderCore/Public/VertexFactory.h
--- a/Engine/Source/Runtime/ShaderCore/Public/VertexFactory.h
+++ b/Engine/Source/Runtime/ShaderCore/Public/VertexFactory.h
@@ -335,22 +335,15 @@ public:
 class SHADERCORE_API FVertexFactory : public FRenderResource
 {
 public:
-	/**
-	 * Default constructor 
-	 */
-	FVertexFactory() {}
+	FVertexFactory() 
+	{
+	}
 
-	/**
-	* Constructor specifiying feature level
-	*/
 	FVertexFactory(ERHIFeatureLevel::Type InFeatureLevel) 
 		: FRenderResource(InFeatureLevel)
 	{
 	}
 
-	/**
-	 * @return The vertex factory's type.
-	 */
 	virtual FVertexFactoryType* GetType() const { return NULL; }
 
 	/**
diff --git a/Engine/Source/Runtime/SlateRHIRenderer/Private/SlateRHIRenderingPolicy.cpp b/Engine/Source/Runtime/SlateRHIRenderer/Private/SlateRHIRenderingPolicy.cpp
--- a/Engine/Source/Runtime/SlateRHIRenderer/Private/SlateRHIRenderingPolicy.cpp
+++ b/Engine/Source/Runtime/SlateRHIRenderer/Private/SlateRHIRenderingPolicy.cpp
@@ -217,6 +217,8 @@ static FSceneView& CreateSceneView( FSceneViewFamilyContext& ViewFamilyContext,
 	FViewUniformShaderParameters ViewUniformShaderParameters;
 	FFrameUniformShaderParameters FrameUniformShaderParameters;
 
+	FMemory::Memzero(ViewUniformShaderParameters);
+
 	ViewUniformShaderParameters.TranslatedWorldToClip = View->ViewMatrices.TranslatedViewProjectionMatrix;
 	ViewUniformShaderParameters.WorldToClip = ViewProjectionMatrix;
 	ViewUniformShaderParameters.TranslatedWorldToView = EffectiveTranslatedViewMatrix;
@@ -248,8 +250,6 @@ static FSceneView& CreateSceneView( FSceneViewFamilyContext& ViewFamilyContext,
 	FrameUniformShaderParameters.RealTime = View->Family->CurrentRealTime;
 	FrameUniformShaderParameters.Random = FMath::Rand();
 	FrameUniformShaderParameters.FrameNumber = View->Family->FrameNumber;
-	FrameUniformShaderParameters.DirectionalLightShadowTexture = GWhiteTexture->TextureRHI;
-	FrameUniformShaderParameters.DirectionalLightShadowSampler = TStaticSamplerState<SF_Point, AM_Clamp, AM_Clamp, AM_Clamp>::GetRHI();
 
 	{
 		// setup a matrix to transform float4(SvPosition.xyz,1) directly to TranslatedWorld (quality, performance as we don't need to convert or use interpolator)
diff --git a/Engine/Source/Runtime/Windows/D3D11RHI/Private/D3D11Commands.cpp b/Engine/Source/Runtime/Windows/D3D11RHI/Private/D3D11Commands.cpp
--- a/Engine/Source/Runtime/Windows/D3D11RHI/Private/D3D11Commands.cpp
+++ b/Engine/Source/Runtime/Windows/D3D11RHI/Private/D3D11Commands.cpp
@@ -1269,7 +1269,7 @@ void FD3D11DynamicRHI::CommitComputeShaderConstants()
 }
 
 template <EShaderFrequency Frequency>
-FORCEINLINE void SetResource(FD3D11DynamicRHI* RESTRICT D3D11RHI, FD3D11StateCache* RESTRICT StateCache, uint32 BindIndex, FD3D11BaseShaderResource* RESTRICT ShaderResource, ID3D11ShaderResourceView* RESTRICT SRV, FName ResourceName)
+FORCEINLINE void SetResource(FD3D11DynamicRHI* RESTRICT D3D11RHI, FD3D11StateCache* RESTRICT StateCache, uint32 BindIndex, FD3D11BaseShaderResource* RESTRICT ShaderResource, ID3D11ShaderResourceView* RESTRICT SRV, FName ResourceName = FName())
 {
 	// We set the resource through the RHI to track state for the purposes of unbinding SRVs when a UAV or RTV is bound.
 	// todo: need to support SRV_Static for faster calls when possible
@@ -1277,31 +1277,100 @@ FORCEINLINE void SetResource(FD3D11DynamicRHI* RESTRICT D3D11RHI, FD3D11StateCac
 }
 
 template <EShaderFrequency Frequency>
-FORCEINLINE void SetResource(FD3D11DynamicRHI* RESTRICT D3D11RHI, FD3D11StateCache* RESTRICT StateCache, uint32 BindIndex, FD3D11BaseShaderResource* RESTRICT ShaderResource, ID3D11SamplerState* RESTRICT SamplerState, FName ResourceName)
+FORCEINLINE void SetResource(FD3D11DynamicRHI* RESTRICT D3D11RHI, FD3D11StateCache* RESTRICT StateCache, uint32 BindIndex, ID3D11SamplerState* RESTRICT SamplerState)
 {
 	StateCache->SetSamplerState<Frequency>(SamplerState,BindIndex);
 }
 
-template <class D3DResourceType, EShaderFrequency ShaderFrequency>
-inline int32 SetShaderResourcesFromBuffer(FD3D11DynamicRHI* RESTRICT D3D11RHI, FD3D11StateCache* RESTRICT StateCache, FD3D11UniformBuffer* RESTRICT Buffer, const uint32* RESTRICT ResourceMap, int32 BufferIndex)
+template <EShaderFrequency ShaderFrequency>
+inline int32 SetShaderResourcesFromBuffer_Surface(FD3D11DynamicRHI* RESTRICT D3D11RHI, FD3D11StateCache* RESTRICT StateCache, FD3D11UniformBuffer* RESTRICT Buffer, const uint32* RESTRICT ResourceMap, int32 BufferIndex)
 {
+	const TRefCountPtr<FRHIResource>* RESTRICT Resources = Buffer->ResourceTable.GetData();
+	float CurrentTime = FApp::GetCurrentTime();
 	int32 NumSetCalls = 0;
 	uint32 BufferOffset = ResourceMap[BufferIndex];
 	if (BufferOffset > 0)
 	{
 		const uint32* RESTRICT ResourceInfos = &ResourceMap[BufferOffset];
 		uint32 ResourceInfo = *ResourceInfos++;
-		do 
+		do
 		{
 			checkSlow(FRHIResourceTableEntry::GetUniformBufferIndex(ResourceInfo) == BufferIndex);
 			const uint16 ResourceIndex = FRHIResourceTableEntry::GetResourceIndex(ResourceInfo);
 			const uint8 BindIndex = FRHIResourceTableEntry::GetBindIndex(ResourceInfo);
 
+			FD3D11BaseShaderResource* ShaderResource = nullptr;
+			ID3D11ShaderResourceView* D3D11Resource = nullptr;
+
+			FRHITexture* TextureRHI = (FRHITexture*)Resources[ResourceIndex].GetReference();
+			TextureRHI->SetLastRenderTime(CurrentTime);
+			FD3D11TextureBase* TextureD3D11 = GetD3D11TextureFromRHITexture(TextureRHI);
+			ShaderResource = TextureD3D11->GetBaseShaderResource();
+			D3D11Resource = TextureD3D11->GetShaderResourceView();
+
+			// todo: could coalesce adjacent bound resources.
+			SetResource<ShaderFrequency>(D3D11RHI, StateCache, BindIndex, ShaderResource, D3D11Resource, TextureRHI->GetName());
+			NumSetCalls++;
+			ResourceInfo = *ResourceInfos++;
+		} while (FRHIResourceTableEntry::GetUniformBufferIndex(ResourceInfo) == BufferIndex);
+	}
+	return NumSetCalls;
+}
+
+
+template <EShaderFrequency ShaderFrequency>
+inline int32 SetShaderResourcesFromBuffer_SRV(FD3D11DynamicRHI* RESTRICT D3D11RHI, FD3D11StateCache* RESTRICT StateCache, FD3D11UniformBuffer* RESTRICT Buffer, const uint32* RESTRICT ResourceMap, int32 BufferIndex)
+{
+	const TRefCountPtr<FRHIResource>* RESTRICT Resources = Buffer->ResourceTable.GetData();
+	float CurrentTime = FApp::GetCurrentTime();
+	int32 NumSetCalls = 0;
+	uint32 BufferOffset = ResourceMap[BufferIndex];
+	if (BufferOffset > 0)
+	{
+		const uint32* RESTRICT ResourceInfos = &ResourceMap[BufferOffset];
+		uint32 ResourceInfo = *ResourceInfos++;
+		do
+		{
+			checkSlow(FRHIResourceTableEntry::GetUniformBufferIndex(ResourceInfo) == BufferIndex);
+			const uint16 ResourceIndex = FRHIResourceTableEntry::GetResourceIndex(ResourceInfo);
+			const uint8 BindIndex = FRHIResourceTableEntry::GetBindIndex(ResourceInfo);
+
+			FD3D11BaseShaderResource* ShaderResource = nullptr;
+			ID3D11ShaderResourceView* D3D11Resource = nullptr;
+
+			FD3D11ShaderResourceView* ShaderResourceViewRHI = (FD3D11ShaderResourceView*)Resources[ResourceIndex].GetReference();
+			ShaderResource = ShaderResourceViewRHI->Resource.GetReference();
+			D3D11Resource = ShaderResourceViewRHI->View.GetReference();
+
+			// todo: could coalesce adjacent bound resources.
+			SetResource<ShaderFrequency>(D3D11RHI, StateCache, BindIndex, ShaderResource, D3D11Resource);
+			NumSetCalls++;
+			ResourceInfo = *ResourceInfos++;
+		} while (FRHIResourceTableEntry::GetUniformBufferIndex(ResourceInfo) == BufferIndex);
+	}
+	return NumSetCalls;
+}
+
+template <EShaderFrequency ShaderFrequency>
+inline int32 SetShaderResourcesFromBuffer_Sampler(FD3D11DynamicRHI* RESTRICT D3D11RHI, FD3D11StateCache* RESTRICT StateCache, FD3D11UniformBuffer* RESTRICT Buffer, const uint32* RESTRICT ResourceMap, int32 BufferIndex)
+{
+	const TRefCountPtr<FRHIResource>* RESTRICT Resources = Buffer->ResourceTable.GetData();
+	int32 NumSetCalls = 0;
+	uint32 BufferOffset = ResourceMap[BufferIndex];
+	if (BufferOffset > 0)
+	{
+		const uint32* RESTRICT ResourceInfos = &ResourceMap[BufferOffset];
+		uint32 ResourceInfo = *ResourceInfos++;
+		do
+		{
+			checkSlow(FRHIResourceTableEntry::GetUniformBufferIndex(ResourceInfo) == BufferIndex);
+			const uint16 ResourceIndex = FRHIResourceTableEntry::GetResourceIndex(ResourceInfo);
+			const uint8 BindIndex = FRHIResourceTableEntry::GetBindIndex(ResourceInfo);
+
+			ID3D11SamplerState* D3D11Resource = ((FD3D11SamplerState*)Resources[ResourceIndex].GetReference())->Resource.GetReference();
+
 			// todo: could coalesce adjacent bound resources.
-			FD3D11UniformBuffer::FResourcePair* RESTRICT ResourcePair = &Buffer->RawResourceTable[ResourceIndex];
-			FD3D11BaseShaderResource* ShaderResource = ResourcePair->ShaderResource;
-			D3DResourceType* D3D11Resource = (D3DResourceType*)ResourcePair->D3D11Resource;
-			SetResource<ShaderFrequency>(D3D11RHI, StateCache, BindIndex, ShaderResource, D3D11Resource, ResourcePair->GetResourceName());
+			SetResource<ShaderFrequency>(D3D11RHI, StateCache, BindIndex, D3D11Resource);
 			NumSetCalls++;
 			ResourceInfo = *ResourceInfos++;
 		} while (FRHIResourceTableEntry::GetUniformBufferIndex(ResourceInfo) == BufferIndex);
@@ -1309,6 +1378,7 @@ inline int32 SetShaderResourcesFromBuffer(FD3D11DynamicRHI* RESTRICT D3D11RHI, F
 	return NumSetCalls;
 }
 
+
 template <class ShaderType>
 void FD3D11DynamicRHI::SetResourcesFromTables(const ShaderType* RESTRICT Shader)
 {
@@ -1354,11 +1424,10 @@ void FD3D11DynamicRHI::SetResourcesFromTables(const ShaderType* RESTRICT Shader)
 		}
 #endif
 
-		Buffer->CacheResources(ResourceTableFrameCounter);
-
 		// todo: could make this two pass: gather then set
-		SetShaderResourcesFromBuffer<ID3D11ShaderResourceView, (EShaderFrequency)ShaderType::StaticFrequency>(this, &StateCache, Buffer, Shader->ShaderResourceTable.ShaderResourceViewMap.GetData(), BufferIndex);
-		SetShaderResourcesFromBuffer<ID3D11SamplerState, (EShaderFrequency)ShaderType::StaticFrequency>(this, &StateCache, Buffer, Shader->ShaderResourceTable.SamplerMap.GetData(), BufferIndex);
+		SetShaderResourcesFromBuffer_Surface<(EShaderFrequency)ShaderType::StaticFrequency>(this, &StateCache, Buffer, Shader->ShaderResourceTable.TextureMap.GetData(), BufferIndex);
+		SetShaderResourcesFromBuffer_SRV<(EShaderFrequency)ShaderType::StaticFrequency>(this, &StateCache, Buffer, Shader->ShaderResourceTable.ShaderResourceViewMap.GetData(), BufferIndex);
+		SetShaderResourcesFromBuffer_Sampler<(EShaderFrequency)ShaderType::StaticFrequency>(this, &StateCache, Buffer, Shader->ShaderResourceTable.SamplerMap.GetData(), BufferIndex);
 	}
 	DirtyUniformBuffers[ShaderType::StaticFrequency] = 0;
 }
diff --git a/Engine/Source/Runtime/Windows/D3D11RHI/Private/D3D11Device.cpp b/Engine/Source/Runtime/Windows/D3D11RHI/Private/D3D11Device.cpp
--- a/Engine/Source/Runtime/Windows/D3D11RHI/Private/D3D11Device.cpp
+++ b/Engine/Source/Runtime/Windows/D3D11RHI/Private/D3D11Device.cpp
@@ -186,6 +186,7 @@ FD3D11DynamicRHI::FD3D11DynamicRHI(IDXGIFactory1* InDXGIFactory1,D3D_FEATURE_LEV
 	GMaxTextureMipCount = FMath::Min<int32>( MAX_TEXTURE_MIP_COUNT, GMaxTextureMipCount );
 	GMaxShadowDepthBufferSizeX = 4096;
 	GMaxShadowDepthBufferSizeY = 4096;
+	GRHISupportsInstancing = GMaxRHIFeatureLevel > ERHIFeatureLevel::ES2;
 
 	// Initialize the constant buffers.
 	InitConstantBuffers();
diff --git a/Engine/Source/Runtime/Windows/D3D11RHI/Private/D3D11UAV.cpp b/Engine/Source/Runtime/Windows/D3D11RHI/Private/D3D11UAV.cpp
--- a/Engine/Source/Runtime/Windows/D3D11RHI/Private/D3D11UAV.cpp
+++ b/Engine/Source/Runtime/Windows/D3D11RHI/Private/D3D11UAV.cpp
@@ -197,6 +197,50 @@ FShaderResourceViewRHIRef FD3D11DynamicRHI::RHICreateShaderResourceView(FVertexB
 	return new FD3D11ShaderResourceView(ShaderResourceView,VertexBuffer);
 }
 
+FShaderResourceViewRHIRef FD3D11DynamicRHI::RHICreateShaderResourceView(FIndexBufferRHIParamRef BufferRHI)
+{
+	FD3D11IndexBuffer* Buffer = ResourceCast(BufferRHI);
+	check(Buffer);
+	check(Buffer->Resource);
+
+	// The stride in bytes of the index buffer; must be 2 or 4
+	uint32 Stride = BufferRHI->GetStride();
+
+	check(Stride == 2 || Stride == 4);
+
+	D3D11_BUFFER_DESC BufferDesc;
+	Buffer->Resource->GetDesc(&BufferDesc);
+
+	// Create a Shader Resource View
+	D3D11_SHADER_RESOURCE_VIEW_DESC SRVDesc;
+	FMemory::Memzero(SRVDesc);
+	SRVDesc.ViewDimension = D3D11_SRV_DIMENSION_BUFFER;
+	SRVDesc.Buffer.FirstElement = 0;
+
+	EPixelFormat Format = (Stride == 2) ? PF_R16_UINT : PF_R32_UINT;
+
+	SRVDesc.Format = FindShaderResourceDXGIFormat((DXGI_FORMAT)GPixelFormats[Format].PlatformFormat,false);
+	SRVDesc.Buffer.NumElements = BufferDesc.ByteWidth / Stride;
+	TRefCountPtr<ID3D11ShaderResourceView> ShaderResourceView;
+	
+	HRESULT hr = Direct3DDevice->CreateShaderResourceView(Buffer->Resource, &SRVDesc, (ID3D11ShaderResourceView**)ShaderResourceView.GetInitReference());
+	if (FAILED(hr))
+	{
+		if (hr == E_OUTOFMEMORY)
+		{
+			// There appears to be a driver bug that causes SRV creation to fail with an OOM error and then succeed on the next call.
+			hr = Direct3DDevice->CreateShaderResourceView(Buffer->Resource, &SRVDesc, (ID3D11ShaderResourceView**)ShaderResourceView.GetInitReference());
+		}
+		if (FAILED(hr))
+		{
+			UE_LOG(LogD3D11RHI,Error,TEXT("Failed to create shader resource view for index buffer: ByteWidth=%d NumElements=%d Format=%s"),BufferDesc.ByteWidth,BufferDesc.ByteWidth / Stride, GPixelFormats[Format].Name);
+			VerifyD3D11Result(hr,"Direct3DDevice->CreateShaderResourceView",__FILE__,__LINE__,Direct3DDevice);
+		}
+	}
+
+	return new FD3D11ShaderResourceView(ShaderResourceView, Buffer);
+}
+
 void FD3D11DynamicRHI::RHIClearUAV(FUnorderedAccessViewRHIParamRef UnorderedAccessViewRHI, const uint32* Values)
 {
 	FD3D11UnorderedAccessView* UnorderedAccessView = ResourceCast(UnorderedAccessViewRHI);
diff --git a/Engine/Source/Runtime/Windows/D3D11RHI/Private/D3D11UniformBuffer.cpp b/Engine/Source/Runtime/Windows/D3D11RHI/Private/D3D11UniformBuffer.cpp
--- a/Engine/Source/Runtime/Windows/D3D11RHI/Private/D3D11UniformBuffer.cpp
+++ b/Engine/Source/Runtime/Windows/D3D11RHI/Private/D3D11UniformBuffer.cpp
@@ -211,8 +211,6 @@ FUniformBufferRHIRef FD3D11DynamicRHI::RHICreateUniformBuffer(const void* Conten
 			check(InResources[i]);
 			NewUniformBuffer->ResourceTable[i] = InResources[i];
 		}
-		NewUniformBuffer->RawResourceTable.Empty(NumResources);
-		NewUniformBuffer->RawResourceTable.AddZeroed(NumResources);
 	}
 
 	return NewUniformBuffer;
@@ -251,60 +249,6 @@ FD3D11UniformBuffer::~FD3D11UniformBuffer()
 	}
 }
 
-void FD3D11UniformBuffer::CacheResourcesInternal()
-{
-	const FRHIUniformBufferLayout& Layout = GetLayout();
-	int32 NumResources = Layout.Resources.Num();
-	const uint8* RESTRICT ResourceTypes = Layout.Resources.GetData();
-	const TRefCountPtr<FRHIResource>* RESTRICT Resources = ResourceTable.GetData();
-	FResourcePair* RESTRICT RawResources = RawResourceTable.GetData();
-	float CurrentTime = FApp::GetCurrentTime();
-
-	// todo: Immutable resources, i.e. not textures, can be safely cached across frames.
-	// Texture streaming makes textures complicated :)
-	for (int32 i = 0; i < NumResources; ++i)
-	{
-		RawResources[i].SetResourceName(NAME_None);
-		switch (ResourceTypes[i])
-		{
-		case UBMT_SRV:
-			{
-				FD3D11ShaderResourceView* ShaderResourceViewRHI = (FD3D11ShaderResourceView*)Resources[i].GetReference();
-				RawResources[i].ShaderResource = ShaderResourceViewRHI->Resource.GetReference();
-				RawResources[i].D3D11Resource = (IUnknown*)ShaderResourceViewRHI->View.GetReference();				
-			}
-			break;
-
-		case UBMT_TEXTURE:
-			{
-				// todo: this does multiple virtual function calls to find the right type to cast to
-				// this is due to multiple inheritance nastiness, NEEDS CLEANUP
-				FRHITexture* TextureRHI = (FRHITexture*)Resources[i].GetReference();
-				TextureRHI->SetLastRenderTime(CurrentTime);
-				FD3D11TextureBase* TextureD3D11 = GetD3D11TextureFromRHITexture(TextureRHI);
-				RawResources[i].ShaderResource = TextureD3D11->GetBaseShaderResource();
-				RawResources[i].D3D11Resource = (IUnknown*)TextureD3D11->GetShaderResourceView();
-				RawResources[i].SetResourceName(TextureRHI->GetName());
-			}
-			break;
-
-		case UBMT_UAV:
-			RawResources[i].ShaderResource = NULL;
-			RawResources[i].D3D11Resource = NULL;
-			check(0);
-			break;
-
-		case UBMT_SAMPLER:
-			RawResources[i].ShaderResource = NULL;
-			RawResources[i].D3D11Resource = (IUnknown*)((FD3D11SamplerState*)Resources[i].GetReference())->Resource.GetReference();
-			break;
-
-		default:
-			check(0);
-			break;
-		}
-	}
-}
 
 void FD3D11DynamicRHI::ReleasePooledUniformBuffers()
 {
diff --git a/Engine/Source/Runtime/Windows/D3D11RHI/Public/D3D11Resources.h b/Engine/Source/Runtime/Windows/D3D11RHI/Public/D3D11Resources.h
--- a/Engine/Source/Runtime/Windows/D3D11RHI/Public/D3D11Resources.h
+++ b/Engine/Source/Runtime/Windows/D3D11RHI/Public/D3D11Resources.h
@@ -390,14 +390,19 @@ public:
 	bool IsCubemap() const { return bCubemap; }
 
 	/** FRHITexture override.  See FRHITexture::GetNativeResource() */
-	virtual void* GetNativeResource() const override
+	virtual void* GetNativeResource() const override final
 	{ 
 		return GetResource();
 	}
-	virtual void* GetNativeShaderResourceView() const override
+	virtual void* GetNativeShaderResourceView() const override final
 	{
 		return GetShaderResourceView();
 	}
+	virtual void* GetTextureBaseRHI() override final
+	{
+		return static_cast<FD3D11TextureBase*>(this);
+	}
+
 
 	// IRefCountedObject interface.
 	virtual uint32 AddRef() const
@@ -467,6 +472,11 @@ public:
 	// Accessors.
 	ID3D11Texture3D* GetResource() const { return (ID3D11Texture3D*)FD3D11TextureBase::GetResource(); }
 	
+	virtual void* GetTextureBaseRHI() override final
+	{
+		return static_cast<FD3D11TextureBase*>(this);
+	}
+
 	// IRefCountedObject interface.
 	virtual uint32 AddRef() const
 	{
@@ -533,6 +543,10 @@ public:
 		FRHITextureReference::SetReferencedTexture(InTexture);
 	}
 
+	virtual void* GetTextureBaseRHI() override final
+	{
+		return static_cast<FD3D11TextureBase*>(this);
+	}
 	// IRefCountedObject interface.
 	virtual uint32 AddRef() const
 	{
@@ -549,37 +563,15 @@ public:
 };
 
 /** Given a pointer to a RHI texture that was created by the D3D11 RHI, returns a pointer to the FD3D11TextureBase it encapsulates. */
-inline FD3D11TextureBase* GetD3D11TextureFromRHITexture(FRHITexture* Texture)
+FORCEINLINE FD3D11TextureBase* GetD3D11TextureFromRHITexture(FRHITexture* Texture)
 {
-	if(!Texture)
-	{
-		return NULL;
-	}
-	else if(Texture->GetTexture2D())
-	{
-		return static_cast<FD3D11Texture2D*>(Texture);
-	}
-	else if(Texture->GetTextureReference())
-	{
-		return static_cast<FD3D11TextureReference*>(Texture);
-	}
-	else if(Texture->GetTexture2DArray())
+	if (!Texture)
 	{
-		return static_cast<FD3D11Texture2DArray*>(Texture);
-	}
-	else if(Texture->GetTexture3D())
-	{
-		return static_cast<FD3D11Texture3D*>(Texture);
-	}
-	else if(Texture->GetTextureCube())
-	{
-		return static_cast<FD3D11TextureCube*>(Texture);
-	}
-	else
-	{
-		UE_LOG(LogD3D11RHI, Fatal,TEXT("Unknown RHI texture type"));
 		return NULL;
 	}
+	FD3D11TextureBase* Result((FD3D11TextureBase*)Texture->GetTextureBaseRHI());
+	check(Result);
+	return Result;
 }
 
 /** D3D11 render query */
@@ -641,56 +633,18 @@ public:
 	/** Resource table containing RHI references. */
 	TArray<TRefCountPtr<FRHIResource> > ResourceTable;
 
-	/** Cached resources need to retain the associated shader resource for bookkeeping purposes. */
-	struct FResourcePair
-	{		
-		FD3D11BaseShaderResource* ShaderResource;
-		IUnknown* D3D11Resource;
-		
-#if SUPPORT_RESOURCE_NAME
-		// useful to log resource transitions
-		void SetResourceName(const FName InName) { ResourceName = InName; }
-		FName GetResourceName() const { return ResourceName; }
-	private:
-		FName ResourceName;
-	public:
-#else
-		void SetResourceName(const FName InName) {}
-		FName GetResourceName() const { return NAME_None; }
-#endif
-	};
-
-	/** Raw resource table, cached once per frame. */
-	TArray<FResourcePair> RawResourceTable;
-
-	/** The frame in which RawResourceTable was last cached. */
-	uint32 LastCachedFrame;
-
 	/** Initialization constructor. */
 	FD3D11UniformBuffer(class FD3D11DynamicRHI* InD3D11RHI, const FRHIUniformBufferLayout& InLayout, ID3D11Buffer* InResource,const FRingAllocation& InRingAllocation)
 	: FRHIUniformBuffer(InLayout)
 	, Resource(InResource)
 	, RingAllocation(InRingAllocation)
-	, LastCachedFrame((uint32)-1)
 	, D3D11RHI(InD3D11RHI)
 	{}
 
 	virtual ~FD3D11UniformBuffer();
 
-	/** Cache resources if needed. */
-	inline void CacheResources(uint32 InFrameCounter)
-	{
-		if (InFrameCounter == INDEX_NONE || LastCachedFrame != InFrameCounter)
-		{
-			CacheResourcesInternal();
-			LastCachedFrame = InFrameCounter;
-		}
-	}
-
 private:
 	class FD3D11DynamicRHI* D3D11RHI;
-	/** Actually cache resources. */
-	void CacheResourcesInternal();
 };
 
 /** Index buffer resource class that stores stride information. */
diff --git a/Engine/Source/Runtime/Windows/D3D11RHI/Public/D3D11ShaderResources.h b/Engine/Source/Runtime/Windows/D3D11RHI/Public/D3D11ShaderResources.h
--- a/Engine/Source/Runtime/Windows/D3D11RHI/Public/D3D11ShaderResources.h
+++ b/Engine/Source/Runtime/Windows/D3D11RHI/Public/D3D11ShaderResources.h
@@ -8,4 +8,22 @@
 
 struct FD3D11ShaderResourceTable : public FBaseShaderResourceTable
 {
+	/** Mapping of bound Textures to their location in resource tables. */
+	TArray<uint32> TextureMap;
+
+	friend bool operator==(const FD3D11ShaderResourceTable& A, const FD3D11ShaderResourceTable& B)
+	{
+		const FBaseShaderResourceTable& BaseA = A;
+		const FBaseShaderResourceTable& BaseB = B;
+		return A == B && (FMemory::Memcmp(A.TextureMap.GetData(), B.TextureMap.GetData(), A.TextureMap.GetTypeSize()*A.TextureMap.Num()) == 0);
+	}
 };
+
+inline FArchive& operator<<(FArchive& Ar, FD3D11ShaderResourceTable& SRT)
+{
+	FBaseShaderResourceTable& BaseSRT = SRT;
+	Ar << BaseSRT;
+	Ar << SRT.TextureMap;
+	return Ar;
+}
+
diff --git a/Engine/Source/Runtime/Windows/D3D12RHI/Private/D3D12UAV.cpp b/Engine/Source/Runtime/Windows/D3D12RHI/Private/D3D12UAV.cpp
--- a/Engine/Source/Runtime/Windows/D3D12RHI/Private/D3D12UAV.cpp
+++ b/Engine/Source/Runtime/Windows/D3D12RHI/Private/D3D12UAV.cpp
@@ -215,6 +215,13 @@ FShaderResourceViewRHIRef FD3D12DynamicRHI::RHICreateShaderResourceView(FVertexB
 	return SRV;
 }
 
+FShaderResourceViewRHIRef FD3D12DynamicRHI::RHICreateShaderResourceView(FIndexBufferRHIParamRef BufferRHI)
+{
+	UE_LOG(LogRHI, Fatal, TEXT("D3D12 RHI doesn't support RHICreateShaderResourceView with FIndexBufferRHIParamRef yet!"));
+	
+	return FShaderResourceViewRHIRef();
+}
+
 void FD3D12CommandContext::RHIClearUAV(FUnorderedAccessViewRHIParamRef UnorderedAccessViewRHI, const uint32* Values)
 {
 	FD3D12UnorderedAccessView*  UnorderedAccessView = FD3D12DynamicRHI::ResourceCast(UnorderedAccessViewRHI);
