HMD: Engine changes for Oculus 1.2 support
- WIP: NickW has more improvements coming next week

diff --git a/Engine/Source/Runtime/Core/Private/Android/AndroidWindow.cpp b/Engine/Source/Runtime/Core/Private/Android/AndroidWindow.cpp
--- a/Engine/Source/Runtime/Core/Private/Android/AndroidWindow.cpp
+++ b/Engine/Source/Runtime/Core/Private/Android/AndroidWindow.cpp
@@ -96,13 +96,14 @@ void FAndroidWindow::ReleaseWindowRef(ANativeWindow* InWindow)
 	ANativeWindow_release(InWindow);
 }
 
+extern bool AndroidThunkCpp_HasMetaDataKey(const FString& Key);
 extern bool AndroidThunkCpp_IsGearVRApplication();
 
 FPlatformRect FAndroidWindow::GetScreenRect()
 {
 	// CSF is a multiplier to 1280x720
 	static IConsoleVariable* CVar = IConsoleManager::Get().FindConsoleVariable(TEXT("r.MobileContentScaleFactor"));
-	const bool bIsGearVRApp = AndroidThunkCpp_IsGearVRApplication();
+	const bool bIsGearVRApp = AndroidThunkCpp_HasMetaDataKey(TEXT("com.samsung.android.vr.application.mode"));
 	// If the app is for GearVR then always use 0 as ScaleFactor (to match window size).
 	float RequestedContentScaleFactor = (!bIsGearVRApp) ? CVar->GetFloat() : 0;
 
diff --git a/Engine/Source/Runtime/Core/Private/Misc/App.cpp b/Engine/Source/Runtime/Core/Private/Misc/App.cpp
--- a/Engine/Source/Runtime/Core/Private/Misc/App.cpp
+++ b/Engine/Source/Runtime/Core/Private/Misc/App.cpp
@@ -4,6 +4,7 @@
 #include "Runtime/Launch/Resources/Version.h"
 #include "Misc/App.h"
 
+DEFINE_LOG_CATEGORY_STATIC(LogApp, Log, All);
 
 /* FApp static initialization
  *****************************************************************************/
@@ -22,6 +23,8 @@ double FApp::LastTime = 0.0;
 double FApp::DeltaTime = 1 / 30.0;
 float FApp::VolumeMultiplier = 1.0f;
 float FApp::UnfocusedVolumeMultiplier = 0.0f;
+bool FApp::bUseVRFocus = false;
+bool FApp::bHasVRFocus = false;
 
 
 /* FApp static interface
@@ -195,4 +198,16 @@ void FApp::SetUnfocusedVolumeMultiplier(float InVolumeMultiplier)
 {
 	UnfocusedVolumeMultiplier = InVolumeMultiplier;
 	GUnfocusedVolumeMultiplierItialised = true;
-}
\ No newline at end of file
+}
+
+void FApp::SetUseVRFocus(bool bInUseVRFocus)
+{
+	UE_CLOG(bUseVRFocus != bInUseVRFocus, LogApp, Log, TEXT("UseVRFocus has changed to %d"), int(bInUseVRFocus));
+	bUseVRFocus = bInUseVRFocus;
+}
+
+void FApp::SetHasVRFocus(bool bInHasVRFocus)
+{
+	UE_CLOG(bHasVRFocus != bInHasVRFocus, LogApp, Log, TEXT("HasVRFocus has changed to %d"), int(bInHasVRFocus));
+	bHasVRFocus = bInHasVRFocus;
+}
diff --git a/Engine/Source/Runtime/Core/Private/Windows/WindowsApplication.cpp b/Engine/Source/Runtime/Core/Private/Windows/WindowsApplication.cpp
--- a/Engine/Source/Runtime/Core/Private/Windows/WindowsApplication.cpp
+++ b/Engine/Source/Runtime/Core/Private/Windows/WindowsApplication.cpp
@@ -1816,6 +1816,11 @@ void FWindowsApplication::PollGameDeviceState( const float TimeDelta )
 		bHasLoadedInputPlugins = true;
 	}
 
+	if (FApp::UseVRFocus() && !FApp::HasVRFocus())
+	{
+		return; // do not proceed if the app uses VR focus but doesn't have it
+	}
+
 	// Poll game device states and send new events
 	XInput->SendControllerEvents();
 
@@ -1829,6 +1834,11 @@ void FWindowsApplication::PollGameDeviceState( const float TimeDelta )
 
 void FWindowsApplication::SetForceFeedbackChannelValue(int32 ControllerId, FForceFeedbackChannelType ChannelType, float Value)
 {
+	if (FApp::UseVRFocus() && !FApp::HasVRFocus())
+	{
+		return; // do not proceed if the app uses VR focus but doesn't have it
+	}
+
 	// send vibration to externally-implemented devices
 	for( auto DeviceIt = ExternalInputDevices.CreateIterator(); DeviceIt; ++DeviceIt )
 	{
@@ -1838,6 +1848,11 @@ void FWindowsApplication::SetForceFeedbackChannelValue(int32 ControllerId, FForc
 
 void FWindowsApplication::SetForceFeedbackChannelValues(int32 ControllerId, const FForceFeedbackValues &Values)
 {
+	if (FApp::UseVRFocus() && !FApp::HasVRFocus())
+	{
+		return; // do not proceed if the app uses VR focus but doesn't have it
+	}
+
 	const FForceFeedbackValues* InternalValues = &Values;
  
 	XInput->SetChannelValues( ControllerId, *InternalValues );
@@ -1851,6 +1866,11 @@ void FWindowsApplication::SetForceFeedbackChannelValues(int32 ControllerId, cons
 
 void FWindowsApplication::SetHapticFeedbackValues(int32 ControllerId, int32 Hand, const FHapticFeedbackValues& Values)
 {
+	if (FApp::UseVRFocus() && !FApp::HasVRFocus())
+	{
+		return; // do not proceed if the app uses VR focus but doesn't have it
+	}
+
 	for (auto DeviceIt = ExternalInputDevices.CreateIterator(); DeviceIt; ++DeviceIt)
 	{
 		IHapticDevice* HapticDevice = (*DeviceIt)->GetHapticDevice();
diff --git a/Engine/Source/Runtime/Core/Private/Windows/WindowsPlatformMisc.cpp b/Engine/Source/Runtime/Core/Private/Windows/WindowsPlatformMisc.cpp
--- a/Engine/Source/Runtime/Core/Private/Windows/WindowsPlatformMisc.cpp
+++ b/Engine/Source/Runtime/Core/Private/Windows/WindowsPlatformMisc.cpp
@@ -896,11 +896,8 @@ void FWindowsPlatformMisc::PumpMessages(bool bFromMainLoop)
 	GPumpingMessagesOutsideOfMainLoop = false;
 	WinPumpMessages();
 
-	// check to see if the window in the foreground was made by this process (ie, does this app
-	// have focus)
-	::DWORD ForegroundProcess;
-	GetWindowThreadProcessId(GetForegroundWindow(), &ForegroundProcess);
-	bool HasFocus = ForegroundProcess == GetCurrentProcessId();
+	// Determine if application has focus
+	bool HasFocus = FApp::UseVRFocus() ? FApp::HasVRFocus() : FWindowsPlatformProcess::IsThisApplicationForeground();
 
 	// If editor thread doesn't have the focus, don't suck up too much CPU time.
 	if( GIsEditor )
diff --git a/Engine/Source/Runtime/Core/Public/Misc/App.h b/Engine/Source/Runtime/Core/Public/Misc/App.h
--- a/Engine/Source/Runtime/Core/Public/Misc/App.h
+++ b/Engine/Source/Runtime/Core/Public/Misc/App.h
@@ -541,6 +541,37 @@ public:
 	*/
 	static void SetUnfocusedVolumeMultiplier(float InVolumeMultiplier);
 
+	/**
+	 * Sets if VRFocus should be used.
+	 *
+	 * @param  bInUseVRFocus	new bUseVRFocus value
+	 */
+	static void SetUseVRFocus(bool bInUseVRFocus);
+
+	/**
+	 * Gets if VRFocus should be used
+	 */
+	FORCEINLINE static bool UseVRFocus()
+	{
+		return bUseVRFocus;
+	}
+	/**
+	 * Sets VRFocus, which indicates that the application should continue to render 
+	 * Audio and Video as if it had window focus, even though it may not.
+	 *
+	 * @param  bInVRFocus	new VRFocus value
+	 */
+	static void SetHasVRFocus(bool bInHasVRFocus);
+
+	/**
+	 * Gets VRFocus, which indicates that the application should continue to render 
+	 * Audio and Video as if it had window focus, even though it may not.
+	 */
+	FORCEINLINE static bool HasVRFocus()
+	{
+		return bHasVRFocus;
+	}
+
 private:
 
 	/** Holds the instance identifier. */
@@ -584,6 +615,12 @@ private:
 
 	/** Read from config to define the volume when app loses focus */
 	static float UnfocusedVolumeMultiplier;
+
+	/** Holds a flag indicating if VRFocus should be used */
+	static bool bUseVRFocus;
+
+	/** Holds a flag indicating if app has focus in side the VR headset */
+	static bool bHasVRFocus;
 };
 
 
diff --git a/Engine/Source/Runtime/Engine/Classes/Kismet/HeadMountedDisplayFunctionLibrary.h b/Engine/Source/Runtime/Engine/Classes/Kismet/HeadMountedDisplayFunctionLibrary.h
--- a/Engine/Source/Runtime/Engine/Classes/Kismet/HeadMountedDisplayFunctionLibrary.h
+++ b/Engine/Source/Runtime/Engine/Classes/Kismet/HeadMountedDisplayFunctionLibrary.h
@@ -3,6 +3,16 @@
 #pragma once
 #include "HeadMountedDisplayFunctionLibrary.generated.h"
 
+UENUM()
+namespace ETrackingOriginSelector
+{
+	enum Type
+	{
+		Eye UMETA(DisplayName = "Eye Level"),
+		Floor UMETA(DisplayName = "Floor Level"),
+	};
+}
+
 UENUM()
 namespace EOrientPositionSelector
 {
@@ -52,14 +62,14 @@ class ENGINE_API UHeadMountedDisplayFunctionLibrary : public UBlueprintFunctionL
 	static bool HasValidTrackingPosition();
 
 	/**
-	 * If the HMD has a positional tracking camera, this will return the game-world location of the camera, as well as the parameters for the bounding region of tracking.
-	 * This allows an in-game representation of the legal positional tracking range.  All values will be zeroed if the camera is not available or the HMD does not support it.
+	 * If the HMD has a positional sensor, this will return the game-world location of it, as well as the parameters for the bounding region of tracking.
+	 * This allows an in-game representation of the legal positional tracking range.  All values will be zeroed if the sensor is not available or the HMD does not support it.
 	 *
-	 * @param CameraOrigin		(out) Origin, in world-space, of the tracking camera
-	 * @param CameraRotation	(out) Rotation, in world-space, of the tracking camera
-	 * @param HFOV				(out) Field-of-view, horizontal, in degrees, of the valid tracking zone of the camera
-	 * @param VFOV				(out) Field-of-view, vertical, in degrees, of the valid tracking zone of the camera
-	 * @param CameraDistance	(out) Nominal distance to camera, in world-space
+	 * @param CameraOrigin		(out) Origin, in world-space, of the sensor
+	 * @param CameraRotation	(out) Rotation, in world-space, of the sensor
+	 * @param HFOV				(out) Field-of-view, horizontal, in degrees, of the valid tracking zone of the sensor
+	 * @param VFOV				(out) Field-of-view, vertical, in degrees, of the valid tracking zone of the sensor
+	 * @param CameraDistance	(out) Nominal distance to sensor, in world-space
 	 * @param NearPlane			(out) Near plane distance of the tracking volume, in world-space
 	 * @param FarPlane			(out) Far plane distance of the tracking volume, in world-space
 	 */
@@ -131,4 +141,25 @@ class ENGINE_API UHeadMountedDisplayFunctionLibrary : public UBlueprintFunctionL
 	*/
 	UFUNCTION(BlueprintPure, Category = "Input|HeadMountedDisplay", meta = (WorldContext = "WorldContext"))
 	static float GetWorldToMetersScale(UObject* WorldContext);
+
+	/**
+	 * Sets current tracking origin type (EyeLevel or FloorLevel).
+	 */
+	UFUNCTION(BlueprintCallable, Category = "Input|HeadMountedDisplay")
+	static void SetTrackingOrigin(ETrackingOriginSelector::Type Origin);
+
+	/**
+	 * Returns current tracking origin type (EyeLevel or FloorLevel).
+	 */
+	UFUNCTION(BlueprintPure, Category = "Input|HeadMountedDisplay")
+	static ETrackingOriginSelector::Type GetTrackingOrigin();
+
+	/**
+	 * Returns current state of VR focus.
+	 *
+	 * @param bUseFocus		(out) if set to true, then this App does use VR focus.
+	 * @param bHasFocus		(out) if set to true, then this App currently has VR focus.
+	 */
+	UFUNCTION(BlueprintPure, Category="Input|HeadMountedDisplay", meta=(DisplayName="Get VR Focus State"))
+	static void GetVRFocusState(bool& bUseFocus, bool& bHasFocus);
 };
diff --git a/Engine/Source/Runtime/Engine/Private/HeadMountedDisplayFunctionLibrary.cpp b/Engine/Source/Runtime/Engine/Private/HeadMountedDisplayFunctionLibrary.cpp
--- a/Engine/Source/Runtime/Engine/Private/HeadMountedDisplayFunctionLibrary.cpp
+++ b/Engine/Source/Runtime/Engine/Private/HeadMountedDisplayFunctionLibrary.cpp
@@ -167,4 +167,48 @@ void UHeadMountedDisplayFunctionLibrary::SetWorldToMetersScale(UObject* WorldCon
 float UHeadMountedDisplayFunctionLibrary::GetWorldToMetersScale(UObject* WorldContext)
 {
 	return WorldContext ? WorldContext->GetWorld()->GetWorldSettings()->WorldToMeters : 0.f;
-}
\ No newline at end of file
+}
+
+void UHeadMountedDisplayFunctionLibrary::SetTrackingOrigin(ETrackingOriginSelector::Type InOrigin)
+{
+	if (GEngine->HMDDevice.IsValid())
+	{
+		IHeadMountedDisplay::ETrackingOrigin Origin = IHeadMountedDisplay::Eye;
+		switch (InOrigin)
+		{
+		case ETrackingOriginSelector::Eye:
+			Origin = IHeadMountedDisplay::Eye;
+			break;
+		case ETrackingOriginSelector::Floor:
+			Origin = IHeadMountedDisplay::Floor;
+			break;
+		default: check(0); // shouldn't happen, check the EOriginTypeSelector
+		}
+		GEngine->HMDDevice->SetTrackingOrigin(Origin);
+	}
+}
+
+ETrackingOriginSelector::Type UHeadMountedDisplayFunctionLibrary::GetTrackingOrigin()
+{
+	if (GEngine->HMDDevice.IsValid())
+	{
+		IHeadMountedDisplay::ETrackingOrigin Origin = GEngine->HMDDevice->GetTrackingOrigin();
+		switch (Origin)
+		{
+		case IHeadMountedDisplay::Eye:
+			return ETrackingOriginSelector::Eye;
+			break;
+		case IHeadMountedDisplay::Floor:
+			return ETrackingOriginSelector::Floor;
+			break;
+		default: check(0); // shouldn't happen, check the EOriginTypeSelector
+		}
+	}
+	return ETrackingOriginSelector::Eye;
+}
+
+void UHeadMountedDisplayFunctionLibrary::GetVRFocusState(bool& bUseFocus, bool& bHasFocus)
+{
+	bUseFocus = GEngine->HMDDevice->DoesAppUseVRFocus();
+	bHasFocus = GEngine->HMDDevice->DoesAppHaveVRFocus();
+}
diff --git a/Engine/Source/Runtime/Engine/Private/Slate/SceneViewport.cpp b/Engine/Source/Runtime/Engine/Private/Slate/SceneViewport.cpp
--- a/Engine/Source/Runtime/Engine/Private/Slate/SceneViewport.cpp
+++ b/Engine/Source/Runtime/Engine/Private/Slate/SceneViewport.cpp
@@ -898,6 +898,34 @@ FReply FSceneViewport::OnFocusReceived(const FFocusEvent& InFocusEvent)
 	{
 		FScopedConditionalWorldSwitcher WorldSwitcher(ViewportClient);
 		ViewportClient->ReceivedFocus(this);
+
+		if ((FApp::IsGame() && !GIsEditor) || bIsPlayInEditorViewport)
+		{
+			if (IsForegroundWindow() || FApp::UseVRFocus())
+			{
+				bool bIsCursorForcedVisible = false;
+				if (ViewportClient->GetWorld() && ViewportClient->GetWorld()->GetGameInstance() && ViewportClient->GetWorld()->GetGameInstance()->GetFirstLocalPlayerController())
+				{
+					bIsCursorForcedVisible = ViewportClient->GetWorld()->GetGameInstance()->GetFirstLocalPlayerController()->GetMouseCursor() != EMouseCursor::None;
+				}
+
+				const bool bPlayInEditorCapture = !bIsPlayInEditorViewport || InFocusEvent.GetCause() != EFocusCause::SetDirectly || bPlayInEditorGetsMouseControl;
+
+				// capturing the mouse interferes with slate UI (like the virtual joysticks)
+				if (FPlatformProperties::SupportsWindowedMode() && bPlayInEditorCapture && !bIsCursorForcedVisible && !FSlateApplication::Get().IsFakingTouchEvents())
+				{
+					// Only require the user to click in the window the first time - after that return focus to the game so long as it was the last focused widget.
+					// Means that tabbing in/out will return the mouse control to where it was & the in-game console won't leave the mouse under editor control.
+					bPlayInEditorGetsMouseControl = true;
+					CurrentReplyState.UseHighPrecisionMouseMovement(ViewportWidget.Pin().ToSharedRef());
+					CurrentReplyState.LockMouseToWidget(ViewportWidget.Pin().ToSharedRef());
+				}
+			}
+			else
+			{
+				CurrentReplyState.ClearUserFocus(true);
+			}
+		}
 	}
 
 	// Update key state mappings so that the the viewport modifier states are valid upon focus.
@@ -1012,7 +1040,9 @@ void FSceneViewport::ResizeFrame(uint32 NewSizeX, uint32 NewSizeY, EWindowMode::
 			uint32 ViewportSizeX = NewSizeX;
 			uint32 ViewportSizeY = NewSizeY;
 
-			if (GEngine->HMDDevice.IsValid() && GEngine->HMDDevice->IsHMDConnected())
+			bool bIsHMDConnected = GEngine->HMDDevice.IsValid() && GEngine->HMDDevice->IsHMDConnected();
+
+			if (bIsHMDConnected)
 			{
 				WindowToResize->SetViewportSizeDrivenByWindow(true);
 				// Resize & move only if moving to a fullscreen mode
@@ -1043,11 +1073,11 @@ void FSceneViewport::ResizeFrame(uint32 NewSizeX, uint32 NewSizeY, EWindowMode::
 			}
 
 			// Avoid resizing if nothing changes.
-			bool bNeedsResize = SizeX != ViewportSizeX || SizeY != ViewportSizeY || NewWindowMode != DesiredWindowMode || DesiredWindowMode != WindowToResize->GetWindowMode();
+			bool bNeedsResize = SizeX != ViewportSizeX || SizeY != ViewportSizeY || NewWindowMode != WindowMode || DesiredWindowMode != WindowToResize->GetWindowMode();
 
 			if (bNeedsResize)
 			{
-				if (GEngine->HMDDevice.IsValid() && GEngine->HMDDevice->IsHMDConnected())
+				if (bIsHMDConnected)
 				{
 					// Resize & move only if moving to a fullscreen mode
 					if (NewWindowMode != EWindowMode::Windowed)
@@ -1070,7 +1100,7 @@ void FSceneViewport::ResizeFrame(uint32 NewSizeX, uint32 NewSizeY, EWindowMode::
 				// Toggle fullscreen and resize
 				WindowToResize->SetWindowMode(DesiredWindowMode);
 
-				if (GEngine->HMDDevice.IsValid() && GEngine->HMDDevice->IsHMDEnabled())
+				if (bIsHMDConnected)
 				{
 					if (NewWindowMode == EWindowMode::Windowed)
 					{
diff --git a/Engine/Source/Runtime/HeadMountedDisplay/Private/HeadMountedDisplayModule.cpp b/Engine/Source/Runtime/HeadMountedDisplay/Private/HeadMountedDisplayModule.cpp
--- a/Engine/Source/Runtime/HeadMountedDisplay/Private/HeadMountedDisplayModule.cpp
+++ b/Engine/Source/Runtime/HeadMountedDisplay/Private/HeadMountedDisplayModule.cpp
@@ -42,6 +42,36 @@ static TAutoConsoleVariable<int32> CVarHiddenAreaMask(
 	TEXT("0 to disable hidden area mask, 1 to enable."),
 	ECVF_Scalability | ECVF_RenderThreadSafe);
 
+
+static void SetTrackingOrigin(const TArray<FString>& Args)
+{
+	int Origin = 0;
+	if (Args.Num())
+	{
+		Origin = FCString::Atoi(*Args[0]);
+		if (GEngine && GEngine->HMDDevice.IsValid())
+		{
+			GEngine->HMDDevice->SetTrackingOrigin(IHeadMountedDisplay::ETrackingOrigin(Origin));
+		}
+	}
+	else
+	{
+		if (GEngine && GEngine->HMDDevice.IsValid())
+		{
+			Origin = GEngine->HMDDevice->GetTrackingOrigin();
+		}
+		if (GLog)
+		{
+			GLog->Logf(ELogVerbosity::Display, TEXT("Tracking orgin is set to %d"), Origin);
+		}
+	}
+}
+
+static FAutoConsoleCommand CTrackingOriginCmd(
+	TEXT("vr.SetTrackingOrigin"),
+	TEXT("0 - tracking origin is at the eyes/head, 1 - tracking origin is at the floor."),
+	FConsoleCommandWithArgsDelegate::CreateStatic(SetTrackingOrigin));
+
 class FHeadMountedDisplayModule : public IHeadMountedDisplayModule
 {
 	virtual TSharedPtr< class IHeadMountedDisplay, ESPMode::ThreadSafe > CreateHeadMountedDisplay()
@@ -104,4 +134,32 @@ void IHeadMountedDisplay::ApplyLateUpdate(FSceneInterface* Scene, const FTransfo
 		}
 	}
 	LateUpdatePrimitives.Reset();
-}
\ No newline at end of file
+}
+
+void IHeadMountedDisplay::SetTrackingOrigin(ETrackingOrigin InOrigin)
+{
+	if (GLog)
+	{
+		GLog->Log(ELogVerbosity::Display, TEXT("Not implemented IHeadMountedDisplay::SetTrackingOrigin is called"));
+	}
+}
+
+IHeadMountedDisplay::ETrackingOrigin IHeadMountedDisplay::GetTrackingOrigin()
+{
+	if (GLog)
+	{
+		GLog->Log(ELogVerbosity::Display, TEXT("Not implemented IHeadMountedDisplay::GetTrackingOrigin is called"));
+	}
+	return Eye;
+}
+bool IHeadMountedDisplay::DoesAppUseVRFocus() const
+{
+	return FApp::UseVRFocus();
+}
+
+bool IHeadMountedDisplay::DoesAppHaveVRFocus() const
+{
+	return FApp::HasVRFocus();
+}
+
+
diff --git a/Engine/Source/Runtime/HeadMountedDisplay/Public/IHeadMountedDisplay.h b/Engine/Source/Runtime/HeadMountedDisplay/Public/IHeadMountedDisplay.h
--- a/Engine/Source/Runtime/HeadMountedDisplay/Public/IHeadMountedDisplay.h
+++ b/Engine/Source/Runtime/HeadMountedDisplay/Public/IHeadMountedDisplay.h
@@ -70,7 +70,7 @@ public:
 	virtual void	GetFieldOfView(float& InOutHFOVInDegrees, float& InOutVFOVInDegrees) const = 0;
 
 	/**
-	 * Whether or not the HMD supports positional tracking (either via camera or other means)
+	 * Whether or not the HMD supports positional tracking (either via sensor or other means)
 	 */
 	virtual bool	DoesSupportPositionalTracking() const = 0;
 
@@ -80,7 +80,7 @@ public:
 	virtual bool	HasValidTrackingPosition() = 0;
 
 	/**
-	 * If the HMD supports positional tracking via a camera, this returns the frustum properties (all in game-world space) of the tracking camera.
+	 * If the HMD supports positional tracking via a sensor, this returns the frustum properties (all in game-world space) of the sensor.
 	 */
 	virtual void	GetPositionalTrackingCameraProperties(FVector& OutOrigin, FQuat& OutOrientation, float& OutHFOV, float& OutVFOV, float& OutCameraDistance, float& OutNearPlane, float& OutFarPlane) const = 0;
 
@@ -346,6 +346,34 @@ public:
 	 */
 	virtual FString GetVersionString() const { return FString(TEXT("GenericHMD")); }
 
+	enum ETrackingOrigin
+	{
+		Eye		= 0,
+		Floor	= 1,
+	};
+
+	/**
+	 * Sets tracking origin (either 'eye'-level or 'floor'-level).
+	 */
+	virtual void SetTrackingOrigin(ETrackingOrigin);
+
+	/**
+	 * Returns current tracking origin.
+	 */
+	virtual ETrackingOrigin GetTrackingOrigin();
+
+	/**
+	 * Returns true, if the App is using VR focus. This means that the app may handle lifecycle events differently from
+	 * the regular desktop apps. In this case, FCoreDelegates::ApplicationWillEnterBackgroundDelegate and FCoreDelegates::ApplicationHasEnteredForegroundDelegate
+	 * reflect the state of VR focus (either the app should be rendered in HMD or not).
+	 */
+	virtual bool DoesAppUseVRFocus() const;
+
+	/**
+	 * Returns true, if the app has VR focus, meaning if it is rendered in the HMD.
+	 */
+	virtual bool DoesAppHaveVRFocus() const;
+
 	/** Setup state for applying the render thread late update */
 	virtual void SetupLateUpdate(const FTransform& ParentToWorld, USceneComponent* Component);
 
diff --git a/Engine/Source/Runtime/Online/Voice/Private/Windows/VoiceCaptureWindows.cpp b/Engine/Source/Runtime/Online/Voice/Private/Windows/VoiceCaptureWindows.cpp
--- a/Engine/Source/Runtime/Online/Voice/Private/Windows/VoiceCaptureWindows.cpp
+++ b/Engine/Source/Runtime/Online/Voice/Private/Windows/VoiceCaptureWindows.cpp
@@ -124,7 +124,7 @@ bool FVoiceCaptureWindows::Init(int32 SampleRate, int32 NumChannels)
 	}
 
 	// DSDEVID_DefaultCapture WAVEINCAPS 
-	HRESULT hr = DirectSoundCaptureCreate8(&DSDEVID_DefaultVoiceCapture, &CV->VoiceCaptureDev, NULL);
+	HRESULT hr = DirectSoundCaptureCreate8(&VoiceDev->VoiceCaptureDeviceGuid, &CV->VoiceCaptureDev, NULL);
 	if (FAILED(hr))
 	{
 		//DSERR_ALLOCATED, DSERR_INVALIDPARAM, DSERR_NOAGGREGATION, DSERR_OUTOFMEMORY
diff --git a/Engine/Source/Runtime/Online/Voice/Private/Windows/VoiceCaptureWindows.h b/Engine/Source/Runtime/Online/Voice/Private/Windows/VoiceCaptureWindows.h
--- a/Engine/Source/Runtime/Online/Voice/Private/Windows/VoiceCaptureWindows.h
+++ b/Engine/Source/Runtime/Online/Voice/Private/Windows/VoiceCaptureWindows.h
@@ -61,6 +61,13 @@ public:
 	/** DirectSound8 Interface */
 	LPDIRECTSOUND8 DirectSound;
 
+	/** HMD voice capture device index, -1 if unused */
+	int HmdVoiceCaptureDeviceIndex;
+	/** Count of enumerated voice capture devices */
+	int VoiceCaptureDeviceCount;
+	/** GUID of selected voice capture device */
+	GUID VoiceCaptureDeviceGuid;
+
 	/** Constructors */
 	FVoiceCaptureDeviceWindows();
 	~FVoiceCaptureDeviceWindows();
diff --git a/Engine/Source/Runtime/Online/Voice/Private/Windows/VoiceModuleWindows.cpp b/Engine/Source/Runtime/Online/Voice/Private/Windows/VoiceModuleWindows.cpp
--- a/Engine/Source/Runtime/Online/Voice/Private/Windows/VoiceModuleWindows.cpp
+++ b/Engine/Source/Runtime/Online/Voice/Private/Windows/VoiceModuleWindows.cpp
@@ -5,6 +5,13 @@
 #include "VoiceCaptureWindows.h"
 #include "VoiceCodecOpus.h"
 #include "Voice.h"
+#include "Runtime/HeadMountedDisplay/Public/IHeadMountedDisplayModule.h"
+
+static TAutoConsoleVariable<int32> CVarHmdDirectSoundVoiceCaptureDeviceIndex(
+	TEXT("hmd.DirectSoundVoiceCaptureDeviceIndex"),
+	-1,
+	TEXT("Specifies the DirectSound device index to use when HMD is connected. (-1 == Unknown)\n"),
+	ECVF_Default);
 
 FVoiceCaptureDeviceWindows* FVoiceCaptureDeviceWindows::Singleton = NULL;
 
@@ -32,8 +39,17 @@ BOOL CALLBACK CaptureDeviceCallback(
 	)
 {
 	// @todo identify the proper device
-	FVoiceCaptureWindows* VCPtr = (FVoiceCaptureWindows*)(lpContext);
+	FVoiceCaptureDeviceWindows* VCPtr = (FVoiceCaptureDeviceWindows*)(lpContext);
 	UE_LOG(LogVoiceCapture, Display, TEXT("Device: %s Desc: %s GUID: %s Context:0x%08x"), lpcstrDescription, lpcstrModule, *PrintMSGUID(lpGuid), lpContext);
+
+	// Allow HMD to override the voice capture device
+	if(VCPtr->VoiceCaptureDeviceCount == VCPtr->HmdVoiceCaptureDeviceIndex)
+	{
+		UE_LOG(LogVoice, Display, TEXT("VoiceCapture device overridden by HMD to use '%s' %s"), lpcstrDescription, *PrintMSGUID(lpGuid));
+		VCPtr->VoiceCaptureDeviceGuid = *lpGuid;
+	}
+
+	VCPtr->VoiceCaptureDeviceCount++;
 	return true;
 }
 
@@ -245,16 +261,11 @@ bool FVoiceCaptureDeviceWindows::Init()
 		UE_LOG(LogVoiceCapture, Warning, TEXT("Failed to init DirectSound %d"), hr);
 		return false;
 	}
-	
-	if (0)
-	{
-		hr = DirectSoundEnumerate((LPDSENUMCALLBACK)CaptureDeviceCallback, this);
-		if (FAILED(hr))
-		{
-			UE_LOG(LogVoiceCapture, Warning, TEXT("Failed to enumerate render devices %d"), hr);
-			return false;
-		}
-	}
+
+	bool bHmdAvailable = IModularFeatures::Get().IsModularFeatureAvailable(IHeadMountedDisplayModule::GetModularFeatureName());
+	HmdVoiceCaptureDeviceIndex = bHmdAvailable ? CVarHmdDirectSoundVoiceCaptureDeviceIndex.GetValueOnGameThread() : -1;
+	VoiceCaptureDeviceCount = 0;
+	VoiceCaptureDeviceGuid = DSDEVID_DefaultVoiceCapture;
 
 	hr = DirectSoundCaptureEnumerate((LPDSENUMCALLBACK)CaptureDeviceCallback, this);
 	if (FAILED(hr))
diff --git a/Engine/Source/Runtime/OpenGLDrv/Private/Android/AndroidOpenGLPrivate.h b/Engine/Source/Runtime/OpenGLDrv/Private/Android/AndroidOpenGLPrivate.h
--- a/Engine/Source/Runtime/OpenGLDrv/Private/Android/AndroidOpenGLPrivate.h
+++ b/Engine/Source/Runtime/OpenGLDrv/Private/Android/AndroidOpenGLPrivate.h
@@ -7,6 +7,8 @@
 
 #include "AndroidApplication.h"
 
+extern bool AndroidThunkCpp_HasMetaDataKey(const FString& Key);
+
 bool GAndroidGPUInfoReady = false;
 
 // call out to JNI to see if the application was packaged for GearVR
@@ -52,7 +54,8 @@ private:
 		}
 #if !PLATFORM_ANDROIDES31
 		// Do not create a window surface if the app is for GearVR
-		EGL->InitSurface(!AndroidThunkCpp_IsGearVRApplication());
+		const bool bIsGearVRApp = AndroidThunkCpp_HasMetaDataKey(TEXT("com.samsung.android.vr.application.mode"));
+		EGL->InitSurface(!bIsGearVRApp);
 #endif
 		EGL->SetCurrentSharedContext();
 
diff --git a/Engine/Source/Runtime/Windows/D3D11RHI/Private/D3D11Commands.cpp b/Engine/Source/Runtime/Windows/D3D11RHI/Private/D3D11Commands.cpp
--- a/Engine/Source/Runtime/Windows/D3D11RHI/Private/D3D11Commands.cpp
+++ b/Engine/Source/Runtime/Windows/D3D11RHI/Private/D3D11Commands.cpp
@@ -32,7 +32,9 @@
 
 #if PLATFORM_DESKTOP
 // For Depth Bounds Test interface
-#include "nvapi.h"
+#include "AllowWindowsPlatformTypes.h"
+	#include "nvapi.h"
+#include "HideWindowsPlatformTypes.h"
 #endif
 
 FGlobalBoundShaderState GD3D11ClearMRTBoundShaderState[8];
diff --git a/Engine/Source/Runtime/Windows/D3D11RHI/Private/Windows/WindowsD3D11Device.cpp b/Engine/Source/Runtime/Windows/D3D11RHI/Private/Windows/WindowsD3D11Device.cpp
--- a/Engine/Source/Runtime/Windows/D3D11RHI/Private/Windows/WindowsD3D11Device.cpp
+++ b/Engine/Source/Runtime/Windows/D3D11RHI/Private/Windows/WindowsD3D11Device.cpp
@@ -28,7 +28,7 @@ static TAutoConsoleVariable<int32> CVarGraphicsAdapter(
 	ECVF_RenderThreadSafe);
 
 static TAutoConsoleVariable<int32> CVarHmdGraphicsAdapter(
-	TEXT("r.HmdGraphicsAdapter"),
+	TEXT("hmd.D3D11GraphicsAdapter"),
 	-1,
 	TEXT("Specifies the index of the graphics adapter where the HMD is connected.  Overrides r.GraphicsAdapter when the Hmd is enabled.\n")
 	TEXT("At the moment this only works on Direct3D 11.\n")
@@ -319,8 +319,10 @@ void FD3D11DynamicRHIModule::FindAdapter()
 
 				FD3D11Adapter CurrentAdapter(AdapterIndex, ActualFeatureLevel);
 
-				if(!OutputCount)
+				if(!OutputCount && !bUseHmdGraphicsAdapter)
 				{
+					// Add special check to support HMDs, which do not have associated outputs.
+
 					// This device has no outputs. Reject it, 
 					// http://msdn.microsoft.com/en-us/library/windows/desktop/bb205075%28v=vs.85%29.aspx#WARP_new_for_Win8
 					continue;
diff --git a/Engine/Source/Runtime/Windows/D3D12RHI/Private/D3D12RHIPrivate.h b/Engine/Source/Runtime/Windows/D3D12RHI/Private/D3D12RHIPrivate.h
--- a/Engine/Source/Runtime/Windows/D3D12RHI/Private/D3D12RHIPrivate.h
+++ b/Engine/Source/Runtime/Windows/D3D12RHI/Private/D3D12RHIPrivate.h
@@ -1165,6 +1165,11 @@ public:
 	virtual class IRHIComputeContext* RHIGetDefaultAsyncComputeContext() final override;
 	virtual class IRHICommandContextContainer* RHIGetCommandContextContainer() final override;
 
+	// FD3D12DynamicRHI interface.
+	virtual ID3D12CommandQueue* RHIGetD3DCommandQueue();
+	virtual FTexture2DRHIRef RHICreateTexture2DFromD3D12Resource( uint8 Format, uint32 Flags, const FClearValueBinding& ClearValueBinding, ID3D12Resource* Resource );
+	virtual void RHIAliasTexture2DResources( FTexture2DRHIParamRef DestTexture2D, FTexture2DRHIParamRef SrcTexture2D );
+
 #if UE_BUILD_DEBUG	
 	uint32 SubmissionLockStalls;
 	uint32 DrawCount;
diff --git a/Engine/Source/Runtime/Windows/D3D12RHI/Private/D3D12Texture.cpp b/Engine/Source/Runtime/Windows/D3D12RHI/Private/D3D12Texture.cpp
--- a/Engine/Source/Runtime/Windows/D3D12RHI/Private/D3D12Texture.cpp
+++ b/Engine/Source/Runtime/Windows/D3D12RHI/Private/D3D12Texture.cpp
@@ -2122,3 +2122,166 @@ void FD3D12CommandContext::RHIUpdateTextureReference(FTextureReferenceRHIParamRe
 		TextureRef->SetReferencedTexture(NewTextureRHI, NewTexture, NewSRV);
 	}
 }
+
+ID3D12CommandQueue* FD3D12DynamicRHI::RHIGetD3DCommandQueue()
+{
+	return MainDevice->GetCommandListManager().GetD3DCommandQueue();
+}
+
+FTexture2DRHIRef FD3D12DynamicRHI::RHICreateTexture2DFromD3D12Resource(uint8 Format, uint32 Flags, const FClearValueBinding& ClearValueBinding, ID3D12Resource* Resource)
+{
+	check(Resource);
+
+	D3D12_RESOURCE_DESC TextureDesc = Resource->GetDesc();
+	uint32 SizeX = TextureDesc.Width;
+	uint32 SizeY = TextureDesc.Height;
+	uint32 SizeZ = TextureDesc.DepthOrArraySize;
+	uint32 NumMips = TextureDesc.MipLevels;
+	uint32 NumSamples = TextureDesc.SampleDesc.Count;
+
+	SCOPE_CYCLE_COUNTER(STAT_D3D12CreateTextureTime);
+
+	if (GMaxRHIFeatureLevel <= ERHIFeatureLevel::ES2)
+	{
+		// Remove sRGB read flag when not supported
+		Flags &= ~TexCreate_SRGB;
+	}
+
+	const bool bSRGB = (Flags & TexCreate_SRGB) != 0;
+
+	const DXGI_FORMAT PlatformResourceFormat = TextureDesc.Format;
+	const DXGI_FORMAT PlatformShaderResourceFormat = FindShaderResourceDXGIFormat(PlatformResourceFormat, bSRGB);
+	const DXGI_FORMAT PlatformRenderTargetFormat = FindShaderResourceDXGIFormat(PlatformResourceFormat, bSRGB);
+	const DXGI_FORMAT PlatformDepthStencilFormat = FindDepthStencilDXGIFormat(PlatformResourceFormat);
+
+	// Determine the MSAA settings to use for the texture.
+	D3D12_DSV_DIMENSION DepthStencilViewDimension = D3D12_DSV_DIMENSION_TEXTURE2D;
+	D3D12_RTV_DIMENSION RenderTargetViewDimension = D3D12_RTV_DIMENSION_TEXTURE2D;
+	D3D12_SRV_DIMENSION ShaderResourceViewDimension = D3D12_SRV_DIMENSION_TEXTURE2D;
+
+	uint32 ActualMSAACount = NumSamples;
+	uint32 ActualMSAAQuality = TextureDesc.SampleDesc.Quality;
+
+	if (ActualMSAACount > 1)
+	{
+		DepthStencilViewDimension = D3D12_DSV_DIMENSION_TEXTURE2DMS;
+		RenderTargetViewDimension = D3D12_RTV_DIMENSION_TEXTURE2DMS;
+		ShaderResourceViewDimension = D3D12_SRV_DIMENSION_TEXTURE2DMS;
+	}
+
+	// Set up the texture bind flags.
+	bool bCreateRTV = (TextureDesc.Flags & D3D12_RESOURCE_FLAG_ALLOW_RENDER_TARGET) != 0;
+	bool bCreateDSV = (TextureDesc.Flags & D3D12_RESOURCE_FLAG_ALLOW_DEPTH_STENCIL) != 0;
+	bool bCreatedRTVPerSlice = false;
+
+	TRefCountPtr<FD3D12Resource> TextureResource = new FD3D12Resource(Resource, TextureDesc);
+	TArray<TRefCountPtr<FD3D12RenderTargetView> > RenderTargetViews;
+	TRefCountPtr<FD3D12DepthStencilView> DepthStencilViews[FExclusiveDepthStencil::MaxIndex];
+
+	if (bCreateRTV)
+	{
+		// Create a render target view for each mip
+		for (uint32 MipIndex = 0; MipIndex < NumMips; MipIndex++)
+		{
+			D3D12_RENDER_TARGET_VIEW_DESC RTVDesc;
+			FMemory::Memzero(&RTVDesc, sizeof(RTVDesc));
+			RTVDesc.Format = PlatformRenderTargetFormat;
+			RTVDesc.ViewDimension = RenderTargetViewDimension;
+			RTVDesc.Texture2D.MipSlice = MipIndex;
+			RTVDesc.Texture2D.PlaneSlice = GetPlaneSliceFromViewFormat(PlatformResourceFormat, RTVDesc.Format);
+
+			TRefCountPtr<FD3D12RenderTargetView> RenderTargetView = new FD3D12RenderTargetView(GetRHIDevice(), &RTVDesc, TextureResource);
+			RenderTargetViews.Add(RenderTargetView);
+		}
+	}
+
+	if (bCreateDSV)
+	{
+		// Create a depth-stencil-view for the texture.
+		D3D12_DEPTH_STENCIL_VIEW_DESC DSVDesc;
+		FMemory::Memzero(&DSVDesc, sizeof(DSVDesc));
+		DSVDesc.Format = FindDepthStencilDXGIFormat(PlatformResourceFormat);
+		DSVDesc.ViewDimension = DepthStencilViewDimension;
+		DSVDesc.Texture2D.MipSlice = 0;
+
+		const bool HasStencil = HasStencilBits(DSVDesc.Format);
+		for (uint32 AccessType = 0; AccessType < FExclusiveDepthStencil::MaxIndex; ++AccessType)
+		{
+			// Create a read-only access views for the texture.
+			DSVDesc.Flags = (AccessType & FExclusiveDepthStencil::DepthRead_StencilWrite) ? D3D12_DSV_FLAG_READ_ONLY_DEPTH : D3D12_DSV_FLAG_NONE;
+			if (HasStencil)
+			{
+				DSVDesc.Flags |= (AccessType & FExclusiveDepthStencil::DepthWrite_StencilRead) ? D3D12_DSV_FLAG_READ_ONLY_STENCIL : D3D12_DSV_FLAG_NONE;
+			}
+			DepthStencilViews[AccessType] = new FD3D12DepthStencilView(GetRHIDevice(), &DSVDesc, TextureResource, HasStencil);
+		}
+	}
+
+	// Create a shader resource view for the texture.
+    D3D12_SHADER_RESOURCE_VIEW_DESC SRVDesc = {};
+    SRVDesc.Shader4ComponentMapping = D3D12_DEFAULT_SHADER_4_COMPONENT_MAPPING;
+	SRVDesc.Format = PlatformShaderResourceFormat;
+	SRVDesc.ViewDimension = ShaderResourceViewDimension;
+	SRVDesc.Texture2D.MostDetailedMip = 0;
+	SRVDesc.Texture2D.MipLevels = NumMips;
+	SRVDesc.Texture2D.PlaneSlice = GetPlaneSliceFromViewFormat(PlatformResourceFormat, SRVDesc.Format);
+
+	TRefCountPtr<FD3D12ResourceLocation> TextureResourceLocation = new FD3D12ResourceLocation(GetRHIDevice());
+	TextureResourceLocation->SetFromD3DResource(TextureResource, 0, 0);
+
+	FD3D12Texture2D* Texture2D = new FD3D12Texture2D(
+		GetRHIDevice(),
+		TextureResourceLocation,
+		bCreatedRTVPerSlice,
+		TextureDesc.DepthOrArraySize,
+		RenderTargetViews,
+		DepthStencilViews,
+		SizeX,
+		SizeY,
+		SizeZ,
+		NumMips,
+		ActualMSAACount,
+		(EPixelFormat)Format,
+		false, // bCubeTexture,
+		Flags,
+		false, // bPooledTexture,
+		ClearValueBinding
+#if PLATFORM_SUPPORTS_VIRTUAL_TEXTURES
+		, RawTextureMemory
+#endif
+		);
+
+	if (Flags & TexCreate_CPUReadback)
+	{
+		const uint32 BlockBytes = GPixelFormats[PlatformResourceFormat].BlockBytes;
+		const uint32 XBytesAligned = Align((uint32)SizeX * BlockBytes, D3D12_TEXTURE_DATA_PITCH_ALIGNMENT);
+		D3D12_SUBRESOURCE_FOOTPRINT destSubresource;
+		destSubresource.Depth = SizeZ;
+		destSubresource.Height = SizeY;
+		destSubresource.Width = SizeX;
+		destSubresource.Format = PlatformResourceFormat;
+		destSubresource.RowPitch = XBytesAligned;
+
+		check(destSubresource.RowPitch % D3D12_TEXTURE_DATA_PITCH_ALIGNMENT == 0);	// Make sure we align correctly.
+
+		D3D12_PLACED_SUBRESOURCE_FOOTPRINT placedTexture2D = { 0 };
+		placedTexture2D.Offset = 0;
+		placedTexture2D.Footprint = destSubresource;
+		Texture2D->SetReadBackHeapDesc(placedTexture2D);
+	}
+
+	FD3D12ShaderResourceView* pWrappedShaderResourceView = new FD3D12ShaderResourceView(GetRHIDevice(), &SRVDesc, Texture2D->ResourceLocation);
+	Texture2D->SetShaderResourceView(pWrappedShaderResourceView);
+
+	D3D11TextureAllocated(*Texture2D);
+
+	return Texture2D;
+}
+
+void FD3D12DynamicRHI::RHIAliasTexture2DResources(FTexture2DRHIParamRef DestTexture2DRHI, FTexture2DRHIParamRef SrcTexture2DRHI)
+{
+	FD3D12Texture2D* DestTexture2D = FD3D12DynamicRHI::ResourceCast(DestTexture2DRHI);
+	FD3D12Texture2D* SrcTexture2D = FD3D12DynamicRHI::ResourceCast(SrcTexture2DRHI);
+
+	DestTexture2D->AliasResources(SrcTexture2D);
+}
diff --git a/Engine/Source/Runtime/Windows/D3D12RHI/Private/Windows/WindowsD3D12Device.cpp b/Engine/Source/Runtime/Windows/D3D12RHI/Private/Windows/WindowsD3D12Device.cpp
--- a/Engine/Source/Runtime/Windows/D3D12RHI/Private/Windows/WindowsD3D12Device.cpp
+++ b/Engine/Source/Runtime/Windows/D3D12RHI/Private/Windows/WindowsD3D12Device.cpp
@@ -9,6 +9,7 @@
 #include <delayimp.h>
 #include "HideWindowsPlatformTypes.h"
 #include "HardwareInfo.h"
+#include "Runtime/HeadMountedDisplay/Public/IHeadMountedDisplayModule.h"
 
 #pragma comment(lib, "d3d12.lib")
 
@@ -19,13 +20,21 @@ static TAutoConsoleVariable<int32> CVarGraphicsAdapter(
 	TEXT("r.D3D12GraphicsAdapter"),
 	-1,
 	TEXT("User request to pick a specific graphics adapter (e.g. when using a integrated graphics card with a descrete one)\n")
-	TEXT("At the moment this only works on Direct3D 11.\n")
 	TEXT(" -2: Take the first one that fulfills the criteria\n")
 	TEXT(" -1: Favour non integrated because there are usually faster\n")
 	TEXT("  0: Adpater #0\n")
 	TEXT("  1: Adpater #1, ..."),
 	ECVF_RenderThreadSafe);
 
+static TAutoConsoleVariable<int32> CVarHmdGraphicsAdapter(
+	TEXT("hmd.D3D12GraphicsAdapter"),
+	-1,
+	TEXT("Specifies the index of the graphics adapter where the HMD is connected.  Overrides r.D3D12GraphicsAdapter when the Hmd is enabled.\n")
+	TEXT(" -1: Unknown\n")
+	TEXT("  0: Adpater #0\n")
+	TEXT("  1: Adpater #1, ..."),
+	ECVF_RenderThreadSafe);
+
 namespace D3D12RHI
 {
 
@@ -246,7 +255,11 @@ void FD3D12DynamicRHIModule::FindAdapter()
 	bAllowPerfHUD = false;
 #endif
 
-	int32 CVarValue = CVarGraphicsAdapter.GetValueOnGameThread();
+	// Allow HMD to override which graphics adapter is chosen, so we pick the adapter where the HMD is connected
+	bool bUseHmdGraphicsAdapter = CVarHmdGraphicsAdapter.GetValueOnGameThread() >= 0 && 
+		IModularFeatures::Get().IsModularFeatureAvailable(IHeadMountedDisplayModule::GetModularFeatureName());
+
+	int32 CVarValue = bUseHmdGraphicsAdapter ? CVarHmdGraphicsAdapter.GetValueOnGameThread() : CVarGraphicsAdapter.GetValueOnGameThread();
 
 	const bool bFavorNonIntegrated = CVarValue == -1;
 
@@ -304,9 +317,9 @@ void FD3D12DynamicRHIModule::FindAdapter()
 
 				FD3D12Adapter CurrentAdapter(AdapterIndex, ActualFeatureLevel);
 
-				if (!OutputCount && !bIsWARP)
+				if (!OutputCount && !bIsWARP && !bUseHmdGraphicsAdapter)
 				{
-					// Add special check to support WARP, which does not have an output associated with it.
+					// Add special check to support WARP and HMDs, which do not have associated outputs.
 
 					// This device has no outputs. Reject it, 
 					// http://msdn.microsoft.com/en-us/library/windows/desktop/bb205075%28v=vs.85%29.aspx#WARP_new_for_Win8
diff --git a/Engine/Source/Runtime/Windows/D3D12RHI/Public/D3D12Resources.h b/Engine/Source/Runtime/Windows/D3D12RHI/Public/D3D12Resources.h
--- a/Engine/Source/Runtime/Windows/D3D12RHI/Public/D3D12Resources.h
+++ b/Engine/Source/Runtime/Windows/D3D12RHI/Public/D3D12Resources.h
@@ -2719,6 +2719,19 @@ public:
 		return (NumDepthStencilViews > 0);
 	}
 
+	void AliasResources(FD3D12TextureBase* Texture)
+	{
+		ResourceLocation = Texture->ResourceLocation;
+		BaseShaderResource = Texture->BaseShaderResource;
+		ShaderResourceView = Texture->ShaderResourceView;
+		RenderTargetViews = Texture->RenderTargetViews;
+
+		for (uint32 Index = 0; Index < FExclusiveDepthStencil::MaxIndex; Index++)
+		{
+			DepthStencilViews[Index] = Texture->DepthStencilViews[Index];
+		}
+	}
+
 protected:
 
 	/** Amount of memory allocated by this texture, in bytes. */
diff --git a/Engine/Source/Runtime/Windows/XAudio2/Private/XAudio2Device.cpp b/Engine/Source/Runtime/Windows/XAudio2/Private/XAudio2Device.cpp
--- a/Engine/Source/Runtime/Windows/XAudio2/Private/XAudio2Device.cpp
+++ b/Engine/Source/Runtime/Windows/XAudio2/Private/XAudio2Device.cpp
@@ -24,6 +24,13 @@
 #include "HideWindowsPlatformTypes.h"
 #include "TargetPlatform.h"
 #include "XAudio2Support.h"
+#include "Runtime/HeadMountedDisplay/Public/IHeadMountedDisplayModule.h"
+
+static TAutoConsoleVariable<int32> CVarXAudio2HmdDeviceIndex(
+	TEXT("hmd.XAudio2DeviceIndex"),
+	-1,
+	TEXT("Specifies the XAudio2 device index to use when HMD is connected. (-1 == Unknown)\n"),
+	ECVF_Default);
 
 DEFINE_LOG_CATEGORY(LogXAudio2);
 
@@ -128,7 +135,7 @@ bool FXAudio2Device::InitializeHardware()
 	}
 
 	// Initialize the audio device index to 0, which is the windows default device index
-	int32 DeviceIndex = 0;
+	UINT32 DeviceIndex = 0;
 
 	FString WindowsAudioDeviceName;
 	GConfig->GetString(TEXT("/Script/WindowsTargetPlatform.WindowsTargetSettings"), TEXT("AudioDevice"), WindowsAudioDeviceName, GEngineIni);
@@ -150,6 +157,20 @@ bool FXAudio2Device::InitializeHardware()
 		}
 	}
 
+	// Allow HMD to override which audio device is chosen
+	bool bUseHmdDeviceIndex = CVarXAudio2HmdDeviceIndex.GetValueOnGameThread() >= 0 && 
+		IModularFeatures::Get().IsModularFeatureAvailable(IHeadMountedDisplayModule::GetModularFeatureName());
+
+	if( bUseHmdDeviceIndex )
+	{
+		DeviceIndex = CVarXAudio2HmdDeviceIndex.GetValueOnGameThread();
+	}
+
+	if(DeviceIndex >= DeviceCount)
+	{
+		DeviceIndex = 0;
+	}
+
 	// Get the details of the desired device index (0 is default)
 	if (!ValidateAPICall(TEXT("GetDeviceDetails"),
 		DeviceProperties->XAudio2->GetDeviceDetails(DeviceIndex, &FXAudioDeviceProperties::DeviceDetails)))
@@ -201,7 +222,7 @@ bool FXAudio2Device::InitializeHardware()
 #else	//XAUDIO_SUPPORTS_DEVICE_DETAILS
 	// Create the final output voice
 	if (!ValidateAPICall(TEXT("CreateMasteringVoice"),
-		DeviceProperties->XAudio2->CreateMasteringVoice(&DeviceProperties->MasteringVoice, UE4_XAUDIO2_NUMCHANNELS, UE4_XAUDIO2_SAMPLERATE, 0, 0, NULL )))
+		DeviceProperties->XAudio2->CreateMasteringVoice(&DeviceProperties->MasteringVoice, UE4_XAUDIO2_NUMCHANNELS, UE4_XAUDIO2_SAMPLERATE, 0, DeviceIndex, nullptr )))
 	{
 		UE_LOG(LogInit, Warning, TEXT( "Failed to create the mastering voice for XAudio2" ) );
 		DeviceProperties->XAudio2 = nullptr;
