Merging //UE4/Release-4.11 to //UE4/Main (Up to CL#2897738)

diff --git a/Engine/Source/Runtime/Windows/D3D12RHI/Private/D3D12Allocation.cpp b/Engine/Source/Runtime/Windows/D3D12RHI/Private/D3D12Allocation.cpp
new file mode 100644
--- /dev/null
+++ b/Engine/Source/Runtime/Windows/D3D12RHI/Private/D3D12Allocation.cpp
@@ -0,0 +1,1210 @@
+// Copyright 1998-2016 Epic Games, Inc. All Rights Reserved.
+
+// Implementation of Memory Allocation Strategies
+
+//-----------------------------------------------------------------------------
+//	Include Files
+//-----------------------------------------------------------------------------
+#include "D3D12RHIPrivate.h"
+#include "D3D12Allocation.h"
+
+
+//-----------------------------------------------------------------------------
+//	Allocator Base
+//-----------------------------------------------------------------------------
+FD3D12ResourceAllocator::FD3D12ResourceAllocator(FD3D12Device* ParentDevice, FString Name, D3D12_HEAP_TYPE heapType, D3D12_RESOURCE_FLAGS flags, uint32 MaxSizeForPooling)
+	: DebugName(Name)
+	, HeapType(heapType)
+	, ResourceFlags(flags)
+	, Initialized(false)
+	, MaximumAllocationSizeForPooling(MaxSizeForPooling)
+#if defined(UE_BUILD_DEBUG)
+	, PeakUsage(0)
+	, SpaceUsed(0)
+	, InternalFragmentation(0)
+	, NumBlocksInDeferredDeletionQueue(0)
+	, FailedAllocationSpace(0)
+#endif
+	, FD3D12DeviceChild(ParentDevice)
+{
+}
+
+FD3D12ResourceAllocator::~FD3D12ResourceAllocator()
+{
+}
+
+void FD3D12ResourceAllocator::InitializeDefaultBuffer(FD3D12Resource* Destination, uint64 DestinationOffset, const void* Data, uint64 DataSize)
+{
+	check(HeapType == D3D12_HEAP_TYPE_DEFAULT);
+	FD3D12CommandListHandle& hCommandList = GetParentDevice()->GetDefaultCommandContext().CommandListHandle;
+	// Queue up a CopySubresource() to handle initial data...
+
+	// Get an upload heap and initialize data
+	FD3D12ResourceLocation SrcResourceLoc(GetParentDevice());
+	void* pData = GetParentDevice()->GetBufferInitFastAllocator().Allocate(DataSize, 4UL, &SrcResourceLoc);
+	check(pData);
+	FMemory::Memcpy(pData, Data, DataSize);
+
+	// Copy from the temporary upload heap to the default resource
+	{
+		// Writable structured bufferes are sometimes initialized with inital data which means they sometimes need tracking.
+		FConditionalScopeResourceBarrier ConditionalScopeResourceBarrier(hCommandList, Destination, D3D12_RESOURCE_STATE_COPY_DEST, 0);
+
+		GetParentDevice()->GetDefaultCommandContext().numCopies++;
+		hCommandList->CopyBufferRegion(
+			Destination->GetResource(),
+			DestinationOffset,
+			SrcResourceLoc.GetResource()->GetResource(),
+			SrcResourceLoc.GetOffset(), DataSize);
+	}
+}
+
+//-----------------------------------------------------------------------------
+//	Buddy Allocator
+//-----------------------------------------------------------------------------
+
+FD3D12BuddyAllocator::FD3D12BuddyAllocator(FD3D12Device* ParentDevice, FString Name, eBuddyAllocationStrategy allocationStrategy, D3D12_HEAP_TYPE heapType,
+	D3D12_HEAP_FLAGS heapFlags, D3D12_RESOURCE_FLAGS flags, uint32 MaxSizeForPooling, uint32 maxBlockSize, uint32 MinBlockSize)
+	: AllocationStrategy(allocationStrategy)
+	, MaxBlockSize(maxBlockSize)
+	, MinBlockSize(MinBlockSize)
+	, HeapFlags(heapFlags)
+	, BackingHeap(nullptr)
+	, BaseAddress(nullptr)
+	, HeapFullMessageDisplayed(false)
+	, TotalSizeUsed(0)
+	, FD3D12ResourceAllocator(ParentDevice, Name, heapType, flags, MaxSizeForPooling)
+{
+	// maxBlockSize should be evenly dividable by MinBlockSize and  
+	// maxBlockSize / MinBlockSize should be a power of two  
+	check((maxBlockSize / MinBlockSize) * MinBlockSize == maxBlockSize); // Evenly dividable  
+	check(0 == ((maxBlockSize / MinBlockSize) & ((maxBlockSize / MinBlockSize) - 1))); // Power of two  
+
+	MaxOrder = UnitSizeToOrder(SizeToUnitSize(maxBlockSize));
+
+	Reset();
+}
+
+void FD3D12BuddyAllocator::Initialize()
+{
+	if (AllocationStrategy == eBuddyAllocationStrategy::kPlacedResourceStrategy)
+	{
+		D3D12_HEAP_PROPERTIES heapProps = CD3DX12_HEAP_PROPERTIES(HeapType);
+
+		D3D12_HEAP_DESC desc = {};
+		desc.SizeInBytes = MaxBlockSize;
+		desc.Properties = heapProps;
+		desc.Alignment = 0;
+		desc.Flags = HeapFlags;
+
+		VERIFYD3D11RESULT(GetParentDevice()->GetDevice()->CreateHeap(&desc, IID_PPV_ARGS(BackingHeap.GetInitReference())));
+		SetName(BackingHeap, L"Placed Resource Allocator Backing Heap");
+	}
+	else
+	{
+		VERIFYD3D11RESULT(GetParentDevice()->GetResourceHelper().CreateBuffer(HeapType, MaxBlockSize, BackingResource.GetInitReference(), ResourceFlags));
+		SetName(BackingResource, L"Resource Allocator Underlying Buffer");
+
+		if (IsCPUWritable(HeapType))
+		{
+			BackingResource->GetResource()->Map(0, nullptr, &BaseAddress);
+		}
+	}
+}
+
+void FD3D12BuddyAllocator::Destroy()
+{
+	ReleaseAllResources();
+}
+
+uint32 FD3D12BuddyAllocator::AllocateBlock(uint32 order)
+{
+	uint32 offset;
+
+	if (order > MaxOrder)
+	{
+		check(false); // Can't allocate a block that large  
+	}
+
+	if (FreeBlocks[order].Num() == 0)
+	{
+		// No free nodes in the requested pool.  Try to find a higher-order block and split it.  
+		uint32 left = AllocateBlock(order + 1);
+
+		uint32 size = OrderToUnitSize(order);
+
+		uint32 right = left + size;
+
+		FreeBlocks[order].Add(right); // Add the right block to the free pool  
+
+		offset = left; // Return the left block  
+	}
+
+	else
+	{
+		TSet<uint32>::TConstIterator it(FreeBlocks[order]);
+		offset = *it;
+
+		// Remove the block from the free list
+		FreeBlocks[order].Remove(*it);
+	}
+
+	return offset;
+}
+
+void FD3D12BuddyAllocator::DeallocateBlock(uint32 offset, uint32 order)
+{
+	// See if the buddy block is free  
+	uint32 size = OrderToUnitSize(order);
+
+	uint32 buddy = GetBuddyOffset(offset, size);
+
+	uint32* it = FreeBlocks[order].Find(buddy);
+
+	if (it != nullptr)
+	{
+		// Deallocate merged blocks
+		DeallocateBlock(FMath::Min(offset, buddy), order + 1);
+		// Remove the buddy from the free list  
+		FreeBlocks[order].Remove(*it);
+	}
+	else
+	{
+		// Add the block to the free list
+		FreeBlocks[order].Add(offset);
+	}
+}
+
+FD3D12ResourceBlockInfo* FD3D12BuddyAllocator::Allocate(uint32 SizeInBytes, uint32 Alignment, const void* InitialData)
+{
+	FScopeLock Lock(&CS);
+
+	if (Initialized == false)
+	{
+		Initialize();
+		Initialized = true;
+	}
+
+	// Work out what size block is needed and allocate one
+	uint32 unitSize = SizeToUnitSize(SizeInBytes);
+	uint32 order = UnitSizeToOrder(unitSize);
+	uint32 offset = AllocateBlock(order); // This is the offset in MinBlockSize units
+
+	FD3D12ResourceBlockInfo::AllocationTrackingValues alloc = {};
+	alloc.Size = uint32(OrderToUnitSize(order) * MinBlockSize);
+	alloc.AllocationBlockOffset = uint32(offset * MinBlockSize);
+	alloc.UnpaddedSize = SizeInBytes;
+
+	check(Alignment == 0 || MinBlockSize % Alignment == 0);
+
+	INCREASE_ALLOC_COUNTER(SpaceUsed, alloc.Size);
+	INCREASE_ALLOC_COUNTER(InternalFragmentation, (alloc.Size - SizeInBytes));
+
+	TotalSizeUsed += alloc.Size;
+
+#if defined(UE_BUILD_DEBUG)
+	if (SpaceUsed > PeakUsage)
+	{
+		PeakUsage = SpaceUsed;
+	}
+#endif
+
+	FD3D12ResourceBlockInfo* Block = new FD3D12ResourceBlockInfo(this, alloc.AllocationBlockOffset, alloc);
+
+	uint64 InitialDataOffset = 0;
+	if (AllocationStrategy == eBuddyAllocationStrategy::kManualSubAllocationStrategy)
+	{
+		Block->ResourceHeap = BackingResource;
+		InitialDataOffset = alloc.AllocationBlockOffset;
+
+		if (IsCPUWritable(HeapType))
+		{
+			Block->Address = (byte*)BaseAddress + alloc.AllocationBlockOffset;
+		}
+	}
+
+	if (InitialData)
+	{
+		InitializeDefaultBuffer(Block->ResourceHeap, InitialDataOffset, InitialData, SizeInBytes);
+	}
+
+	return Block;
+}
+
+FD3D12ResourceBlockInfo * FD3D12BuddyAllocator::TryAllocate(uint32 SizeInBytes, uint32 Alignment)
+{
+	FScopeLock Lock(&CS);
+
+	if (CanAllocate(SizeInBytes))
+	{
+		return Allocate(SizeInBytes, Alignment);
+	}
+	else
+	{
+		INCREASE_ALLOC_COUNTER(FailedAllocationSpace, SizeInBytes);
+		return nullptr;
+	}
+}
+
+void FD3D12BuddyAllocator::Deallocate(FD3D12ResourceBlockInfo* Block)
+{
+	FScopeLock Lock(&CS);
+
+	Block->FrameFence = GetParentDevice()->GetCommandListManager().GetFence(EFenceType::FT_Frame).GetCurrentFence();
+
+	if (Block->IsPlacedResource)
+	{
+		// Blocks don't have ref counted pointers to their resources so add a ref for the deletion queue
+		// (Otherwise the resource may be deleted before the GPU is finished with it)
+		Block->ResourceHeap->AddRef();
+	}
+	DeferredDeletionQueue.Add(Block);
+	INCREASE_ALLOC_COUNTER(NumBlocksInDeferredDeletionQueue, 1);
+}
+
+void FD3D12BuddyAllocator::DeallocateInternal(FD3D12ResourceBlockInfo* Block)
+{
+	check(IsOwner(*Block));
+
+	uint32 offset = SizeToUnitSize(Block->GetOffset());
+
+	uint32 size = SizeToUnitSize(Block->GetSize());
+
+	uint32 order = UnitSizeToOrder(size);
+
+	// Blocks are cleaned up async so need a lock
+
+	DeallocateBlock(offset, order);
+
+	DECREASE_ALLOC_COUNTER(SpaceUsed, Block->GetSize());
+	DECREASE_ALLOC_COUNTER(InternalFragmentation, (Block->AllocatorValues.Size - Block->AllocatorValues.UnpaddedSize));
+
+	TotalSizeUsed -= Block->AllocatorValues.Size;
+
+	if (AllocationStrategy == eBuddyAllocationStrategy::kPlacedResourceStrategy)
+	{
+		// Release the resource
+		Block->Destroy();
+	}
+
+	delete Block;
+};
+
+void FD3D12BuddyAllocator::CleanUpAllocations()
+{
+	FScopeLock Lock(&CS);
+
+	FD3D12CommandListManager& manager = GetParentDevice()->GetCommandListManager();
+
+	uint32 popCount = 0;
+	for (int32 i = 0; i < DeferredDeletionQueue.Num(); i++)
+	{
+		FD3D12ResourceBlockInfo* Block = DeferredDeletionQueue[i];
+
+		if (manager.GetFence(FT_Frame).IsFenceComplete(Block->FrameFence))
+		{
+			DeallocateInternal(Block);
+			DECREASE_ALLOC_COUNTER(NumBlocksInDeferredDeletionQueue, 1);
+			popCount = i + 1;
+		}
+		else
+		{
+			break;
+		}
+	}
+
+	if (popCount)
+	{
+		// clear out all of the released blocks, don't allow the array to shrink
+		DeferredDeletionQueue.RemoveAt(0, popCount, false);
+	}
+}
+
+void FD3D12BuddyAllocator::ReleaseAllResources()
+{
+	FD3D12ResourceBlockInfo* Block = nullptr;
+
+	for (FD3D12ResourceBlockInfo*& Block : DeferredDeletionQueue)
+	{
+		DeallocateInternal(Block);
+		DECREASE_ALLOC_COUNTER(NumBlocksInDeferredDeletionQueue, 1);
+	}
+
+	DeferredDeletionQueue.Empty();
+
+	if (BackingResource)
+	{
+		check(BackingResource->GetRefCount() == 1);
+		BackingResource = nullptr;
+	}
+
+	if (BackingHeap)
+	{
+		BackingHeap = nullptr;
+	}
+}
+
+void FD3D12BuddyAllocator::DumpAllocatorStats(class FOutputDevice& Ar)
+{
+#if defined(UE_BUILD_DEBUG)
+	FBufferedOutputDevice BufferedOutput;
+	{
+		// This is the memory tracked inside individual allocation pools
+		FD3D12DynamicRHI* D3DRHI = FD3D12DynamicRHI::GetD3DRHI();
+		FName categoryName(&DebugName.GetCharArray()[0]);
+
+		BufferedOutput.CategorizedLogf(categoryName, ELogVerbosity::Log, TEXT(""));
+		BufferedOutput.CategorizedLogf(categoryName, ELogVerbosity::Log, TEXT("Heap Size | MinBlock Size | Space Used | Peak Usage | Unpooled Allocations | Internal Fragmentation | Blocks in Deferred Delete Queue "));
+		BufferedOutput.CategorizedLogf(categoryName, ELogVerbosity::Log, TEXT("----------"));
+
+		BufferedOutput.CategorizedLogf(categoryName, ELogVerbosity::Log, TEXT("% 10i % 10i % 16i % 12i % 13i % 8i % 10I"),
+			MaxBlockSize,
+			MinBlockSize,
+			SpaceUsed,
+			PeakUsage,
+			FailedAllocationSpace,
+			InternalFragmentation,
+			NumBlocksInDeferredDeletionQueue);
+	}
+
+	BufferedOutput.RedirectTo(Ar);
+#endif
+}
+
+bool FD3D12BuddyAllocator::CanAllocate(uint32 size)
+{
+	if (TotalSizeUsed == MaxBlockSize)
+	{
+		return false;
+	}
+
+	uint32 blockSize = MaxBlockSize;
+
+	for (int32 i = FreeBlocks.Num() - 1; i >= 0; i--)
+	{
+		if (FreeBlocks[i].Num() && blockSize >= size)
+		{
+			return true;
+		}
+
+		// Halve the block size;
+		blockSize = blockSize >> 1;
+
+		if (blockSize < size) return false;
+	}
+	return false;
+}
+
+void FD3D12BuddyAllocator::Reset()
+{
+	// Clear the free blocks collection
+	FreeBlocks.Empty();
+
+	// Initialize the pool with a free inner block of max inner block size
+	FreeBlocks.SetNum(MaxOrder + 1);
+	FreeBlocks[MaxOrder].Add((uint32)0);
+}
+
+//-----------------------------------------------------------------------------
+//	Multi-Buddy Allocator
+//-----------------------------------------------------------------------------
+
+FD3D12MultiBuddyAllocator::FD3D12MultiBuddyAllocator(FD3D12Device* ParentDevice,
+	FString Name,
+	eBuddyAllocationStrategy allocationStrategyIn,
+	D3D12_HEAP_TYPE heapType,
+	D3D12_HEAP_FLAGS heapFlagsIn,
+	D3D12_RESOURCE_FLAGS flags,
+	uint32 MaxSizeForPooling,
+	uint32 maxBlockSizeIn,
+	uint32 minBlockSizeIn) :
+	AllocationStrategy(allocationStrategyIn)
+	, HeapFlags(heapFlagsIn)
+	, MaxBlockSize(maxBlockSizeIn)
+	, MinBlockSize(minBlockSizeIn)
+	, FD3D12ResourceAllocator(ParentDevice, Name, heapType, flags, MaxSizeForPooling)
+{}
+
+
+FD3D12ResourceBlockInfo* FD3D12MultiBuddyAllocator::TryAllocate(uint32 SizeInBytes, uint32 Alignment)
+{
+	FScopeLock Lock(&CS);
+
+	for (int32 i = 0; i < Allocators.Num(); i++)
+	{
+		FD3D12ResourceBlockInfo* Block = Allocators[i]->TryAllocate(SizeInBytes, Alignment);
+		if (Block)
+		{
+			return Block;
+		}
+	}
+
+	Allocators.Add(CreateNewAllocator());
+	return Allocators.Last()->TryAllocate(SizeInBytes, Alignment);
+}
+
+void FD3D12MultiBuddyAllocator::Deallocate(FD3D12ResourceBlockInfo* Block)
+{
+	//The sub-allocators should handle the deallocation
+	check(false);
+}
+
+FD3D12BuddyAllocator* FD3D12MultiBuddyAllocator::CreateNewAllocator()
+{
+	return new FD3D12BuddyAllocator(GetParentDevice(),
+		DebugName,
+		AllocationStrategy,
+		HeapType,
+		HeapFlags,
+		ResourceFlags,
+		MaximumAllocationSizeForPooling,
+		MaxBlockSize,
+		MinBlockSize);
+}
+
+void FD3D12MultiBuddyAllocator::Initialize()
+{
+	Allocators.Add(CreateNewAllocator());
+}
+
+void FD3D12MultiBuddyAllocator::Destroy()
+{
+	ReleaseAllResources();
+}
+
+void FD3D12MultiBuddyAllocator::CleanUpAllocations()
+{
+	FScopeLock Lock(&CS);
+
+	for (auto*& Allocator : Allocators)
+	{
+		Allocator->CleanUpAllocations();
+	}
+
+	// Trim empty allocators
+	for (int32 i = (Allocators.Num() - 1); i >= 0; i--)
+	{
+		if (Allocators[i]->IsEmpty())
+		{
+			Allocators[i]->Destroy();
+			delete(Allocators[i]);
+			Allocators.RemoveAt(i);
+		}
+	}
+}
+
+void FD3D12MultiBuddyAllocator::DumpAllocatorStats(class FOutputDevice& Ar)
+{
+	//TODO
+}
+
+void FD3D12MultiBuddyAllocator::ReleaseAllResources()
+{
+	for (int32 i = (Allocators.Num() - 1); i >= 0; i--)
+	{
+		Allocators[i]->CleanUpAllocations();
+		Allocators[i]->Destroy();
+		delete(Allocators[i]);
+	}
+
+	Allocators.Empty();
+}
+
+void FD3D12MultiBuddyAllocator::Reset()
+{
+
+}
+
+//-----------------------------------------------------------------------------
+//	Linear Allocator
+//-----------------------------------------------------------------------------
+
+FD3D12LinearAllocator::FD3D12LinearAllocator(FD3D12Device* ParentDevice,
+	FString Name,
+	D3D12_HEAP_TYPE heapType,
+	D3D12_RESOURCE_FLAGS flags,
+	uint32 MaxSizeForPooling,
+	uint32 PageSize) :
+	AllocationPageSize(PageSize),
+	FD3D12ResourceAllocator(ParentDevice, Name, heapType, flags, MaxSizeForPooling)
+{
+	Reset();
+}
+
+FD3D12ResourceBlockInfo* FD3D12LinearAllocator::TryAllocate(uint32 SizeInBytes, uint32 Alignment)
+{
+	FScopeLock Lock(&CS);
+
+	if (Initialized == false)
+	{
+		Initialize();
+		Initialized = true;
+	}
+
+	check(SizeInBytes <= AllocationPageSize);
+
+	FD3D12LinearAllocatorPage* CurrentPage = FindPageWithBestFit(SizeInBytes, Alignment);
+
+	uint32 AlignedOffset = Align(CurrentPage->CurrentOffset, Alignment);
+
+	FD3D12ResourceBlockInfo::AllocationTrackingValues alloc;
+	alloc.Size = SizeInBytes + (AlignedOffset - CurrentPage->CurrentOffset);
+	alloc.AllocatorPrivateData = CurrentPage;
+
+	FD3D12ResourceBlockInfo* Block = new FD3D12ResourceBlockInfo(this, AlignedOffset, alloc);
+
+	Block->ResourceHeap = CurrentPage->Resource;
+
+	if (IsCPUWritable(HeapType))
+	{
+		Block->Address = (byte*)CurrentPage->BaseAddress + AlignedOffset;
+	}
+
+	CurrentPage->CurrentOffset += alloc.Size;
+
+	return Block;
+}
+
+FD3D12LinearAllocatorPage* FD3D12LinearAllocator::FindPageWithBestFit(uint32 AllocationSize, uint32 Alignment)
+{
+	// Search for a space that is a good fit
+	// TODO: This could potentially be super slow! (O(n^2)), keeping the list sorted could help
+	for (int32 i = 0; i < ActivePages.Num(); i++)
+	{
+		uint32 AlignedOffset = Align(ActivePages[i]->CurrentOffset, Alignment);
+
+		if (AllocationSize < (AllocationPageSize - AlignedOffset))
+		{
+			return ActivePages[i];
+		}
+	}
+
+	ActivePages.Add(AllocatePage());
+	return ActivePages.Last();
+}
+
+void FD3D12LinearAllocator::Deallocate(FD3D12ResourceBlockInfo* Block)
+{
+	FScopeLock Lock(&CS);
+
+	check(IsOwner(*Block));
+	FD3D12LinearAllocatorPage* page = (FD3D12LinearAllocatorPage*)Block->AllocatorValues.AllocatorPrivateData;
+	page->DeallocatedSpace += Block->AllocatorValues.Size;
+
+	delete(Block);
+}
+
+void FD3D12LinearAllocator::Initialize()
+{
+	ActivePages.Reserve(32);
+
+	ActivePages.Add(AllocatePage());
+}
+
+void FD3D12LinearAllocator::Destroy()
+{
+	for (FD3D12LinearAllocatorPage*& page : ActivePages)
+	{
+		DestroyPage(*page);
+		delete(page);
+	}
+
+	ActivePages.Empty();
+}
+
+void FD3D12LinearAllocator::CleanUpAllocations()
+{
+	FScopeLock Lock(&CS);
+
+	for (int32 i = (ActivePages.Num() - 1); i >= 0; i--)
+	{
+		// All the allcocations that came out of this page are emptied so it can be release;
+		if (ActivePages[i]->DeallocatedSpace == AllocationPageSize)
+		{
+			DestroyPage(*ActivePages[i]);
+			delete(ActivePages[i]);
+			ActivePages.RemoveAt(i);
+		}
+	}
+}
+
+void FD3D12LinearAllocator::DumpAllocatorStats(class FOutputDevice& Ar)
+{
+	//TODO:
+}
+
+void FD3D12LinearAllocator::ReleaseAllResources()
+{
+	for (FD3D12LinearAllocatorPage*& page : ActivePages)
+	{
+		DestroyPage(*page);
+		delete(page);
+	}
+
+	ActivePages.Empty();
+}
+
+void FD3D12LinearAllocator::Reset()
+{
+	ReleaseAllResources();
+}
+
+FD3D12LinearAllocatorPage* FD3D12LinearAllocator::AllocatePage()
+{
+	FD3D12LinearAllocatorPage* newPage = new FD3D12LinearAllocatorPage();
+
+	VERIFYD3D11RESULT(GetParentDevice()->GetResourceHelper().CreateBuffer(HeapType, AllocationPageSize, &newPage->Resource, ResourceFlags));
+	SetName(newPage->Resource, L"Linear Allocator Page");
+
+	if (IsCPUWritable(HeapType))
+	{
+		VERIFYD3D11RESULT(newPage->Resource->GetResource()->Map(0, nullptr, &newPage->BaseAddress));
+	}
+	return newPage;
+}
+
+void FD3D12LinearAllocator::DestroyPage(FD3D12LinearAllocatorPage& Page)
+{
+	check(Page.Resource);
+	delete(Page.Resource);
+	Page.Resource = nullptr;
+}
+
+//-----------------------------------------------------------------------------
+//	Bucket Allocator
+//-----------------------------------------------------------------------------
+FD3D12BucketAllocator::FD3D12BucketAllocator(FD3D12Device* ParentDevice,
+	FString Name,
+	D3D12_HEAP_TYPE heapType,
+	D3D12_RESOURCE_FLAGS flags,
+	uint64 BlockRetentionFrameCountIn) : 
+	BlockRetentionFrameCount(BlockRetentionFrameCountIn),
+	FD3D12ResourceAllocator(ParentDevice, Name, heapType, flags, 32 * 1024 * 1024)
+{}
+
+FD3D12ResourceBlockInfo* FD3D12BucketAllocator::TryAllocate(uint32 SizeInBytes, uint32 Alignment)
+{
+	FScopeLock Lock(&CS);
+
+	// Size cannot be smaller than the requested alignment
+	SizeInBytes = FMath::Max(SizeInBytes, Alignment);
+
+	uint32 bucket = BucketFromSize(SizeInBytes, BucketShift);
+	check(bucket < NumBuckets);
+	FD3D12ResourceBlockInfo* pBlock = nullptr;
+
+	// See if a block is already available in the bucket
+	if (AvailableBlocks[bucket].Dequeue(pBlock))
+	{
+		check(pBlock);
+		if (IsCPUWritable(HeapType))
+		{
+			check(pBlock->Address);
+		}
+		check(pBlock->ResourceHeap);
+	}
+	else
+	{
+		// No blocks of the requested size are available so make one
+		FD3D12Resource* Resource = nullptr;
+		void *BaseAddress = nullptr;
+
+		// Allocate a block
+		uint32 BlockSize = BlockSizeFromBufferSize(SizeInBytes, BucketShift);
+
+		check(BlockSize >= SizeInBytes);
+
+		if (FAILED(GetParentDevice()->GetResourceHelper().CreateBuffer(HeapType, BlockSize < MIN_HEAP_SIZE ? MIN_HEAP_SIZE : SizeInBytes, &Resource, ResourceFlags)))
+		{
+			return nullptr;
+		}
+
+		if (IsCPUWritable(HeapType))
+		{
+			VERIFYD3D11RESULT(Resource->GetResource()->Map(0, nullptr, &BaseAddress));
+			check(BaseAddress == (uint8*)(((uint64)BaseAddress + Alignment - 1) & ~((uint64)Alignment - 1)));
+		}
+
+		FD3D12ResourceBlockInfo::AllocationTrackingValues alloc = {};
+		alloc.AllocatorPrivateData = (void*)size_t(bucket); // Stash the bucket index in the private data
+
+		pBlock = new FD3D12ResourceBlockInfo(this, 0, alloc);
+		pBlock->Address = BaseAddress;
+		pBlock->ResourceHeap = Resource;
+
+		// If bucket blocks are less than 65536 bytes then use FDynamicResourcePool
+		if (BlockSize < MIN_HEAP_SIZE)
+		{
+			// Create additional available blocks that can be sub-allocated from the same resource
+			for (uint32 Offset = BlockSize; Offset <= MIN_HEAP_SIZE - BlockSize; Offset += BlockSize)
+			{
+				void *Start = (void *)((uint8 *)BaseAddress + Offset);
+				FD3D12ResourceBlockInfo *NewBlock = new FD3D12ResourceBlockInfo(this, Offset, alloc);
+				NewBlock->Address = Start;
+				NewBlock->ResourceHeap = Resource;
+				NewBlock->ResourceHeap->AddRef();
+
+				// Add the bucket to the available list
+				AvailableBlocks[bucket].Enqueue(NewBlock);
+			}
+		}
+	}
+
+	// Check the returned block matches the requested alignment
+	check(Alignment == 0 || ((size_t)(pBlock->Address) % Alignment) == 0);
+
+	return pBlock;
+}
+
+void FD3D12BucketAllocator::Deallocate(FD3D12ResourceBlockInfo* Block)
+{
+	FScopeLock Lock(&CS);
+
+	FD3D12CommandListManager& CommandListManager = GetParentDevice()->GetCommandListManager();
+	Block->FrameFence = CommandListManager.GetFence(FT_Frame).GetCurrentFence();
+
+	ExpiredBlocks.Enqueue(Block);
+}
+
+void FD3D12BucketAllocator::Initialize()
+{
+
+}
+
+void FD3D12BucketAllocator::Destroy()
+{
+	ReleaseAllResources();
+}
+void FD3D12BucketAllocator::CleanUpAllocations()
+{
+	FScopeLock Lock(&CS);
+
+	struct FDequeueAlloc
+	{
+		FDequeueAlloc(FD3D12CommandListManager* InCommandListManager, const uint64 InRetentionFrameCount)
+			: CommandListManager(InCommandListManager)
+			, RetentionFrameCount(InRetentionFrameCount)
+		{
+		}
+
+		bool operator() (FD3D12ResourceBlockInfo* pBlockInfo)
+		{
+			return CommandListManager->GetFence(FT_Frame).IsFenceComplete(pBlockInfo->FrameFence + RetentionFrameCount);
+		}
+
+		FD3D12CommandListManager* CommandListManager;
+		const uint64              RetentionFrameCount;
+	};
+
+	FD3D12CommandListManager& CommandListManager = GetParentDevice()->GetCommandListManager();
+
+#if SUB_ALLOCATED_DEFAULT_ALLOCATIONS
+	const static uint32 MinCleanupBucket = FMath::Max<uint32>(0, BucketFromSize(MIN_HEAP_SIZE, BucketShift));
+#else
+	const static uint32 MinCleanupBucket = 0;
+#endif
+
+	// Start at bucket 8 since smaller buckets are sub-allocated resources
+	// and would be fragmented by deleting blocks
+	for (uint32 bucket = MinCleanupBucket; bucket < NumBuckets; bucket++)
+	{
+		FDequeueAlloc DequeueAlloc(&CommandListManager, BlockRetentionFrameCount);
+		FD3D12ResourceBlockInfo* pBlockInfo = nullptr;
+
+		while (AvailableBlocks[bucket].Dequeue(pBlockInfo, DequeueAlloc))
+		{
+			pBlockInfo->ResourceHeap->Release();
+			delete(pBlockInfo);
+		}
+	}
+
+	FD3D12ResourceBlockInfo* pBlockInfo = nullptr;
+	FDequeueAlloc DequeueAlloc(&CommandListManager, 0);
+
+	while (ExpiredBlocks.Dequeue(pBlockInfo, DequeueAlloc))
+	{
+		check(pBlockInfo);
+		// Get the bucket info from the allocator private data
+		uint32 bucket = uint32((size_t)pBlockInfo->AllocatorValues.AllocatorPrivateData);
+		// Add the bucket to the available list
+		AvailableBlocks[bucket].Enqueue(pBlockInfo);
+	}
+}
+
+void FD3D12BucketAllocator::DumpAllocatorStats(class FOutputDevice& Ar)
+{
+	//TODO:
+}
+void FD3D12BucketAllocator::ReleaseAllResources()
+{
+	// Start at bucket 8 since smaller buckets are sub-allocated resources
+	// and would be fragmented by deleting blocks
+	for (uint32 bucket = 0; bucket < NumBuckets; bucket++)
+	{
+		FD3D12ResourceBlockInfo* pBlockInfo = nullptr;
+		while (AvailableBlocks[bucket].Dequeue(pBlockInfo))
+		{
+			pBlockInfo->ResourceHeap->Release();
+			delete(pBlockInfo);
+		}
+	}
+
+	FD3D12ResourceBlockInfo* pBlockInfo = nullptr;
+
+	while (ExpiredBlocks.Dequeue(pBlockInfo))
+	{
+		// Get the bucket info from the allocator private data
+		uint32 bucket = *(uint32*)pBlockInfo->AllocatorValues.AllocatorPrivateData;
+
+		pBlockInfo->ResourceHeap->Release();
+		delete(pBlockInfo);
+	}
+
+	for (FD3D12Resource*& Resource : SubAllocatedResources)
+	{
+		Resource->Release();
+		delete(Resource);
+	}
+}
+
+void FD3D12BucketAllocator::Reset()
+{
+
+}
+
+//-----------------------------------------------------------------------------
+//	Dynamic Buffer Allocator
+//-----------------------------------------------------------------------------
+
+FD3D12DynamicHeapAllocator::FD3D12DynamicHeapAllocator(FD3D12Device* InParent, FString Name, eBuddyAllocationStrategy allocationStrategy,
+	uint32 MaxSizeForPooling,
+	uint32 maxBlockSize,
+	uint32 minBlockSize)
+	: CurrentCommandContext(nullptr)
+	, Allocator(nullptr)
+	, FD3D12DeviceChild(InParent)
+{
+	const uint32 pageSize = 1024 * 1024 * 4;// 4MB
+
+	Allocator = new FD3D12MultiBuddyAllocator(InParent, Name,
+		allocationStrategy,
+		D3D12_HEAP_TYPE_UPLOAD,
+		D3D12_HEAP_FLAG_ALLOW_ONLY_BUFFERS,
+		D3D12_RESOURCE_FLAG_NONE,
+		MaxSizeForPooling,
+		pageSize,
+		minBlockSize);
+}
+
+void* FD3D12DynamicHeapAllocator::AllocUploadResource(uint32 size, uint32 alignment, FD3D12ResourceLocation* ResourceLocation)
+{
+	ResourceLocation->Clear();
+
+	//TODO: For some reason 0 sized buffers are being created and then expected to have a resource
+	if (size == 0)
+	{
+		size = 16;
+	}
+
+	//Work loads like infiltrator create enourmous amounts of buffer space in setup
+	//clean up as we go as it can even run out of memory before the first frame.
+	if (GetParentDevice()->FirstFrameSeen == false)
+	{
+		GetParentDevice()->GetDeferredDeletionQueue().ReleaseResources(true);
+		Allocator->CleanUpAllocations();
+	}
+
+	if (size <= Allocator->MaximumAllocationSizeForPooling)
+	{
+		FD3D12ResourceBlockInfo* Block = Allocator->TryAllocate(size, alignment);
+
+		if (Block)
+		{
+			FD3D12CommandListManager& CommandListManager = GetParentDevice()->GetCommandListManager();
+			Block->FrameFence = CommandListManager.GetFence(FT_Frame).GetCurrentFence();
+			check(Block->FrameFence);
+
+			FD3D12StateCache* pStateCache = CurrentCommandContext ? &CurrentCommandContext->StateCache : nullptr;
+			ResourceLocation->SetBlockInfo(Block, pStateCache);
+			ResourceLocation->EffectiveBufferSize = size;
+
+			check(Block->Address);
+			check(Block->ResourceHeap);
+
+			return Block->Address;
+		}
+	}
+
+	//Allocate Standalone
+
+	// Allocations are 64k aligned
+	if (alignment)
+	{
+		alignment = (D3D_BUFFER_ALIGNMENT % alignment) == 0 ? 0 : alignment;
+	}
+
+	const uint64 BufferSize = size + alignment;
+
+	VERIFYD3D11RESULT(GetParentDevice()->GetResourceHelper().CreateBuffer(D3D12_HEAP_TYPE_UPLOAD, BufferSize, ResourceLocation->Resource.GetInitReference()));
+	SetName(ResourceLocation->Resource, L"Stand Alone Upload Buffer");
+
+	void* Data = nullptr;
+	VERIFYD3D11RESULT(ResourceLocation->Resource->GetResource()->Map(0, nullptr, &Data));
+
+	ResourceLocation->EffectiveBufferSize = size;
+	ResourceLocation->Offset = alignment;
+
+	return Data;
+}
+
+void FD3D12DynamicHeapAllocator::CleanUpAllocations()
+{
+	Allocator->CleanUpAllocations();
+}
+
+void FD3D12DynamicHeapAllocator::Destroy()
+{
+	Allocator->Destroy();
+	delete(Allocator);
+	Allocator = nullptr;
+}
+
+//-----------------------------------------------------------------------------
+//	Default Buffer Allocator
+//-----------------------------------------------------------------------------
+
+FD3D12DefaultBufferPool::FD3D12DefaultBufferPool(FD3D12Device* InParent, FD3D12ResourceAllocator* AllocatorIn) :
+	Allocator(AllocatorIn),
+	FD3D12DeviceChild(InParent)
+{
+}
+
+void FD3D12DefaultBufferPool::CleanUpAllocations()
+{
+	Allocator->CleanUpAllocations();
+}
+
+
+// Grab a buffer from the available buffers or create a new buffer if none are available
+
+void FD3D12DefaultBufferPool::AllocDefaultResource(const D3D12_RESOURCE_DESC& Desc, D3D12_SUBRESOURCE_DATA* InitialData, FD3D12ResourceLocation* ResourceLocation, uint32 Alignment)
+{
+	// If the resource location owns a block, this will deallocate it.
+	ResourceLocation->Clear();
+
+	if (Desc.Width == 0) return;
+
+	const void* Data = (InitialData) ? InitialData->pData : nullptr;
+
+	const bool PoolResource = Desc.Width < Allocator->MaximumAllocationSizeForPooling &&
+		((Desc.Width % (1024 * 64)) != 0);
+
+	if (PoolResource)
+	{
+		// Ensure we're allocating from the correct pool
+		check(Desc.Flags == Allocator->ResourceFlags);
+
+		FD3D12ResourceBlockInfo* Block = Allocator->TryAllocate(Desc.Width, Alignment);
+
+		if (Block)
+		{
+			ResourceLocation->SetBlockInfo(Block, &GetParentDevice()->GetDefaultCommandContext().StateCache);
+
+			ResourceLocation->EffectiveBufferSize = Desc.Width;
+
+			if (Alignment > 0)
+			{
+				uint32 AlignmentMismatch = ResourceLocation->GetOffset() % Alignment;
+
+				if (AlignmentMismatch != 0)
+				{
+					ResourceLocation->SetPadding(Alignment - AlignmentMismatch);
+
+					check(ResourceLocation->GetOffset() % Alignment == 0);
+					check(ResourceLocation->GetOffset() % 4 == 0);
+				}
+			}
+
+			if (InitialData)
+			{
+				Allocator->InitializeDefaultBuffer(Block->ResourceHeap, Block->Offset, InitialData->pData, Desc.Width);
+			}
+
+			return;
+		}
+	}
+
+	//Allocate Standalone
+
+	//Allocations are 64k aligned
+	if (Alignment)
+	{
+		Alignment = (D3D_BUFFER_ALIGNMENT % Alignment) == 0 ? 0 : Alignment;
+	}
+
+	const uint64 Size = (Desc.Width + Alignment);
+	VERIFYD3D11RESULT(GetParentDevice()->GetResourceHelper().CreateBuffer(D3D12_HEAP_TYPE_DEFAULT, Size, ResourceLocation->Resource.GetInitReference(), Allocator->ResourceFlags));
+	SetName(ResourceLocation->Resource, L"Stand Alone Default Buffer");
+
+	ResourceLocation->EffectiveBufferSize = Desc.Width;
+	ResourceLocation->Offset = Alignment;
+
+	if (Data)
+	{
+		Allocator->InitializeDefaultBuffer(ResourceLocation->Resource.GetReference(), ResourceLocation->Offset, Data, InitialData->RowPitch);
+	}
+}
+
+// Grab a buffer from the available buffers or create a new buffer if none are available
+HRESULT FD3D12DefaultBufferAllocator::AllocDefaultResource(const D3D12_RESOURCE_DESC& Desc, D3D12_SUBRESOURCE_DATA* InitialData, FD3D12ResourceLocation* ResourceLocation, uint32 Alignment)
+{
+	//NOTE: Indexing based on the resource flags looks weird but is necessary e.g. the flags dictate if the resource
+	//      can be used as a UAV. So each type of buffer has to come from a separate pool.
+
+	check((uint32)Desc.Flags < MAX_DEFAULT_POOLS);
+	if (DefaultBufferPools[Desc.Flags] == nullptr)
+	{
+		FD3D12ResourceAllocator* Allocator = nullptr;
+
+		// D3D12_RESOURCE_FLAG_DENY_SHADER_RESOURCE are used much more frequently
+		if (Desc.Flags == D3D12_RESOURCE_FLAG_DENY_SHADER_RESOURCE)
+		{
+			Allocator = new FD3D12MultiBuddyAllocator(GetParentDevice(),
+				FString(L"Default Buffer Multi Buddy Allocator"),
+				kManualSubAllocationStrategy,
+				D3D12_HEAP_TYPE_DEFAULT,
+				D3D12_HEAP_FLAG_ALLOW_ONLY_BUFFERS,
+				Desc.Flags,
+				DEFAULT_BUFFER_POOL_MAX_ALLOC_SIZE,
+				DEFAULT_BUFFER_POOL_SIZE,
+				16);
+			//Allocator = new FD3D12BucketAllocator(GetParentDevice(), FString(L"Default Buffer Bucket Allocator"), D3D12_HEAP_TYPE_DEFAULT, Desc.Flags, 10);
+		}
+		else
+		{
+			Allocator = new FD3D12BuddyAllocator(GetParentDevice(),
+					FString(L"Default Buffer Buddy Allocator"),
+					kManualSubAllocationStrategy,
+					D3D12_HEAP_TYPE_DEFAULT,
+					D3D12_HEAP_FLAG_ALLOW_ONLY_BUFFERS,
+					Desc.Flags,
+					64 * 1024,
+					4 * 1024 * 1024,
+					32);
+		}
+
+		DefaultBufferPools[Desc.Flags] = new FD3D12DefaultBufferPool(GetParentDevice(), Allocator);
+	}
+
+	DefaultBufferPools[Desc.Flags]->AllocDefaultResource(Desc, InitialData, ResourceLocation, Alignment);
+
+	return S_OK;
+}
+
+void FD3D12DefaultBufferAllocator::FreeDefaultBufferPools()
+{
+	for (uint32 i = 0; i < MAX_DEFAULT_POOLS; ++i)
+	{
+		if (DefaultBufferPools[i])
+		{
+			DefaultBufferPools[i]->CleanUpAllocations();
+		}
+		delete DefaultBufferPools[i];
+		DefaultBufferPools[i] = nullptr;
+	}
+}
+
+void FD3D12DefaultBufferAllocator::CleanupFreeBlocks()
+{
+	for (uint32 i = 0; i < MAX_DEFAULT_POOLS; ++i)
+	{
+		if (DefaultBufferPools[i])
+		{
+			DefaultBufferPools[i]->CleanUpAllocations();
+		}
+	}
+}
+
+//-----------------------------------------------------------------------------
+//	Texture Allocator
+//-----------------------------------------------------------------------------
+
+FD3D12TextureAllocator::FD3D12TextureAllocator(FD3D12Device* Device, FString Name, uint32 HeapSize, D3D12_HEAP_FLAGS Flags) :
+	FD3D12BuddyAllocator(Device, Name,
+		kPlacedResourceStrategy,
+		D3D12_HEAP_TYPE_DEFAULT,
+		Flags | D3D12_HEAP_FLAG_DENY_BUFFERS,
+		D3D12_RESOURCE_FLAG_NONE,
+		D3D12_DEFAULT_RESOURCE_PLACEMENT_ALIGNMENT,
+		HeapSize,
+		D3D12_SMALL_RESOURCE_PLACEMENT_ALIGNMENT)
+{
+	// Inform the texture streaming system of this heap so that it correctly accounts for placed textures
+	FD3D12DynamicRHI::GetD3DRHI()->UpdataTextureMemorySize((MaxBlockSize / 1024));
+}
+
+FD3D12TextureAllocator::~FD3D12TextureAllocator()
+{
+	FD3D12DynamicRHI::GetD3DRHI()->UpdataTextureMemorySize(-int32(MaxBlockSize / 1024));
+}
+
+HRESULT FD3D12TextureAllocator::AllocateTexture(D3D12_RESOURCE_DESC Desc, const D3D12_CLEAR_VALUE* ClearValue, FD3D12ResourceLocation* TextureLocation)
+{
+	FD3D12Device* Device = GetParentDevice();
+	FD3D12ResourceHelper& ResourceHelper = Device->GetResourceHelper();
+	TRefCountPtr<FD3D12Resource> Resource;
+	HRESULT hr = S_OK;
+
+	TextureLocation->Clear();
+
+	D3D12_RESOURCE_ALLOCATION_INFO Info = Device->GetDevice()->GetResourceAllocationInfo(0, 1, &Desc);
+
+	if (Info.SizeInBytes < D3D12_DEFAULT_RESOURCE_PLACEMENT_ALIGNMENT)
+	{
+		FD3D12ResourceBlockInfo* Block = TryAllocate(Info.SizeInBytes, Info.Alignment);
+		if (Block)
+		{
+			hr = FD3D12DynamicRHI::CreatePlacedResource(Desc, BackingHeap, Block->Offset, D3D12_RESOURCE_STATE_COMMON, ClearValue, Resource.GetInitReference());
+
+			Block->ResourceHeap = Resource.GetReference();
+			Block->IsPlacedResource = true;
+			TextureLocation->SetBlockInfo(Block, &Device->GetDefaultCommandContext().StateCache);
+
+			return hr;
+		}
+	}
+
+	// Request default alignment for stand alone textures
+	Desc.Alignment = 0;
+	hr = ResourceHelper.CreateDefaultResource(Desc, ClearValue, Resource.GetInitReference());
+
+	TextureLocation->SetFromD3DResource(Resource, 0, 0);
+
+	return hr;
+}
+
+HRESULT FD3D12TextureAllocatorPool::AllocateTexture(D3D12_RESOURCE_DESC Desc, const D3D12_CLEAR_VALUE* ClearValue, uint8 UEFormat, FD3D12ResourceLocation* TextureLocation)
+{
+	// Multi-Sample texures have much larger alignment requirements (4MB vs 64KB)
+	check(Desc.SampleDesc.Count == 1)
+
+		// 4KB alignment is only available for read only textures
+		if ((Desc.Flags & D3D12_RESOURCE_FLAG_ALLOW_RENDER_TARGET ||
+			Desc.Flags & D3D12_RESOURCE_FLAG_ALLOW_DEPTH_STENCIL ||
+			Desc.Flags & D3D12_RESOURCE_FLAG_ALLOW_UNORDERED_ACCESS) == false)
+		{
+			// The top mip level must be less than 64k
+			if (TextureCanBe4KAligned(Desc, UEFormat))
+			{
+				Desc.Alignment = D3D12_SMALL_RESOURCE_PLACEMENT_ALIGNMENT; // request 4k alignment
+				return ReadOnlyTexturePool.AllocateTexture(Desc, ClearValue, TextureLocation);
+			}
+		}
+
+	TRefCountPtr<FD3D12Resource> Resource = nullptr;
+
+	HRESULT hr = GetParentDevice()->GetResourceHelper().CreateDefaultResource(Desc, ClearValue, Resource.GetInitReference());
+	TextureLocation->SetFromD3DResource(Resource, 0, 0);
+
+	return hr;
+}
diff --git a/Engine/Source/Runtime/Windows/D3D12RHI/Private/D3D12Allocation.h b/Engine/Source/Runtime/Windows/D3D12RHI/Private/D3D12Allocation.h
new file mode 100644
--- /dev/null
+++ b/Engine/Source/Runtime/Windows/D3D12RHI/Private/D3D12Allocation.h
@@ -0,0 +1,500 @@
+// Copyright 1998-2016 Epic Games, Inc. All Rights Reserved.
+
+/*=============================================================================
+D3D12Allocation.h: A Collection of allocators
+=============================================================================*/
+
+#pragma once
+#include "D3D12Resources.h"
+
+class FD3D12ResourceAllocator : public FD3D12DeviceChild
+{
+public:
+
+	FD3D12ResourceAllocator(FD3D12Device* ParentDevice,
+		FString Name,
+		D3D12_HEAP_TYPE heapType,
+		D3D12_RESOURCE_FLAGS flags,
+		uint32 MaxSizeForPooling);
+
+	~FD3D12ResourceAllocator();
+
+	virtual FD3D12ResourceBlockInfo* TryAllocate(uint32 SizeInBytes, uint32 Alignment) = 0;
+
+	virtual void Deallocate(FD3D12ResourceBlockInfo* Block) = 0;
+
+	virtual void Initialize() = 0;
+
+	virtual void Destroy() = 0;
+
+	virtual void Reset() = 0;
+
+	virtual void CleanUpAllocations() = 0;
+
+	virtual void DumpAllocatorStats(class FOutputDevice& Ar) = 0;
+
+	virtual void ReleaseAllResources() = 0;
+
+	inline bool IsOwner(const FD3D12ResourceBlockInfo &block)
+	{
+		return block.Allocator == this;
+	}
+
+	// Any allocation larger than this just gets straight up allocated (i.e. not pooled).
+	// These large allocations should be infrequent so the CPU overhead should be minimal
+	const uint32 MaximumAllocationSizeForPooling;
+	D3D12_RESOURCE_FLAGS ResourceFlags;
+
+	void InitializeDefaultBuffer(FD3D12Resource* Destination, uint64 DestinationOffset, const void* Data, uint64 DataSize);
+
+protected:
+
+	const FString DebugName;
+
+	bool Initialized;
+
+	const D3D12_HEAP_TYPE HeapType;
+
+	FCriticalSection CS;
+
+#if defined(UE_BUILD_DEBUG)
+	uint32 SpaceUsed;
+	uint32 InternalFragmentation;
+	uint32 NumBlocksInDeferredDeletionQueue;
+	uint32 PeakUsage;
+	uint32 FailedAllocationSpace;
+#endif
+};
+
+//-----------------------------------------------------------------------------
+//	Buddy Allocator
+//-----------------------------------------------------------------------------
+// Allocates blocks from a fixed range using buddy allocation method.
+// Buddy allocation allows reasonably fast allocation of arbitrary size blocks
+// with minimal fragmentation and provides efficient reuse of freed ranges.
+// When a block is de-allocated an attempt is made to merge it with it's 
+// neighbour (buddy) if it is contiguous and free.
+// Based on reference implementation by MSFT: billkris
+
+// Unfortunately the api restricts the minimum size of a placed buffer resource to 64k
+#define MIN_PLACED_BUFFER_SIZE (64 * 1024)
+#define D3D_BUFFER_ALIGNMENT (64 * 1024)
+
+#if defined(UE_BUILD_DEBUG)
+#define INCREASE_ALLOC_COUNTER(A, B) (A = A + B);
+#define DECREASE_ALLOC_COUNTER(A, B) (A = A - B);
+#else
+#define INCREASE_ALLOC_COUNTER(A, B)
+#define DECREASE_ALLOC_COUNTER(A, B)
+#endif
+
+enum eBuddyAllocationStrategy
+{
+	// This strategy uses Placed Resources to sub-allocate a buffer out of an underlying ID3D12Heap.
+	// The benefit of this is that each buffer can have it's own resource state and can be treated
+	// as any other buffer. The downside of this strategy is the API limitiation which enforces
+	// the minimum buffer size to 64k leading to large internal fragmentation in the allocator
+	kPlacedResourceStrategy,
+	// The alternative is to manualy sub-allocate out of a single large buffer which allows block
+	// allocation granularity down to 1 byte. However, this strategy is only really valid for buffers which
+	// will be treated as read-only after their creation (i.e. most Index and Vertex buffers). This 
+	// is because the underlying resource can only have one state at a time.
+	kManualSubAllocationStrategy
+};
+
+class FD3D12BuddyAllocator : public FD3D12ResourceAllocator
+{
+public:
+
+	FD3D12BuddyAllocator(FD3D12Device* ParentDevice,
+		FString Name,
+		eBuddyAllocationStrategy allocationStrategy,
+		D3D12_HEAP_TYPE heapType,
+		D3D12_HEAP_FLAGS heapFlags,
+		D3D12_RESOURCE_FLAGS flags,
+		uint32 MaxSizeForPooling,
+		uint32 maxBlockSize,
+		uint32 minBlockSize = MIN_PLACED_BUFFER_SIZE);
+
+	// Public Override Functions
+	virtual FD3D12ResourceBlockInfo* TryAllocate(uint32 SizeInBytes, uint32 Alignment) final override;
+
+	virtual void Deallocate(FD3D12ResourceBlockInfo* Block) final override;
+
+	virtual void Initialize() final override;
+
+	virtual void Destroy() final override;
+
+	virtual void CleanUpAllocations() final override;
+
+	virtual void DumpAllocatorStats(class FOutputDevice& Ar) final override;
+
+	virtual void ReleaseAllResources() final override;
+
+	virtual void Reset() final override;
+
+	inline bool IsEmpty()
+	{
+		return FreeBlocks[MaxOrder].Num() == 1;
+	}
+
+	inline uint32 GetTotalSizeUsed() const { return TotalSizeUsed; }
+
+protected:
+	const uint32 MaxBlockSize;
+	const uint32 MinBlockSize;
+	const D3D12_HEAP_FLAGS HeapFlags;
+	const eBuddyAllocationStrategy AllocationStrategy;
+
+	TRefCountPtr<FD3D12Resource> BackingResource;
+	TRefCountPtr<ID3D12Heap> BackingHeap;
+
+	void* BaseAddress;
+
+private:
+
+	TArray<FD3D12ResourceBlockInfo*> DeferredDeletionQueue;
+	TArray<TSet<uint32>> FreeBlocks;
+	uint32 MaxOrder;
+	uint32 TotalSizeUsed;
+
+	bool HeapFullMessageDisplayed;
+
+	inline uint32 SizeToUnitSize(uint32 size) const
+	{
+		return (size + (MinBlockSize - 1)) / MinBlockSize;
+	}
+
+	inline uint32 UnitSizeToOrder(uint32 size) const
+	{
+		return uint32(ceil(log2f(float(size))));
+	}
+
+	inline uint32 GetBuddyOffset(const uint32 &offset, const uint32 &size)
+	{
+		return offset ^ size;
+	}
+
+	uint32 OrderToUnitSize(uint32 order) const { return ((uint32)1) << order; }
+	uint32 AllocateBlock(uint32 order);
+	void DeallocateBlock(uint32 offset, uint32 order);
+
+	bool CanAllocate(uint32 size);
+
+	void DeallocateInternal(FD3D12ResourceBlockInfo* Block);
+
+	FD3D12ResourceBlockInfo* Allocate(uint32 SizeInBytes, uint32 Alignment, const void* InitialData = nullptr);
+};
+
+//-----------------------------------------------------------------------------
+//	Multi-Buddy Allocator
+//-----------------------------------------------------------------------------
+// Builds on top of the Buddy Allocator but covers some of it's deficiencies by
+// managing multiple buddy allocator instances to better match memory usage over
+// time.
+
+class FD3D12MultiBuddyAllocator : public FD3D12ResourceAllocator
+{
+public:
+
+	FD3D12MultiBuddyAllocator(FD3D12Device* ParentDevice,
+		FString Name,
+		eBuddyAllocationStrategy allocationStrategyIn,
+		D3D12_HEAP_TYPE heapType,
+		D3D12_HEAP_FLAGS heapFlagsIn,
+		D3D12_RESOURCE_FLAGS flags,
+		uint32 MaxSizeForPooling,
+		uint32 maxBlockSizeIn,
+		uint32 minBlockSizeIn = MIN_PLACED_BUFFER_SIZE);
+
+	// Public Override Functions
+	virtual FD3D12ResourceBlockInfo* TryAllocate(uint32 SizeInBytes, uint32 Alignment) final override;
+
+	virtual void Deallocate(FD3D12ResourceBlockInfo* Block) final override;
+
+	virtual void Initialize() final override;
+
+	virtual void Destroy() final override;
+
+	virtual void CleanUpAllocations() final override;
+
+	virtual void DumpAllocatorStats(class FOutputDevice& Ar) final override;
+
+	virtual void ReleaseAllResources() final override;
+
+	virtual void Reset() final override;
+
+private:
+	const eBuddyAllocationStrategy AllocationStrategy;
+	const D3D12_HEAP_FLAGS HeapFlags;
+	const uint32 MaxBlockSize;
+	const uint32 MinBlockSize;
+
+	FD3D12BuddyAllocator* CreateNewAllocator();
+
+	TArray<FD3D12BuddyAllocator*> Allocators;
+};
+
+//-----------------------------------------------------------------------------
+//	Linear Allocator
+//-----------------------------------------------------------------------------
+// A very simple allocator which sub-allocates out of pages. Allocations are not re-used thus
+// fragmentation can occur when allocating resources with different life-spans but packs data 
+// well when they are i.e. Allocate all the static meshes for 1 level
+struct FD3D12LinearAllocatorPage
+{
+	FD3D12LinearAllocatorPage() : 
+		Resource(nullptr), DeallocatedSpace(0), CurrentOffset(0), BaseAddress(nullptr){};
+	FD3D12Resource* Resource;
+
+	uint32 CurrentOffset;
+	uint32 DeallocatedSpace;
+	void* BaseAddress;
+};
+
+class FD3D12LinearAllocator : public FD3D12ResourceAllocator
+{
+public:
+
+	FD3D12LinearAllocator(FD3D12Device* ParentDevice,
+		FString Name,
+		D3D12_HEAP_TYPE heapType,
+		D3D12_RESOURCE_FLAGS flags,
+		uint32 MaxSizeForPooling,
+		uint32 PageSize);
+
+	// Public Override Functions
+
+	virtual FD3D12ResourceBlockInfo* TryAllocate(uint32 SizeInBytes, uint32 Alignment) final override;
+
+	virtual void Deallocate(FD3D12ResourceBlockInfo* Block)  final override;
+
+	virtual void Initialize() final override;
+
+	virtual void Destroy() final override;
+
+	virtual void CleanUpAllocations() final override;
+
+	virtual void DumpAllocatorStats(class FOutputDevice& Ar) final override;
+
+	virtual void ReleaseAllResources() final override;
+
+	virtual void Reset() final override;
+
+protected:
+
+	const uint32 AllocationPageSize;
+
+private:
+
+	FD3D12LinearAllocatorPage* FindPageWithBestFit(uint32 AllocationSize, uint32 Alignment);
+
+	struct FD3D12PageSubAllocation
+	{
+		FD3D12LinearAllocatorPage* Parent;
+	};
+	
+	FD3D12LinearAllocatorPage* AllocatePage();
+	void DestroyPage(FD3D12LinearAllocatorPage& Page);
+
+	TArray<FD3D12LinearAllocatorPage*> ActivePages;
+};
+
+//-----------------------------------------------------------------------------
+//	Bucket Allocator
+//-----------------------------------------------------------------------------
+// Resources are allocated from buckets, which are just a collection of resources of a particular size.
+// Blocks can be an entire resource or a sub allocation from a resource.
+class FD3D12BucketAllocator : public FD3D12ResourceAllocator
+{
+public:
+
+	FD3D12BucketAllocator(FD3D12Device* ParentDevice,
+		FString Name,
+		D3D12_HEAP_TYPE heapType,
+		D3D12_RESOURCE_FLAGS flags,
+		uint64 BlockRetentionFrameCountIn);
+
+	// Public Override Functions
+
+	virtual FD3D12ResourceBlockInfo* TryAllocate(uint32 SizeInBytes, uint32 Alignment) final override;
+
+	virtual void Deallocate(FD3D12ResourceBlockInfo* Block)  final override;
+
+	virtual void Initialize() final override;
+
+	virtual void Destroy() final override;
+
+	virtual void CleanUpAllocations() final override;
+
+	virtual void DumpAllocatorStats(class FOutputDevice& Ar) final override;
+
+	virtual void ReleaseAllResources() final override;
+
+	virtual void Reset() final override;
+
+private:
+
+	static uint32 FORCEINLINE BucketFromSize(uint32 size, uint32 bucketShift)
+	{
+		uint32 bucket = FMath::CeilLogTwo(size);
+		bucket = bucket < bucketShift ? 0 : bucket - bucketShift;
+		return bucket;
+	}
+
+	static uint32 FORCEINLINE BlockSizeFromBufferSize(uint32 bufferSize, uint32 bucketShift)
+	{
+		const uint32 minSize = 1 << bucketShift;
+		return bufferSize > minSize ? FMath::RoundUpToPowerOfTwo(bufferSize) : minSize;
+	}
+
+#if SUB_ALLOCATED_DEFAULT_ALLOCATIONS
+	static const uint32 MIN_HEAP_SIZE = 256 * 1024;
+#else
+	static const uint32 MIN_HEAP_SIZE = 64 * 1024;
+#endif
+
+	static const uint32 BucketShift = 6;
+	static const uint32 NumBuckets = 22; // bucket resource sizes range from 64 to 2^28 
+	FThreadsafeQueue<FD3D12ResourceBlockInfo*> AvailableBlocks[NumBuckets];
+	FThreadsafeQueue<FD3D12ResourceBlockInfo*> ExpiredBlocks;
+	TArray<FD3D12Resource*> SubAllocatedResources;// keep a list of the sub-allocated resources so that they may be cleaned up
+
+	// This frame count value helps makes sure that we don't delete resources too soon. If resources are deleted too soon,
+	// we can get in a loop the heap allocator will be constantly deleting and creating resources every frame which
+	// results in CPU stutters. DynamicRetentionFrameCount was tested and set to a value that appears to be adequate for
+	// creating a stable state on the Infiltrator demo.
+	const uint64 BlockRetentionFrameCount;
+};
+
+
+//-----------------------------------------------------------------------------
+//	FD3D12DynamicHeapAllocator
+//-----------------------------------------------------------------------------
+// This is designed for allocation of scratch memory such as temporary staging buffers
+// or shadow buffers for dynamic resources.
+class FD3D12DynamicHeapAllocator : public FD3D12DeviceChild
+{
+public:
+
+	FD3D12DynamicHeapAllocator(FD3D12Device* InParent, FString Name, eBuddyAllocationStrategy allocationStrategy,
+		uint32 MaxSizeForPooling,
+		uint32 maxBlockSize,
+		uint32 minBlockSize);
+
+	// These are mutually exclusive calls!
+	inline void SetCurrentCommandContext(class FD3D12CommandContext* CmdContext)
+	{
+		CurrentCommandContext = CmdContext;
+		check(CurrentCommandListHandle == nullptr);
+	}
+
+	inline void SetCurrentCommandListHandle(FD3D12CommandListHandle CommandListHandle)
+	{
+		check(CurrentCommandContext == nullptr);
+		CurrentCommandListHandle = CommandListHandle;
+	}
+
+	// Allocates <size> bytes from the end of an available resource heap.
+	void* AllocUploadResource(uint32 size, uint32 alignment, class FD3D12ResourceLocation* ResourceLocation);
+
+	void CleanUpAllocations();
+
+	void Destroy();
+
+private:
+	// The currently paired command context.  If none is available (as with async texture loads),
+	// the current command list handle can be provided instead.  The command list handle is the
+	// critical part, but the command context can always return the current list handle.
+	// I.e. you need one but not both of these to be valid
+	class FD3D12CommandContext* CurrentCommandContext;
+	FD3D12CommandListHandle CurrentCommandListHandle;
+
+	FD3D12MultiBuddyAllocator* Allocator;
+};
+
+//-----------------------------------------------------------------------------
+//	FD3D12DefaultBufferPool
+//-----------------------------------------------------------------------------
+class FD3D12DefaultBufferPool : public FD3D12DeviceChild
+{
+public:
+	FD3D12DefaultBufferPool(FD3D12Device* InParent, FD3D12ResourceAllocator* AllocatorIn);
+	~FD3D12DefaultBufferPool() { delete Allocator; }
+
+	// Grab a buffer from the available buffers or create a new buffer if none are available
+	void AllocDefaultResource(const D3D12_RESOURCE_DESC& Desc, D3D12_SUBRESOURCE_DATA* pInitialData, class FD3D12ResourceLocation* ResourceLocation, uint32 Alignment);
+
+	void CleanUpAllocations();
+
+private:
+	FD3D12ResourceAllocator* Allocator;
+};
+
+// FD3D12DefaultBufferAllocator
+//
+class FD3D12DefaultBufferAllocator : public FD3D12DeviceChild
+{
+public:
+	FD3D12DefaultBufferAllocator(FD3D12Device* InParent) :
+		FD3D12DeviceChild(InParent)
+	{
+		FMemory::Memset(DefaultBufferPools, 0);
+	}
+
+	// Grab a buffer from the available buffers or create a new buffer if none are available
+	HRESULT AllocDefaultResource(const D3D12_RESOURCE_DESC& pDesc, D3D12_SUBRESOURCE_DATA* InitialData, class FD3D12ResourceLocation* ResourceLocation, uint32 Alignment);
+	void FreeDefaultBufferPools();
+	void CleanupFreeBlocks();
+
+private:
+
+	static const uint32 MAX_DEFAULT_POOLS = 16; // Should match the max D3D12_RESOURCE_FLAG_FLAG combinations.
+	FD3D12DefaultBufferPool* DefaultBufferPools[MAX_DEFAULT_POOLS];
+
+	bool BufferIsWriteable(const D3D12_RESOURCE_DESC& Desc)
+	{
+		const bool bDSV = (Desc.Flags & D3D12_RESOURCE_FLAG_ALLOW_DEPTH_STENCIL) != 0;
+		const bool bRTV = (Desc.Flags & D3D12_RESOURCE_FLAG_ALLOW_RENDER_TARGET) != 0;
+		const bool bUAV = (Desc.Flags & D3D12_RESOURCE_FLAG_ALLOW_UNORDERED_ACCESS) != 0;
+
+		// Buffer Depth Stencils are invalid
+		check(bDSV == false);
+		const bool bWriteable = bDSV || bRTV || bUAV;
+		return bWriteable;
+	}
+};
+
+//-----------------------------------------------------------------------------
+//	FD3D12TextureAllocator
+//-----------------------------------------------------------------------------
+
+class FD3D12TextureAllocator : public FD3D12BuddyAllocator
+{
+public:
+	FD3D12TextureAllocator(FD3D12Device* Device, FString Name, uint32 HeapSize, D3D12_HEAP_FLAGS Flags);
+
+	~FD3D12TextureAllocator();
+
+	HRESULT AllocateTexture(D3D12_RESOURCE_DESC Desc, const D3D12_CLEAR_VALUE* ClearValue, FD3D12ResourceLocation* TextureLocation);
+};
+
+#define TEXTURE_POOL_SIZE_READABLE (64 * 1024 * 1024)
+
+class FD3D12TextureAllocatorPool : public FD3D12DeviceChild
+{
+public:
+	FD3D12TextureAllocatorPool(FD3D12Device* Device) :
+		ReadOnlyTexturePool(Device, FString(L"Small Read-Only Texture allocator"), TEXTURE_POOL_SIZE_READABLE, D3D12_HEAP_FLAG_ALLOW_ONLY_NON_RT_DS_TEXTURES),
+		FD3D12DeviceChild(Device)
+	{};
+
+	HRESULT AllocateTexture(D3D12_RESOURCE_DESC Desc, const D3D12_CLEAR_VALUE* ClearValue, uint8 UEFormat, FD3D12ResourceLocation* TextureLocation);
+
+	void CleanUpAllocations() { ReadOnlyTexturePool.CleanUpAllocations(); }
+
+	void Destroy() { ReadOnlyTexturePool.Destroy(); }
+
+private:
+	FD3D12TextureAllocator ReadOnlyTexturePool;
+};
\ No newline at end of file
diff --git a/Engine/Source/Runtime/Windows/D3D12RHI/Private/D3D12Buffer.cpp b/Engine/Source/Runtime/Windows/D3D12RHI/Private/D3D12Buffer.cpp
new file mode 100644
--- /dev/null
+++ b/Engine/Source/Runtime/Windows/D3D12RHI/Private/D3D12Buffer.cpp
@@ -0,0 +1,252 @@
+// Copyright 1998-2016 Epic Games, Inc. All Rights Reserved.
+
+/*=============================================================================
+D3D12Buffer.cpp: D3D Common code for buffers.
+=============================================================================*/
+
+#include "D3D12RHIPrivate.h"
+
+// Forward declarations are required for the template functions
+template void* FD3D12DynamicRHI::LockBuffer<FD3D12VertexBuffer>(FRHICommandListImmediate* RHICmdList, FD3D12VertexBuffer* Buffer, uint32 Offset, uint32 Size, EResourceLockMode LockMode);
+template void* FD3D12DynamicRHI::LockBuffer<FD3D12IndexBuffer>(FRHICommandListImmediate* RHICmdList, FD3D12IndexBuffer* Buffer, uint32 Offset, uint32 Size, EResourceLockMode LockMode);
+template void* FD3D12DynamicRHI::LockBuffer<FD3D12StructuredBuffer>(FRHICommandListImmediate* RHICmdList, FD3D12StructuredBuffer* Buffer, uint32 Offset, uint32 Size, EResourceLockMode LockMode);
+
+template void FD3D12DynamicRHI::UnlockBuffer<FD3D12VertexBuffer>(FRHICommandListImmediate* RHICmdList, FD3D12VertexBuffer* Buffer);
+template void FD3D12DynamicRHI::UnlockBuffer<FD3D12IndexBuffer>(FRHICommandListImmediate* RHICmdList, FD3D12IndexBuffer* Buffer);
+template void FD3D12DynamicRHI::UnlockBuffer<FD3D12StructuredBuffer>(FRHICommandListImmediate* RHICmdList, FD3D12StructuredBuffer* Buffer);
+
+struct FRHICommandUpdateBuffer : public FRHICommand<FRHICommandUpdateBuffer>
+{
+	TRefCountPtr<FD3D12ResourceLocation> Source;
+	TRefCountPtr<FD3D12ResourceLocation> Destination;
+	uint32 NumBytes;
+
+	FORCEINLINE_DEBUGGABLE FRHICommandUpdateBuffer(TRefCountPtr<FD3D12ResourceLocation> InDest, TRefCountPtr<FD3D12ResourceLocation> InSource, uint32 InNumBytes) :
+		Source(InSource), Destination(InDest), NumBytes(InNumBytes) {}
+
+	void Execute(FRHICommandListBase& CmdList)
+	{
+		FD3D12DynamicRHI::GetD3DRHI()->UpdateBuffer(Destination->GetResource(), Destination->GetOffset(), Source->GetResource(), Source->GetOffset(), NumBytes);
+	}
+};
+
+// This allows us to rename resources from the RenderThread i.e. all the 'hard' work of allocating a new resource
+// is done in parallel and this small function is called to switch the resource to point to the correct location
+// a the correct time.
+template<typename ResourceType>
+struct FRHICommandRenameUploadBuffer : public FRHICommand<FRHICommandRenameUploadBuffer<ResourceType>>
+{
+	ResourceType* Resource;
+	TRefCountPtr<FD3D12ResourceLocation> NewResource;
+	FORCEINLINE_DEBUGGABLE FRHICommandRenameUploadBuffer(ResourceType* InResource, FD3D12ResourceLocation* NewResource) :
+		Resource(InResource), NewResource(NewResource) {}
+
+	void Execute(FRHICommandListBase& CmdList) { Resource->Rename(NewResource); }
+};
+
+FD3D12ResourceLocation* FD3D12DynamicRHI::CreateBuffer(FRHICommandListImmediate* RHICmdList, const D3D12_RESOURCE_DESC Desc, uint32 Size, uint32 InUsage, FRHIResourceCreateInfo& CreateInfo, uint32 Alignment)
+{
+	// Explicitly check that the size is nonzero before allowing CreateBuffer to opaquely fail.
+	check(Size > 0);
+
+	const bool bIsDynamic = (InUsage & BUF_AnyDynamic) ? true : false;
+
+	FD3D12ResourceLocation* ResourceLocation = new FD3D12ResourceLocation(GetRHIDevice());
+
+	// If a resource array was provided for the resource, create the resource pre-populated
+	D3D12_SUBRESOURCE_DATA InitData = { 0 };
+	D3D12_SUBRESOURCE_DATA* pInitData = nullptr;
+	if (CreateInfo.ResourceArray)
+	{
+		check(Size == CreateInfo.ResourceArray->GetResourceDataSize());
+		InitData.pData = CreateInfo.ResourceArray->GetResourceData();
+		InitData.RowPitch = Size;
+		InitData.SlicePitch = 0;
+		pInitData = &InitData;
+	}
+
+	if (bIsDynamic)
+	{
+		void* pData = GetRHIDevice()->GetDefaultUploadHeapAllocator().AllocUploadResource(Size, Alignment, ResourceLocation);
+		check(ResourceLocation->GetEffectiveBufferSize() == Size);
+
+		if (pInitData)
+		{
+			// Handle initial data
+			FMemory::Memcpy(pData, InitData.pData, Size);
+		}
+	}
+	else
+	{
+		if (RHICmdList && pInitData)
+		{
+			// We only need to synchronize when creating default resource buffers (because we need a command list to initialize them)
+			FScopedRHIThreadStaller StallRHIThread(*RHICmdList);
+
+			VERIFYD3D11RESULT(GetRHIDevice()->GetDefaultBufferAllocator().AllocDefaultResource(Desc, pInitData, ResourceLocation, Alignment));
+		}
+		else
+		{
+			VERIFYD3D11RESULT(GetRHIDevice()->GetDefaultBufferAllocator().AllocDefaultResource(Desc, pInitData, ResourceLocation, Alignment));
+		}
+		check(ResourceLocation->GetEffectiveBufferSize() == Size);
+	}
+
+	if (CreateInfo.ResourceArray)
+	{
+		// Discard the resource array's contents.
+		CreateInfo.ResourceArray->Discard();
+	}
+
+	return ResourceLocation;
+}
+
+template<class BufferType>
+void* FD3D12DynamicRHI::LockBuffer(FRHICommandListImmediate* RHICmdList, BufferType* Buffer, uint32 Offset, uint32 Size, EResourceLockMode LockMode)
+{
+	FD3D12LockedData LockedData;
+
+	// Determine whether the buffer is dynamic or not.
+	const bool bIsDynamic = (Buffer->GetUsage() & BUF_AnyDynamic) ? true : false;
+
+	void* Data = nullptr;
+
+	if (bIsDynamic)
+	{
+		check(LockMode == RLM_WriteOnly);
+
+		TRefCountPtr<FD3D12ResourceLocation> newLocation = new FD3D12ResourceLocation(GetRHIDevice());
+
+		// Allocate a new resource
+		Data = GetRHIDevice()->GetDefaultUploadHeapAllocator().AllocUploadResource(Buffer->ResourceLocation->GetEffectiveBufferSize(), Buffer->BufferAlignment, newLocation);
+
+		// If on the RenderThread, queue up a command on the RHIThread to rename this buffer at the correct time
+		if (ShouldDeferBufferLockOperation(RHICmdList))
+		{
+			new (RHICmdList->AllocCommand<FRHICommandRenameUploadBuffer<BufferType>>()) FRHICommandRenameUploadBuffer<BufferType>(Buffer, newLocation);
+		}
+		else
+		{
+			Buffer->Rename(newLocation);
+		}
+	}
+	else
+	{
+		FD3D12Resource* pResource = Buffer->ResourceLocation->GetResource();
+
+		// Locking for read must occur immediately so we can't queue up the operations later.
+		if (LockMode == RLM_ReadOnly)
+		{
+			// If the static buffer is being locked for reading, create a staging buffer.
+			TRefCountPtr<FD3D12Resource> StagingBuffer;
+			VERIFYD3D11RESULT(GetRHIDevice()->GetResourceHelper().CreateBuffer(D3D12_HEAP_TYPE_READBACK, Offset + Size, StagingBuffer.GetInitReference()));
+			LockedData.StagingResource = StagingBuffer;
+
+			// Copy the contents of the buffer to the staging buffer.
+			{
+				const auto& pfnCopyContents = [&]()
+				{
+					FD3D12CommandContext& DefaultContext = GetRHIDevice()->GetDefaultCommandContext();
+
+					FD3D12CommandListHandle& hCommandList = DefaultContext.CommandListHandle;
+					FScopeResourceBarrier ScopeResourceBarrierSource(hCommandList, pResource, pResource->GetDefaultResourceState(), D3D12_RESOURCE_STATE_COPY_SOURCE, 0);
+					// Don't need to transition upload heaps
+
+					DefaultContext.numCopies++;
+					hCommandList->CopyBufferRegion(
+						StagingBuffer->GetResource(),
+						0,
+						pResource->GetResource(),
+						Offset, Size);
+
+					DefaultContext.FlushCommands(true);
+				};
+
+				if (ShouldDeferBufferLockOperation(RHICmdList))
+				{
+					// Sync when in the render thread implementation
+					check(IsInRHIThread() == false);
+
+					RHICmdList->ImmediateFlush(EImmediateFlushType::FlushRHIThread);
+					pfnCopyContents();
+				}
+				else
+				{
+					check(IsInRHIThread());
+					pfnCopyContents();
+				}
+			}
+
+			// Map the staging buffer's memory for reading.
+			VERIFYD3D11RESULT(StagingBuffer->GetResource()->Map(0, nullptr, &Data));
+		}
+		else
+		{
+			// If the static buffer is being locked for writing, allocate memory for the contents to be written to.
+			TRefCountPtr<FD3D12ResourceLocation> UploadBufferLocation = new FD3D12ResourceLocation(GetRHIDevice());
+			Data = GetRHIDevice()->GetDefaultFastAllocator().Allocate(Offset + Size, D3D12_TEXTURE_DATA_PLACEMENT_ALIGNMENT, UploadBufferLocation);
+
+			// Keep track of the underlying resource for the upload heap so it can be referenced during Unmap.
+			LockedData.UploadHeapLocation = UploadBufferLocation;
+		}
+	}
+
+	check(Data);
+	// Add the lock to the lock map.
+	LockedData.SetData(Data);
+	LockedData.Pitch = Offset + Size;
+
+	// Add the lock to the lock map.
+	FD3D12LockedKey LockedKey(Buffer);
+	AddToOutstandingLocks(LockedKey, LockedData);
+
+	// Return the offset pointer
+	return (void*)((uint8*)LockedData.GetData() + Offset);
+}
+
+template<class BufferType>
+void FD3D12DynamicRHI::UnlockBuffer(FRHICommandListImmediate* RHICmdList, BufferType* Buffer)
+{
+	// Find the outstanding lock for this Buffer.
+	FD3D12LockedKey LockedKey(Buffer);
+	FD3D12LockedData* LockedData = FindInOutstandingLocks(LockedKey);
+	check(LockedData);
+
+	// Determine whether the buffer is dynamic or not.
+	const bool bIsDynamic = (Buffer->GetUsage() & BUF_AnyDynamic) ? true : false;
+
+	if (bIsDynamic)
+	{
+		// If the Buffer is dynamic, its upload heap memory can always stay mapped. Don't do anything.
+	}
+	else
+	{
+		// If the static Buffer lock involved a staging resource, it was locked for reading.
+		if (LockedData->StagingResource)
+		{
+			// Unmap the staging buffer's memory.
+			ID3D12Resource* StagingBuffer = LockedData->StagingResource.GetReference()->GetResource();
+			StagingBuffer->Unmap(0, nullptr);
+		}
+		else
+		{
+			// Copy the contents of the temporary memory buffer allocated for writing into the Buffer.
+			FD3D12ResourceLocation* UploadHeapLocation = LockedData->UploadHeapLocation.GetReference();
+
+			// If we are on the render thread, queue up the copy on the RHIThread so it happens at the correct time.
+			if (ShouldDeferBufferLockOperation(RHICmdList))
+			{
+				new (RHICmdList->AllocCommand<FRHICommandUpdateBuffer>()) FRHICommandUpdateBuffer(Buffer->ResourceLocation, UploadHeapLocation, LockedData->Pitch);
+			}
+			else
+			{
+				UpdateBuffer(Buffer->ResourceLocation->GetResource(),
+					Buffer->ResourceLocation->GetOffset(), UploadHeapLocation->GetResource(), UploadHeapLocation->GetOffset(), LockedData->Pitch);
+			}
+		}
+	}
+
+	// Remove the FD3D12LockedData from the lock map.
+	// If the lock involved a staging resource, this releases it.
+	RemoveFromOutstandingLocks(LockedKey);
+}
diff --git a/Engine/Source/Runtime/Windows/D3D12RHI/Private/D3D12Commands.cpp b/Engine/Source/Runtime/Windows/D3D12RHI/Private/D3D12Commands.cpp
--- a/Engine/Source/Runtime/Windows/D3D12RHI/Private/D3D12Commands.cpp
+++ b/Engine/Source/Runtime/Windows/D3D12RHI/Private/D3D12Commands.cpp
@@ -209,8 +209,6 @@ void FD3D12CommandContext::RHIDispatchComputeShader(uint32 ThreadGroupCountX, ui
 	numDispatches++;
 	CommandListHandle->Dispatch(ThreadGroupCountX, ThreadGroupCountY, ThreadGroupCountZ);
 
-	StateCache.FlushComputeShaderCache();
-
 	DEBUG_EXECUTE_COMMAND_LIST(this);
 
 	StateCache.SetComputeShader(nullptr);
@@ -245,19 +243,195 @@ void FD3D12CommandContext::RHIDispatchIndirectComputeShader(FVertexBufferRHIPara
 		0
 		);
 
-	StateCache.FlushComputeShaderCache();
-
 	DEBUG_EXECUTE_COMMAND_LIST(this);
 
 	StateCache.SetComputeShader(nullptr);
 }
 
+
+void FD3D12CommandContext::RHITransitionResources(EResourceTransitionAccess TransitionType, FTextureRHIParamRef* InTextures, int32 NumTextures)
+{
+#if !USE_D3D12RHI_RESOURCE_STATE_TRACKING
+	check(TransitionType == EResourceTransitionAccess::EReadable || TransitionType == EResourceTransitionAccess::EWritable || TransitionType == EResourceTransitionAccess::ERWSubResBarrier);
+	// TODO: Remove this skip.
+	// Skip for now because we don't have enough info about what mip to transition yet.
+	// Note: This causes visual corruption.
+	if (TransitionType == EResourceTransitionAccess::ERWSubResBarrier)
+	{
+		return;
+	}
+
+	static IConsoleVariable* CVarShowTransitions = IConsoleManager::Get().FindConsoleVariable(TEXT("r.ProfileGPU.ShowTransitions"));
+	const bool bShowTransitionEvents = CVarShowTransitions->GetInt() != 0;
+
+	SCOPED_RHI_CONDITIONAL_DRAW_EVENTF(*this, RHITransitionResources, bShowTransitionEvents, TEXT("TransitionTo: %s: %i Textures"), *FResourceTransitionUtility::ResourceTransitionAccessStrings[(int32)TransitionType], NumTextures);
+
+	// Determine the direction of the transitions.
+	const D3D12_RESOURCE_STATES* pBefore = nullptr;
+	const D3D12_RESOURCE_STATES* pAfter = nullptr;
+	D3D12_RESOURCE_STATES WritableState;
+	D3D12_RESOURCE_STATES ReadableState;
+	switch (TransitionType)
+	{
+	case EResourceTransitionAccess::EReadable:
+		// Write -> Read
+		pBefore = &WritableState;
+		pAfter = &ReadableState;
+		break;
+
+	case EResourceTransitionAccess::EWritable:
+		// Read -> Write
+		pBefore = &ReadableState;
+		pAfter = &WritableState;
+		break;
+
+	default:
+		check(false);
+		break;
+	}
+
+	// Create the resource barrier descs for each texture to transition.
+	uint32 BarrierCount = 0;
+	const uint32 MaxBarrierCount = 64;
+	D3D12_RESOURCE_BARRIER Barriers[MaxBarrierCount];
+	for (int32 i = 0; i < NumTextures; ++i)
+	{
+		if (InTextures[i])
+		{
+			FD3D12Resource* Resource = GetD3D11TextureFromRHITexture(InTextures[i])->GetResource();
+			check(Resource->RequiresResourceStateTracking());
+
+			SCOPED_RHI_CONDITIONAL_DRAW_EVENTF(*this, RHITransitionResourcesLoop, bShowTransitionEvents, TEXT("To:%i - %s"), i, *Resource->GetName().ToString());
+
+			WritableState = Resource->GetWritableState();
+			ReadableState = Resource->GetReadableState();
+
+			check(BarrierCount < MaxBarrierCount);
+			D3D12_RESOURCE_BARRIER& BarrierDesc = Barriers[BarrierCount];
+			BarrierDesc.Type = D3D12_RESOURCE_BARRIER_TYPE_TRANSITION;
+			BarrierDesc.Flags = D3D12_RESOURCE_BARRIER_FLAG_NONE;
+			BarrierDesc.Transition.pResource = Resource->GetResource();
+			BarrierDesc.Transition.StateBefore = *pBefore;
+			BarrierDesc.Transition.StateAfter = *pAfter;
+			BarrierDesc.Transition.Subresource = D3D12_RESOURCE_BARRIER_ALL_SUBRESOURCES;
+			check(BarrierDesc.Transition.StateBefore != BarrierDesc.Transition.StateAfter);
+
+			BarrierCount++;
+
+			DUMP_TRANSITION(Resource->GetName(), TransitionType);
+		}
+	}
+
+	// Transition the resources.
+	if (BarrierCount > 0)
+	{
+		CommandListHandle->ResourceBarrier(BarrierCount, Barriers);
+		this->numBarriers += BarrierCount;
+	}
+#endif // !USE_D3D12RHI_RESOURCE_STATE_TRACKING
+}
+
+
 void FD3D12CommandContext::RHITransitionResources(EResourceTransitionAccess TransitionType, EResourceTransitionPipeline TransitionPipeline, FUnorderedAccessViewRHIParamRef* InUAVs, int32 NumUAVs, FComputeFenceRHIParamRef WriteComputeFenceRHI)
 {
-	FD3D12Fence* Fence = FD3D12DynamicRHI::ResourceCast(WriteComputeFenceRHI);
+	static IConsoleVariable* CVarShowTransitions = IConsoleManager::Get().FindConsoleVariable(TEXT("r.ProfileGPU.ShowTransitions"));
+	const bool bShowTransitionEvents = CVarShowTransitions->GetInt() != 0;
 
-	if (Fence)
+	SCOPED_RHI_CONDITIONAL_DRAW_EVENTF(*this, RHITransitionResources, bShowTransitionEvents, TEXT("TransitionTo: %s: %i UAVs"), *FResourceTransitionUtility::ResourceTransitionAccessStrings[(int32)TransitionType], NumUAVs);
+	const bool bUAVTransition = (TransitionType == EResourceTransitionAccess::EReadable) || (TransitionType == EResourceTransitionAccess::EWritable || TransitionType == EResourceTransitionAccess::ERWBarrier);
+	const bool bUAVBarrier = (TransitionType == EResourceTransitionAccess::ERWBarrier && TransitionPipeline == EResourceTransitionPipeline::EComputeToCompute);
+
+	if (bUAVBarrier)
+	{
+		// UAV barrier between Dispatch() calls to ensure all R/W accesses are complete.
+		StateCache.FlushComputeShaderCache(true);
+	}
+	else if (bUAVTransition)
+	{
+#if !USE_D3D12RHI_RESOURCE_STATE_TRACKING
+		// Determine the direction of the transitions.
+		// Note in this method, the writeable state is always UAV, regardless of the FD3D12Resource's Writeable state.
+		const D3D12_RESOURCE_STATES* pBefore = nullptr;
+		const D3D12_RESOURCE_STATES* pAfter = nullptr;
+		const D3D12_RESOURCE_STATES WritableComputeState = D3D12_RESOURCE_STATE_UNORDERED_ACCESS;
+		D3D12_RESOURCE_STATES WritableGraphicsState;
+		D3D12_RESOURCE_STATES ReadableState;
+		switch (TransitionType)
+		{
+		case EResourceTransitionAccess::EReadable:
+			// Write -> Read
+			pBefore = &WritableComputeState;
+			pAfter = &ReadableState;
+			break;
+
+		case EResourceTransitionAccess::EWritable:
+			// Read -> Write
+			pBefore = &ReadableState;
+			pAfter = &WritableComputeState;
+			break;
+
+		case EResourceTransitionAccess::ERWBarrier:
+			// Write -> Write, but switching from Grfx to Compute.
+			check(TransitionPipeline == EResourceTransitionPipeline::EGfxToCompute);
+			pBefore = &WritableGraphicsState;
+			pAfter = &WritableComputeState;
+			break;
+
+		default:
+			check(false);
+			break;
+		}
+
+		// Create the resource barrier descs for each texture to transition.
+		uint32 BarrierCount = 0;
+		const uint32 MaxBarrierCount = 16;
+		D3D12_RESOURCE_BARRIER Barriers[MaxBarrierCount];
+		for (int32 i = 0; i < NumUAVs; ++i)
+		{
+			if (InUAVs[i])
+			{
+				FD3D12UnorderedAccessView* UnorderedAccessView = FD3D12DynamicRHI::ResourceCast(InUAVs[i]);
+				FD3D12Resource* Resource = UnorderedAccessView->GetResource();
+				check(Resource->RequiresResourceStateTracking());
+
+				SCOPED_RHI_CONDITIONAL_DRAW_EVENTF(*this, RHITransitionResourcesLoop, bShowTransitionEvents, TEXT("To:%i - %s"), i, *Resource->GetName().ToString());
+
+				// The writable compute state is always UAV.
+				WritableGraphicsState = Resource->GetWritableState();
+				ReadableState = Resource->GetReadableState();
+
+				// Some ERWBarriers might have the same before and after states.
+				if (*pBefore != *pAfter)
+				{
+					check(BarrierCount < MaxBarrierCount);
+					D3D12_RESOURCE_BARRIER& BarrierDesc = Barriers[BarrierCount];
+					BarrierDesc.Type = D3D12_RESOURCE_BARRIER_TYPE_TRANSITION;
+					BarrierDesc.Flags = D3D12_RESOURCE_BARRIER_FLAG_NONE;
+					BarrierDesc.Transition.pResource = Resource->GetResource();
+					BarrierDesc.Transition.StateBefore = *pBefore;
+					BarrierDesc.Transition.StateAfter = *pAfter;
+					BarrierDesc.Transition.Subresource = D3D12_RESOURCE_BARRIER_ALL_SUBRESOURCES;
+					check(BarrierDesc.Transition.StateBefore != BarrierDesc.Transition.StateAfter);
+
+					BarrierCount++;
+
+					DUMP_TRANSITION(Resource->GetName(), TransitionType);
+				}
+			}
+		}
+
+		// Transition the resources.
+		if (BarrierCount > 0)
+		{
+			CommandListHandle->ResourceBarrier(BarrierCount, Barriers);
+			this->numBarriers += BarrierCount;
+		}
+#endif // !USE_D3D12RHI_RESOURCE_STATE_TRACKING
+	}
+
+	if (WriteComputeFenceRHI)
 	{
+		FD3D12Fence* Fence = FD3D12DynamicRHI::ResourceCast(WriteComputeFenceRHI);
 		Fence->WriteFence();
 
 		if (!bIsAsyncComputeContext)
@@ -2192,6 +2366,11 @@ void FD3D12CommandContext::RHIClearMRTImpl(bool bClearColor, int32 NumClearColor
 	DEBUG_EXECUTE_COMMAND_LIST(this);
 }
 
+void FD3D12CommandContext::RHIBindClearMRTValues(bool bClearColor, bool bClearDepth, bool bClearStencil)
+{
+	// Not necessary for d3d.
+}
+
 void FD3D12CommandContext::RHIBeginAsyncComputeJob_DrawThread(EAsyncComputePriority Priority)
 {
 	UE_LOG(LogRHI, Fatal, TEXT("%s not implemented yet"), ANSI_TO_TCHAR(__FUNCTION__));
diff --git a/Engine/Source/Runtime/Windows/D3D12RHI/Private/D3D12Device.cpp b/Engine/Source/Runtime/Windows/D3D12RHI/Private/D3D12Device.cpp
--- a/Engine/Source/Runtime/Windows/D3D12RHI/Private/D3D12Device.cpp
+++ b/Engine/Source/Runtime/Windows/D3D12RHI/Private/D3D12Device.cpp
@@ -23,10 +23,18 @@ using namespace D3D12RHI;
 bool D3D12RHI_ShouldCreateWithD3DDebug()
 {
 	// Use a debug device if specified on the command line.
-	return
+	static bool bCreateWithD3DDebug = 
 		FParse::Param(FCommandLine::Get(), TEXT("d3ddebug")) ||
 		FParse::Param(FCommandLine::Get(), TEXT("d3debug")) ||
 		FParse::Param(FCommandLine::Get(), TEXT("dxdebug"));
+	return bCreateWithD3DDebug;
+}
+
+bool D3D12RHI_ShouldCreateWithWarp()
+{
+	// Use the warp adapter if specified on the command line.
+	static bool bCreateWithWarp = FParse::Param(FCommandLine::Get(), TEXT("warp"));
+	return bCreateWithWarp;
 }
 
 bool D3D12RHI_ShouldAllowAsyncResourceCreation()
@@ -141,19 +149,33 @@ public:
 
 		CommandLists.Reset();
 
-		if (GCommandListBatchingMode != CLB_AggressiveBatching)
+		bool Flush = false;
+		// If the GPU is starving (i.e. we are CPU bound) feed it asap!
+		if (OwningDevice->IsGPUIdle() && OwningDevice->PendingCommandLists.Num() > 0)
 		{
-			// Submit when the batch is finished.
-			const bool FinalCommandListInBatch = Index == (Num - 1);
-			if (FinalCommandListInBatch && OwningDevice->PendingCommandLists.Num() > 0)
+			Flush = true;
+		}
+		else
+		{
+			if (GCommandListBatchingMode != CLB_AggressiveBatching)
 			{
+				// Submit when the batch is finished.
+				const bool FinalCommandListInBatch = Index == (Num - 1);
+				if (FinalCommandListInBatch && OwningDevice->PendingCommandLists.Num() > 0)
+				{
+					Flush = true;
+				}
+			}
+		}
+
+		if (Flush)
+		{
 #if SUPPORTS_MEMORY_RESIDENCY
-				OwningDevice->GetOwningRHI()->GetResourceResidencyManager().MakeResident();
+			OwningDevice->GetOwningRHI()->GetResourceResidencyManager().MakeResident();
 #endif
-				OwningDevice->GetCommandListManager().ExecuteCommandLists(OwningDevice->PendingCommandLists);
-				OwningDevice->PendingCommandLists.Reset();
-				OwningDevice->PendingCommandListsTotalWorkCommands = 0;
-			}
+			OwningDevice->GetCommandListManager().ExecuteCommandLists(OwningDevice->PendingCommandLists);
+			OwningDevice->PendingCommandLists.Reset();
+			OwningDevice->PendingCommandListsTotalWorkCommands = 0;
 		}
 
 		delete this;
@@ -264,7 +286,8 @@ FD3D12Device::FD3D12Device(FD3D12DynamicRHI* InOwningRHI, IDXGIFactory4* InDXGIF
 	BufferInitializerFastAllocatorPagePool(this, D3D12_HEAP_TYPE_UPLOAD, 1024 * 512),
 	BufferInitializerFastAllocator(this, &BufferInitializerFastAllocatorPagePool),
 	DefaultUploadHeapAllocator(this, FString(L"Upload Buffer Allocator"),kManualSubAllocationStrategy, DEFAULT_CONTEXT_UPLOAD_POOL_MAX_ALLOC_SIZE, DEFAULT_CONTEXT_UPLOAD_POOL_SIZE, DEFAULT_CONTEXT_UPLOAD_POOL_ALIGNMENT),
-	TextureAllocator(this)
+	TextureAllocator(this),
+	FenceCorePool(this)
 {
 }
 
@@ -426,6 +449,8 @@ FD3D12DynamicRHI::FD3D12DynamicRHI(IDXGIFactory4* InDXGIFactory, FD3D12Adapter&
 	GMaxTextureArrayLayers = D3D12_REQ_TEXTURE2D_ARRAY_AXIS_DIMENSION;
 	GSupportsTimestampRenderQueries = true;
 
+	GMaxTextureMipCount = FMath::CeilLogTwo(GMaxTextureDimensions) + 1;
+	GMaxTextureMipCount = FMath::Min<int32>(MAX_TEXTURE_MIP_COUNT, GMaxTextureMipCount);
 	GMaxShadowDepthBufferSizeX = 4096;
 	GMaxShadowDepthBufferSizeY = 4096;
 
@@ -701,6 +726,12 @@ uint32 FD3D12DynamicRHI::GetMaxMSAAQuality(uint32 SampleCount)
 	return 0xffffffff;
 }
 
+bool FD3D12Device::IsGPUIdle()
+{
+	FD3D12Fence& Fence = CommandListManager.GetFence(EFenceType::FT_CommandList);
+	return Fence.GetLastCompletedFence() >= (Fence.GetCurrentFence() - 1);
+}
+
 void FD3D12Device::CreateSignatures()
 {
 	// ExecuteIndirect command signatures
@@ -747,12 +778,6 @@ void FD3D12Device::CreateSignatures()
 
 void FD3D12Device::SetupAfterDeviceCreation()
 {
-	DXGI_ADAPTER_DESC* AdapterDesc = GetD3DAdapterDesc();
-
-	// Need to get these early to make vendor-specific decisions in initialization
-	GRHIAdapterName = AdapterDesc->Description;
-	GRHIVendorId = AdapterDesc->VendorId;
-
 	CreateSignatures();
 
 	PipelineStateCache = FD3D12PipelineStateCache(this);
@@ -826,12 +851,9 @@ void FD3D12Device::CleanupD3DDevice()
 		// Ensure any pending rendering is finished
 		GetCommandListManager().SignalFrameComplete(true);
 
-		// Flush all pending deletes before destroying the device.
-		FRHIResource::FlushPendingDeletes();
-
 		check(Direct3DDevice);
 
-		OwningRHI->OutstandingLocks.Empty();
+		OwningRHI->ClearOutstandingLocks();
 
 		PipelineStateCache.Close();
 
@@ -853,6 +875,9 @@ void FD3D12Device::CleanupD3DDevice()
 			ResourceIt->ReleaseDynamicRHI();
 		}
 
+		// Flush all pending deletes before destroying the device.
+		FRHIResource::FlushPendingDeletes();
+
 		Viewports.Empty();
 		DrawingViewport = nullptr;
 		SamplerMap.Empty();
@@ -882,6 +907,9 @@ void FD3D12Device::CleanupD3DDevice()
 		// after the above signal frame
 		GetCommandListManager().SignalFrameComplete(true);
 
+		// Flush all pending deletes before destroying the device.
+		FRHIResource::FlushPendingDeletes();
+
 		// Cleanup the allocator near the end, as some resources may be returned to the allocator
 		DefaultBufferAllocator.FreeDefaultBufferPools();
 
@@ -926,6 +954,8 @@ void FD3D12Device::CleanupD3DDevice()
 		AsyncCommandListManager.Destroy();
 
 
+		FenceCorePool.Destroy();
+
 	}
 }
 
diff --git a/Engine/Source/Runtime/Windows/D3D12RHI/Private/D3D12DirectCommandListManager.cpp b/Engine/Source/Runtime/Windows/D3D12RHI/Private/D3D12DirectCommandListManager.cpp
--- a/Engine/Source/Runtime/Windows/D3D12RHI/Private/D3D12DirectCommandListManager.cpp
+++ b/Engine/Source/Runtime/Windows/D3D12RHI/Private/D3D12DirectCommandListManager.cpp
@@ -9,40 +9,68 @@ extern bool D3D12RHI_ShouldCreateWithD3DDebug();
 
 FComputeFenceRHIRef FD3D12DynamicRHI::RHICreateComputeFence(const FName& Name)
 {
-	FD3D12Fence* Fence = new FD3D12Fence(Name);
+	FD3D12Fence* Fence = new FD3D12Fence(GetRHIDevice(),Name);
 
-	Fence->CreateFence(GetRHIDevice()->GetDevice(), 0);
+	Fence->CreateFence(0);
 
 	return Fence;
 }
 
-FD3D12Fence::FD3D12Fence(const FName& Name)
-	: FRHIComputeFence(Name)
-	, CurrentFence(-1)
-	, SignalFence(-1)
-	, LastCompletedFence(-1)
-	, hFenceCompleteEvent(INVALID_HANDLE_VALUE)
+FD3D12FenceCore::FD3D12FenceCore(FD3D12Device* Parent, uint64 InitialValue)
+	: hFenceCompleteEvent(INVALID_HANDLE_VALUE)
+	, FenceValueAvailableAt(0)
+	, FD3D12DeviceChild(Parent)
 {
-	// Create an event
+	check(Parent);
 	hFenceCompleteEvent = CreateEvent(nullptr, false, false, nullptr);
 	check(INVALID_HANDLE_VALUE != hFenceCompleteEvent);
+
+	VERIFYD3D11RESULT(Parent->GetDevice()->CreateFence(InitialValue, D3D12_FENCE_FLAG_NONE, IID_PPV_ARGS(Fence.GetInitReference())));
 }
 
-FD3D12Fence::~FD3D12Fence()
+FD3D12FenceCore::~FD3D12FenceCore()
 {
 	if (hFenceCompleteEvent != INVALID_HANDLE_VALUE)
 	{
 		CloseHandle(hFenceCompleteEvent);
 		hFenceCompleteEvent = INVALID_HANDLE_VALUE;
 	}
+}
+
+FD3D12Fence::FD3D12Fence(FD3D12Device* Parent, const FName& Name)
+	: FRHIComputeFence(Name)
+	, CurrentFence(-1)
+	, SignalFence(-1)
+	, LastCompletedFence(-1)
+	, FenceCore(nullptr)
+	, FD3D12DeviceChild(Parent)
+{
+}
 
+FD3D12Fence::~FD3D12Fence()
+{
+	Destroy();
 }
 
-void FD3D12Fence::CreateFence(ID3D12Device* pDirect3DDevice, uint64 InitialValue)
+void FD3D12Fence::Destroy()
 {
-	check(Fence == nullptr);
+	if (FenceCore)
+	{
+		//Return the underlying fence to the pool, store the last value signaled on this fence
+		GetParentDevice()->GetFenceCorePool().ReleaseFenceCore(FenceCore, SignalFence);
+		FenceCore = nullptr;
+	}
+}
+
+void FD3D12Fence::CreateFence(uint64 InitialValue)
+{
+	check(FenceCore == nullptr);
+
+	//Get a fence from the pool
+	FenceCore = GetParentDevice()->GetFenceCorePool().ObtainFenceCore(InitialValue);
+	ID3D12Fence* Fence = FenceCore->GetFence();
+	check(Fence);
 
-	VERIFYD3D11RESULT(pDirect3DDevice->CreateFence(InitialValue, D3D12_FENCE_FLAG_NONE, IID_PPV_ARGS(Fence.GetInitReference())));
 	SetName(Fence, *GetName().ToString());
 	LastCompletedFence = Fence->GetCompletedValue();
 	CurrentFence = SignalFence = LastCompletedFence + 1;
@@ -52,7 +80,7 @@ uint64 FD3D12Fence::Signal(ID3D12CommandQueue* pCommandQueue)
 {
 	check(pCommandQueue != nullptr);
 
-	VERIFYD3D11RESULT(pCommandQueue->Signal(Fence.GetReference(), CurrentFence));
+	VERIFYD3D11RESULT(pCommandQueue->Signal(FenceCore->GetFence(), CurrentFence));
 
 	// Save the current fence and increment it
 	SignalFence = CurrentFence++;
@@ -65,12 +93,12 @@ void FD3D12Fence::GpuWait(ID3D12CommandQueue* pCommandQueue, uint64 FenceValue)
 {
 	check(pCommandQueue != nullptr);
 
-	VERIFYD3D11RESULT(pCommandQueue->Wait(Fence.GetReference(), FenceValue));
+	VERIFYD3D11RESULT(pCommandQueue->Wait(FenceCore->GetFence(), FenceValue));
 }
 
 bool FD3D12Fence::IsFenceComplete(uint64 FenceValue)
 {
-	check(Fence != nullptr);
+	check(FenceCore);
 
 	// Avoid repeatedly calling GetCompletedValue()
 	if (FenceValue <= LastCompletedFence)
@@ -79,13 +107,13 @@ bool FD3D12Fence::IsFenceComplete(uint64 FenceValue)
 	}
 
 	// Refresh the completed fence value
-	LastCompletedFence = Fence->GetCompletedValue();
+	LastCompletedFence = FenceCore->GetFence()->GetCompletedValue();
 	return FenceValue <= LastCompletedFence;
 }
 
 uint64 FD3D12Fence::GetLastCompletedFence()
 {
-	LastCompletedFence = Fence->GetCompletedValue();
+	LastCompletedFence = FenceCore->GetFence()->GetCompletedValue();
 	return LastCompletedFence;
 }
 
@@ -93,7 +121,7 @@ void FD3D12Fence::WaitForFence(uint64 FenceValue)
 {
 	SCOPE_CYCLE_COUNTER(STAT_D3D12WaitForFenceTime);
 
-	check(Fence != nullptr);
+	check(FenceCore);
 
 	if (IsFenceComplete(FenceValue))
 	{
@@ -101,10 +129,10 @@ void FD3D12Fence::WaitForFence(uint64 FenceValue)
 	}
 
 	// We must wait.  Do so with an event handler so we don't oversleep.
-	VERIFYD3D11RESULT(Fence->SetEventOnCompletion(FenceValue, hFenceCompleteEvent));
+	VERIFYD3D11RESULT(FenceCore->GetFence()->SetEventOnCompletion(FenceValue, FenceCore->GetCompleteionEvent()));
 
 	// Wait for the event to complete (the event is automatically reset afterwards)
-	const uint32 WaitResult = WaitForSingleObject(hFenceCompleteEvent, INFINITE);
+	const uint32 WaitResult = WaitForSingleObject(FenceCore->GetCompleteionEvent(), INFINITE);
 	check(0 == WaitResult);
 
 	LastCompletedFence = FenceValue;
@@ -147,6 +175,10 @@ FD3D12CommandListManager::FD3D12CommandListManager(FD3D12Device* InParent, D3D12
 	, ResourceBarrierCommandAllocator(nullptr)
 	, ResourceBarrierCommandAllocatorManager(InParent, D3D12_COMMAND_LIST_TYPE_DIRECT)
 {
+	for (FD3D12Fence& Fence : Fences)
+	{
+		Fence.SetParentDevice(InParent);
+	}
 }
 
 FD3D12CommandListManager::~FD3D12CommandListManager()
@@ -166,6 +198,11 @@ void FD3D12CommandListManager::Destroy()
 	{
 		ReadyLists.Dequeue(hList);
 	}
+
+	for (FD3D12Fence& Fence : Fences)
+	{
+		Fence.Destroy();
+	}
 }
 
 void FD3D12CommandListManager::Create(uint32 NumCommandLists)
@@ -184,7 +221,7 @@ void FD3D12CommandListManager::Create(uint32 NumCommandLists)
 
 	for (uint32 i = 0; i < FT_NumTypes; i++)
 	{
-		Fences[i].CreateFence(Direct3DDevice, 0);
+		Fences[i].CreateFence(0);
 	}
 
 	if (NumCommandLists > 0)
@@ -312,6 +349,11 @@ void FD3D12CommandListManager::ExecuteCommandLists(TArray<FD3D12CommandListHandl
 	FD3D12CommandListHandle BarrierCommandList[128];
 	if (NeedsResourceBarriers)
 	{
+#if !USE_D3D12RHI_RESOURCE_STATE_TRACKING
+		// If we're using the engine's resource state tracking and barriers, then we should never have pending resource barriers.
+		check(false);
+#endif // !USE_D3D12RHI_RESOURCE_STATE_TRACKING
+
 		//#todo-rco: Need verification from MS
 #if 0//UE_BUILD_DEBUG	
 		if (!ResourceStateCS.TryLock())
@@ -528,3 +570,41 @@ FD3D12CommandListHandle FD3D12CommandListManager::CreateCommandListHandle(FD3D12
 	check(CommandListCount < MAX_ALLOCATED_COMMAND_LISTS);
 	return List;
 }
+
+FD3D12FenceCore* FD3D12FenceCorePool::ObtainFenceCore(uint64 InitialValue)
+{
+	{
+		FScopeLock Lock(&CS);
+		FD3D12FenceCore* Fence = nullptr;
+		if (AvailableFences.Peek(Fence))
+		{
+			if (Fence->IsAvailable())
+			{
+				AvailableFences.Dequeue(Fence);
+
+				//Reset the fence value
+				Fence->GetFence()->Signal(InitialValue);
+
+				return Fence;
+			}
+		}
+	}
+
+	return new FD3D12FenceCore(GetParentDevice(), InitialValue);
+}
+
+void FD3D12FenceCorePool::ReleaseFenceCore(FD3D12FenceCore* Fence, uint64 CurrentFenceValue)
+{
+	FScopeLock Lock(&CS);
+	Fence->FenceValueAvailableAt = CurrentFenceValue;
+	AvailableFences.Enqueue(Fence);
+}
+
+void FD3D12FenceCorePool::Destroy()
+{
+	FD3D12FenceCore* Fence = nullptr;
+	while (AvailableFences.Dequeue(Fence))
+	{
+		delete(Fence);
+	}
+}
diff --git a/Engine/Source/Runtime/Windows/D3D12RHI/Private/D3D12DirectCommandListManager.h b/Engine/Source/Runtime/Windows/D3D12RHI/Private/D3D12DirectCommandListManager.h
--- a/Engine/Source/Runtime/Windows/D3D12RHI/Private/D3D12DirectCommandListManager.h
+++ b/Engine/Source/Runtime/Windows/D3D12RHI/Private/D3D12DirectCommandListManager.h
@@ -20,13 +20,44 @@ enum EFenceType
 	FT_NumTypes
 };
 
-class FD3D12Fence : public FRHIComputeFence, public FNoncopyable
+class FD3D12FenceCore : public FD3D12DeviceChild
 {
 public:
-	FD3D12Fence(const FName& Name = L"<unnamed>");
+	FD3D12FenceCore(FD3D12Device* Parent, uint64 InitialValue);
+	~FD3D12FenceCore();
+
+	inline ID3D12Fence* GetFence() const { return Fence.GetReference(); }
+	inline HANDLE GetCompleteionEvent() const { return hFenceCompleteEvent; }
+	inline bool IsAvailable() const { return FenceValueAvailableAt >= Fence->GetCompletedValue(); }
+
+	uint32 FenceValueAvailableAt;
+private:
+	TRefCountPtr<ID3D12Fence> Fence;
+	HANDLE hFenceCompleteEvent;
+};
+
+class FD3D12FenceCorePool : public FD3D12DeviceChild
+{
+public:
+
+	FD3D12FenceCorePool(FD3D12Device* Parent) : FD3D12DeviceChild(Parent) {};
+
+	FD3D12FenceCore* ObtainFenceCore(uint64 InitialValue);
+	void ReleaseFenceCore(FD3D12FenceCore* Fence, uint64 CurrentFenceValue);
+	void Destroy();
+
+private:
+	FCriticalSection CS;
+	TQueue<FD3D12FenceCore*> AvailableFences;
+};
+
+class FD3D12Fence : public FRHIComputeFence, public FNoncopyable, public FD3D12DeviceChild
+{
+public:
+	FD3D12Fence(FD3D12Device* Parent = nullptr, const FName& Name = L"<unnamed>");
 	~FD3D12Fence();
 
-	void CreateFence(ID3D12Device* pDirect3DDevice, uint64 InitialValue = 0);
+	void CreateFence(uint64 InitialValue = 0);
 	uint64 Signal(ID3D12CommandQueue* pCommandQueue);
 	void GpuWait(ID3D12CommandQueue* pCommandQueue, uint64 FenceValue);
 	bool IsFenceComplete(uint64 FenceValue);
@@ -36,12 +67,14 @@ public:
 	uint64 GetSignalFence() const { return SignalFence; }
 	uint64 GetLastCompletedFence();
 
+	void Destroy();
+
 private:
-	TRefCountPtr<ID3D12Fence> Fence;
 	uint64 CurrentFence;
 	uint64 SignalFence;
 	uint64 LastCompletedFence;
-	HANDLE hFenceCompleteEvent;
+
+	FD3D12FenceCore* FenceCore;
 };
 
 class FD3D12CommandAllocatorManager : public FD3D12DeviceChild
diff --git a/Engine/Source/Runtime/Windows/D3D12RHI/Private/D3D12IndexBuffer.cpp b/Engine/Source/Runtime/Windows/D3D12RHI/Private/D3D12IndexBuffer.cpp
--- a/Engine/Source/Runtime/Windows/D3D12RHI/Private/D3D12IndexBuffer.cpp
+++ b/Engine/Source/Runtime/Windows/D3D12RHI/Private/D3D12IndexBuffer.cpp
@@ -6,13 +6,9 @@
 
 #include "D3D12RHIPrivate.h"
 
-FIndexBufferRHIRef FD3D12DynamicRHI::RHICreateIndexBuffer(uint32 Stride, uint32 Size, uint32 InUsage, FRHIResourceCreateInfo& CreateInfo)
+D3D12_RESOURCE_DESC CreateIndexBufferResourceDesc(uint32 Size, uint32 InUsage)
 {
-	// Explicitly check that the size is nonzero before allowing CreateIndexBuffer to opaquely fail.
-	check(Size > 0);
-
-	// Describe the index buffer.
-	const bool bIsDynamic = (InUsage & BUF_AnyDynamic) ? true : false;
+	// Describe the vertex buffer.
 	D3D12_RESOURCE_DESC Desc = CD3DX12_RESOURCE_DESC::Buffer(Size);
 
 	if (InUsage & BUF_UnorderedAccess)
@@ -25,45 +21,7 @@ FIndexBufferRHIRef FD3D12DynamicRHI::RHICreateIndexBuffer(uint32 Stride, uint32
 		Desc.Flags |= D3D12_RESOURCE_FLAG_DENY_SHADER_RESOURCE;
 	}
 
-	// If a resource array was provided for the resource, create the resource pre-populated
-	D3D12_SUBRESOURCE_DATA InitData ={0};
-	D3D12_SUBRESOURCE_DATA* pInitData = nullptr;
-	if (CreateInfo.ResourceArray)
-	{
-		check(Size == CreateInfo.ResourceArray->GetResourceDataSize());
-		InitData.pData = CreateInfo.ResourceArray->GetResourceData();
-		InitData.RowPitch = Size;
-		InitData.SlicePitch = 0;
-		pInitData = &InitData;
-	}
-
-	TRefCountPtr<FD3D12ResourceLocation> ResourceLocation;
-	if (bIsDynamic)
-	{
-		// Create an empty resource location
-		ResourceLocation = new FD3D12ResourceLocation(GetRHIDevice(), Size);
-
-		if (pInitData)
-		{
-			void* pData = GetRHIDevice()->GetDefaultUploadHeapAllocator().AllocUploadResource(Size, DEFAULT_CONTEXT_UPLOAD_POOL_ALIGNMENT, ResourceLocation);
-			FMemory::Memcpy(pData, InitData.pData, Size);
-		}
-	}
-	else
-	{
-		ResourceLocation = new FD3D12ResourceLocation(GetRHIDevice(), Size);
-		VERIFYD3D11RESULT(GetRHIDevice()->GetDefaultBufferAllocator().AllocDefaultResource(Desc, pInitData, ResourceLocation));
-	}
-
-	if (CreateInfo.ResourceArray)
-	{
-		// Discard the resource array's contents.
-		CreateInfo.ResourceArray->Discard();
-	}
-
-	UpdateBufferStats(ResourceLocation, true, D3D12_BUFFER_TYPE_INDEX);
-
-	return new FD3D12IndexBuffer(GetRHIDevice(), ResourceLocation, Stride, Size, InUsage);
+	return Desc;
 }
 
 FD3D12IndexBuffer::~FD3D12IndexBuffer()
@@ -73,146 +31,65 @@ FD3D12IndexBuffer::~FD3D12IndexBuffer()
 	// Dynamic resources have a permanent entry in the OutstandingLocks that needs to be cleaned up. 
 	if (GetUsage() & BUF_AnyDynamic)
 	{
-		FD3D12LockedKey LockedKey(ResourceLocation.GetReference());
-		FD3D12DynamicRHI::GetD3DRHI()->OutstandingLocks.Remove(LockedKey);
+		FD3D12LockedKey LockedKey(this);
+		FD3D12DynamicRHI::GetD3DRHI()->RemoveFromOutstandingLocks(LockedKey);
 	}
 }
 
-void* FD3D12DynamicRHI::RHILockIndexBuffer(FIndexBufferRHIParamRef IndexBufferRHI, uint32 Offset, uint32 Size, EResourceLockMode LockMode)
+void FD3D12IndexBuffer::Rename(FD3D12ResourceLocation* NewResource)
 {
-	FD3D12IndexBuffer* IndexBuffer = FD3D12DynamicRHI::ResourceCast(IndexBufferRHI);
-
-	FD3D12CommandContext& DefaultContext = GetRHIDevice()->GetDefaultCommandContext();
+	FD3D12CommandContext& DefaultContext = GetParentDevice()->GetDefaultCommandContext();
 
 	// If this resource is bound to the device, unbind it
-	DefaultContext.ConditionalClearShaderResource(IndexBuffer->ResourceLocation);
-
-	// Determine whether the index buffer is dynamic or not.
-	const bool bIsDynamic = (IndexBuffer->GetUsage() & BUF_AnyDynamic) ? true : false;
-
-	FD3D12LockedKey LockedKey(IndexBuffer->ResourceLocation.GetReference());
-	FD3D12LockedData LockedData;
-	uint32 ResourceSize = IndexBuffer->ResourceLocation->GetEffectiveBufferSize();
-
-	if (bIsDynamic)
-	{
-		check(LockMode == RLM_WriteOnly);
-
-		// Use an upload heap for dynamic resources and map its memory for writing.
-		void* pData = GetRHIDevice()->GetDefaultUploadHeapAllocator().AllocUploadResource(Offset + Size, DEFAULT_CONTEXT_UPLOAD_POOL_ALIGNMENT, IndexBuffer->ResourceLocation.GetReference());
-		check(pData);
+	DefaultContext.ConditionalClearShaderResource(ResourceLocation);
 
-		// Add the lock to the lock map.
-		FD3D12LockedKey LockedKey(IndexBuffer->ResourceLocation.GetReference());
-		FD3D12LockedData LockedData;
-		LockedData.SetData(pData);
-		LockedData.Pitch = Offset + Size;
-		OutstandingLocks.Add(LockedKey, LockedData);
-
-		return (void*)((uint8*)pData + Offset);
-	}
-	else
-	{
-		if (LockMode == RLM_ReadOnly)
-		{
-			FD3D12Resource* Resource = IndexBuffer->ResourceLocation->GetResource();
-
-			// If the static buffer is being locked for reading, create a staging buffer.
-			TRefCountPtr<FD3D12Resource> StagingIndexBuffer;
-			VERIFYD3D11RESULT(GetRHIDevice()->GetResourceHelper().CreateBuffer(D3D12_HEAP_TYPE_READBACK, Offset + Size, StagingIndexBuffer.GetInitReference()));
-			LockedData.StagingResource = StagingIndexBuffer;
-
-			// Copy the contents of the index buffer to the staging buffer.
-			DefaultContext.numCopies++;
-			DefaultContext.CommandListHandle->CopyResource(StagingIndexBuffer->GetResource(), Resource->GetResource());
-
-			// Execute the command list and wait for the contents to be ready so it can be returned to the caller
-			GetRHIDevice()->GetDefaultCommandContext().FlushCommands(true);
-
-			// Map the staging buffer's memory for reading.
-			void* pData;
-			VERIFYD3D11RESULT(StagingIndexBuffer->GetResource()->Map(0, nullptr, &pData));
-			LockedData.SetData(pData);
-			LockedData.Pitch = Offset + Size;
-		}
-		else
-		{
-			// If the static buffer is being locked for writing, allocate memory for the contents to be written to.
-
-			// Use an upload heap to copy data to a default resource.
-			TRefCountPtr<FD3D12ResourceLocation> UploadBuffer = new FD3D12ResourceLocation(GetRHIDevice());
-			void* pData = GetRHIDevice()->GetDefaultFastAllocator().Allocate(Offset + Size, D3D12_TEXTURE_DATA_PLACEMENT_ALIGNMENT, UploadBuffer);
-
-			// Add the lock to the lock map.
-			LockedData.SetData(pData);
-			LockedData.Pitch = Offset + Size;
-
-			// Keep track of the underlying resource for the upload heap so it can be referenced during Unmap.
-			LockedData.UploadHeapLocation = UploadBuffer;
-		}
-	}
+	ResourceLocation = NewResource;
+}
 
-	// Add the lock to the lock map.
-	OutstandingLocks.Add(LockedKey, LockedData);
+FIndexBufferRHIRef FD3D12DynamicRHI::RHICreateIndexBuffer(uint32 Stride, uint32 Size, uint32 InUsage, FRHIResourceCreateInfo& CreateInfo)
+{
+	const D3D12_RESOURCE_DESC Desc = CreateIndexBufferResourceDesc(Size, InUsage);
+	const uint32 Alignment = 4;
+
+	FD3D12ResourceLocation* ResourceLocation = CreateBuffer(nullptr, Desc, Size, InUsage, CreateInfo, Alignment);
+	FD3D12IndexBuffer* NewBuffer = new FD3D12IndexBuffer(GetRHIDevice(), ResourceLocation, Stride, Size, InUsage);
+	UpdateBufferStats(NewBuffer->ResourceLocation, true, D3D12_BUFFER_TYPE_INDEX);
+	NewBuffer->BufferAlignment = Alignment;
+	return NewBuffer;
+}
 
-	// Return the offset pointer
-	return (void*)((uint8*)LockedData.GetData() + Offset);
+void* FD3D12DynamicRHI::RHILockIndexBuffer(FIndexBufferRHIParamRef IndexBufferRHI, uint32 Offset, uint32 Size, EResourceLockMode LockMode)
+{
+	return LockBuffer(nullptr, FD3D12DynamicRHI::ResourceCast(IndexBufferRHI), Offset, Size, LockMode);
 }
 
 void FD3D12DynamicRHI::RHIUnlockIndexBuffer(FIndexBufferRHIParamRef IndexBufferRHI)
 {
-	FD3D12IndexBuffer* IndexBuffer = FD3D12DynamicRHI::ResourceCast(IndexBufferRHI);
-
-	// Determine whether the index buffer is dynamic or not.
-	const bool bIsDynamic = (IndexBuffer->GetUsage() & BUF_AnyDynamic) ? true : false;
+	UnlockBuffer(nullptr, FD3D12DynamicRHI::ResourceCast(IndexBufferRHI));
+}
 
-	// Find the outstanding lock for this IB.
-	FD3D12LockedKey LockedKey(IndexBuffer->ResourceLocation.GetReference());
-	FD3D12LockedData* LockedData = OutstandingLocks.Find(LockedKey);
-	check(LockedData);
+FIndexBufferRHIRef FD3D12DynamicRHI::CreateIndexBuffer_RenderThread(class FRHICommandListImmediate& RHICmdList, uint32 Stride, uint32 Size, uint32 InUsage, FRHIResourceCreateInfo& CreateInfo)
+{
+	const D3D12_RESOURCE_DESC Desc = CreateIndexBufferResourceDesc(Size, InUsage);
+	const uint32 Alignment = 4;
+
+	FD3D12ResourceLocation* ResourceLocation = CreateBuffer(&RHICmdList, Desc, Size, InUsage, CreateInfo, Alignment);
+	FD3D12IndexBuffer* NewBuffer = new FD3D12IndexBuffer(GetRHIDevice(), ResourceLocation, Stride, Size, InUsage);
+	UpdateBufferStats(NewBuffer->ResourceLocation, true, D3D12_BUFFER_TYPE_INDEX);
+	NewBuffer->BufferAlignment = Alignment;
+	return NewBuffer;
+}
 
-	FD3D12Resource* Resource = IndexBuffer->ResourceLocation->GetResource();
+void* FD3D12DynamicRHI::LockIndexBuffer_RenderThread(class FRHICommandListImmediate& RHICmdList, FIndexBufferRHIParamRef IndexBufferRHI, uint32 Offset, uint32 SizeRHI, EResourceLockMode LockMode)
+{
+	return LockBuffer(&RHICmdList, FD3D12DynamicRHI::ResourceCast(IndexBufferRHI), Offset, SizeRHI, LockMode);
+}
 
-	if (bIsDynamic)
-	{
-		// If the IB is dynamic, its upload heap memory can always stay mapped. Don't do anything.
-	}
-	else
-	{
-		// If the static IB lock involved a staging resource, it was locked for reading.
-		if (LockedData->StagingResource)
-		{
-			// Unmap the staging buffer's memory.
-			FD3D12Resource* StagingBuffer = LockedData->StagingResource.GetReference();
-			StagingBuffer->GetResource()->Unmap(0, nullptr);
-		}
-		else
-		{
-			check(LockedData->UploadHeapLocation != nullptr);
-
-			// Copy the contents of the temporary memory buffer allocated for writing into the IB.
-			FD3D12ResourceLocation* UploadHeap = LockedData->UploadHeapLocation;
-
-			FD3D12CommandListHandle& hCommandList = GetRHIDevice()->GetDefaultCommandContext().CommandListHandle;
-
-			{
-				FScopeResourceBarrier ScopeResourceBarrier(hCommandList, Resource, D3D12_RESOURCE_STATE_GENERIC_READ, D3D12_RESOURCE_STATE_COPY_DEST, 0);
-				// Don't need to transition upload heaps
-
-				GetRHIDevice()->GetDefaultCommandContext().numCopies++;
-				hCommandList->CopyBufferRegion(
-					Resource->GetResource(), IndexBuffer->ResourceLocation->GetOffset(),
-					UploadHeap->GetResource()->GetResource(), UploadHeap->GetOffset(),
-					LockedData->Pitch);
-			}
-
-
-			DEBUG_RHI_EXECUTE_COMMAND_LIST(this);
-		}
-	}
+void FD3D12DynamicRHI::UnlockIndexBuffer_RenderThread(class FRHICommandListImmediate& RHICmdList, FIndexBufferRHIParamRef IndexBufferRHI)
+{
+	// Pull down the above RHI implementation so that we can flush only when absolutely necessary
+	QUICK_SCOPE_CYCLE_COUNTER(STAT_FDynamicRHI_UnlockIndexBuffer_RenderThread);
+	check(IsInRenderingThread());
 
-	// Remove the FD3D12LockedData from the lock map.
-	// If the lock involved a staging resource, this releases it.
-	// If the lock involved an upload heap resource, this releases it.
-	OutstandingLocks.Remove(LockedKey);
+	UnlockBuffer(&RHICmdList, FD3D12DynamicRHI::ResourceCast(IndexBufferRHI));
 }
