Merging //UE4/Release-4.11 to //UE4/Main (Up to CL#2897738)

diff --git a/Engine/Source/Runtime/Windows/D3D12RHI/Private/D3D12RHI.cpp b/Engine/Source/Runtime/Windows/D3D12RHI/Private/D3D12RHI.cpp
--- a/Engine/Source/Runtime/Windows/D3D12RHI/Private/D3D12RHI.cpp
+++ b/Engine/Source/Runtime/Windows/D3D12RHI/Private/D3D12RHI.cpp
@@ -126,7 +126,7 @@ FD3D12CommandListHandle FD3D12CommandContext::FlushCommands(bool WaitForCompleti
 	check(IsDefaultContext());
 
 	FD3D12Device* Device = GetParentDevice();
-	const bool bExecutePendingWork = GCommandListBatchingMode == CLB_AggressiveBatching;
+	const bool bExecutePendingWork = GCommandListBatchingMode == CLB_AggressiveBatching || GetParentDevice()->IsGPUIdle();
 	const bool bHasPendingWork = bExecutePendingWork && (Device->PendingCommandListsTotalWorkCommands > 0);
 	const bool bHasDoneWork = HasDoneWork() || bHasPendingWork;
 
@@ -879,6 +879,20 @@ void FD3D12DynamicRHI::GetLocalVideoMemoryInfo(DXGI_QUERY_VIDEO_MEMORY_INFO* Loc
 	VERIFYD3D11RESULT(MainDevice->GetAdapter3()->QueryVideoMemoryInfo(0, DXGI_MEMORY_SEGMENT_GROUP_LOCAL, LocalVideoMemoryInfo));
 }
 
+void  FD3D12DynamicRHI::UpdateBuffer(FD3D12Resource* Dest, uint32 DestOffset, FD3D12Resource* Source, uint32 SourceOffset, uint32 NumBytes)
+{
+	FD3D12CommandContext& defaultContext = FD3D12DynamicRHI::GetD3DRHI()->GetRHIDevice()->GetDefaultCommandContext();
+	FD3D12CommandListHandle& hCommandList = defaultContext.CommandListHandle;
+
+	FScopeResourceBarrier ScopeResourceBarrierDest(hCommandList, Dest, Dest->GetDefaultResourceState(), D3D12_RESOURCE_STATE_COPY_DEST, 0);
+	// Don't need to transition upload heaps
+
+	defaultContext.numCopies++;
+	hCommandList->CopyBufferRegion(Dest->GetResource(), DestOffset, Source->GetResource(), SourceOffset, NumBytes);
+
+	DEBUG_RHI_EXECUTE_COMMAND_LIST(this);
+}
+
 #if SUPPORTS_MEMORY_RESIDENCY
 int32 bLogMemoryResidency = 1;
 static FAutoConsoleVariableRef CVarLogMemoryResidency(
diff --git a/Engine/Source/Runtime/Windows/D3D12RHI/Private/D3D12RHIPrivate.h b/Engine/Source/Runtime/Windows/D3D12RHI/Private/D3D12RHIPrivate.h
--- a/Engine/Source/Runtime/Windows/D3D12RHI/Private/D3D12RHIPrivate.h
+++ b/Engine/Source/Runtime/Windows/D3D12RHI/Private/D3D12RHIPrivate.h
@@ -41,6 +41,7 @@ DECLARE_LOG_CATEGORY_EXTERN(LogD3D12RHI, Log, All);
 #include "../Public/D3D12ConstantBuffer.h"
 #include "D3D12StateCache.h"
 #include "D3D12DirectCommandListManager.h"
+#include "D3D12Allocation.h"
 
 #if UE_BUILD_SHIPPING || UE_BUILD_TEST
 #define CHECK_SRV_TRANSITIONS 0
@@ -52,22 +53,24 @@ DECLARE_LOG_CATEGORY_EXTERN(LogD3D12RHI, Log, All);
 #define DX_MAX_MSAA_COUNT 8
 
 // Definitions.
+#define USE_D3D12RHI_RESOURCE_STATE_TRACKING 1	// Fully relying on the engine's resource barriers is a work in progress. For now, continue to use the D3D12 RHI's resource state tracking.
 #define EXECUTE_DEBUG_COMMAND_LISTS 0
 #define DEBUG_RESOURCE_STATES 0
 #define ENABLE_PLACED_RESOURCES 0 // Disabled due to a couple of NVidia bugs related to placed resources. Works fine on Intel
 #define REMOVE_OLD_QUERY_BATCHES 1  // D3D12: MSFT: TODO: Works around a suspected UE4 InfiltratorDemo bug where a query is never released
 
-#define DEFAULT_BUFFER_POOL_SIZE (128 * 1024 * 1024)
-#define DEFAULT_BUFFER_POOL_MAX_ALLOC_SIZE (128 * 1024)
+#define DEFAULT_BUFFER_POOL_SIZE (16 * 1024 * 1024)
+#define DEFAULT_BUFFER_POOL_MAX_ALLOC_SIZE (512 * 1024)
 #define DEFAULT_CONTEXT_UPLOAD_POOL_SIZE (64 * 1024 * 1024)
 #define DEFAULT_CONTEXT_UPLOAD_POOL_MAX_ALLOC_SIZE (1024 * 1024 * 4)
 #define DEFAULT_CONTEXT_UPLOAD_POOL_ALIGNMENT (D3D12_CONSTANT_BUFFER_DATA_PLACEMENT_ALIGNMENT)
 
-
 #if DEBUG_RESOURCE_STATES
 #define LOG_EXECUTE_COMMAND_LISTS 1
+#define ASSERT_RESOURCE_STATES 1
 #else
 #define LOG_EXECUTE_COMMAND_LISTS 0
+#define ASSERT_RESOURCE_STATES 0
 #endif
 
 #if EXECUTE_DEBUG_COMMAND_LISTS
@@ -733,6 +736,7 @@ public:
 	virtual void RHISetMultipleViewports(uint32 Count, const FViewportBounds* Data) final override;
 	virtual void RHIClearUAV(FUnorderedAccessViewRHIParamRef UnorderedAccessViewRHI, const uint32* Values) final override;
 	virtual void RHICopyToResolveTarget(FTextureRHIParamRef SourceTexture, FTextureRHIParamRef DestTexture, bool bKeepOriginalSurface, const FResolveParams& ResolveParams) final override;
+	virtual void RHITransitionResources(EResourceTransitionAccess TransitionType, FTextureRHIParamRef* InTextures, int32 NumTextures) final override;
 	virtual void RHIBeginRenderQuery(FRenderQueryRHIParamRef RenderQuery) final override;
 	virtual void RHIEndRenderQuery(FRenderQueryRHIParamRef RenderQuery) final override;
 	virtual void RHIBeginOcclusionQueryBatch() final override;
@@ -777,7 +781,7 @@ public:
 	virtual void RHISetBlendState(FBlendStateRHIParamRef NewState, const FLinearColor& BlendFactor) final override;
 	virtual void RHISetRenderTargets(uint32 NumSimultaneousRenderTargets, const FRHIRenderTargetView* NewRenderTargets, const FRHIDepthRenderTargetView* NewDepthStencilTarget, uint32 NumUAVs, const FUnorderedAccessViewRHIParamRef* UAVs) final override;
 	virtual void RHISetRenderTargetsAndClear(const FRHISetRenderTargetsInfo& RenderTargetsInfo) final override;
-	//virtual void RHIBindClearMRTValues(bool bClearColor, bool bClearDepth, bool bClearStencil) final override;
+	virtual void RHIBindClearMRTValues(bool bClearColor, bool bClearDepth, bool bClearStencil) final override;
 	virtual void RHIDrawPrimitive(uint32 PrimitiveType, uint32 BaseVertexIndex, uint32 NumPrimitives, uint32 NumInstances) final override;
 	virtual void RHIDrawPrimitiveIndirect(uint32 PrimitiveType, FVertexBufferRHIParamRef ArgumentBuffer, uint32 ArgumentOffset) final override;
 	virtual void RHIDrawIndexedIndirect(FIndexBufferRHIParamRef IndexBufferRHI, uint32 PrimitiveType, FStructuredBufferRHIParamRef ArgumentsBufferRHI, int32 DrawArgumentsIndex, uint32 NumInstances) final override;
@@ -905,6 +909,8 @@ public:
 	inline FD3D12GlobalOnlineHeap&          GetGlobalSamplerHeap() { return GlobalSamplerHeap; }
 	inline FD3D12GlobalOnlineHeap&          GetGlobalViewHeap() { return GlobalViewHeap; }
 
+	bool IsGPUIdle();
+
 	inline const D3D12_HEAP_PROPERTIES &GetConstantBufferPageProperties() { return ConstantBufferPageProperties; }
 
 	inline uint32 GetNumContexts() { return CommandContextArray.Num(); }
@@ -938,6 +944,8 @@ public:
 	inline D3D12_RESOURCE_HEAP_TIER    GetResourceHeapTier() { return ResourceHeapTier; }
 	inline D3D12_RESOURCE_BINDING_TIER GetResourceBindingTier() { return ResourceBindingTier; }
 
+	inline FD3D12FenceCorePool& GetFenceCorePool() { return FenceCorePool; }
+
 	TArray<FD3D12CommandListHandle> PendingCommandLists;
 	uint32 PendingCommandListsTotalWorkCommands;
 
@@ -1036,6 +1044,8 @@ protected:
 	FD3D12DynamicHeapAllocator	DefaultUploadHeapAllocator;
 
 	FD3D12TextureAllocatorPool TextureAllocator;
+
+	FD3D12FenceCorePool FenceCorePool;
 };
 
 /** The interface which is implemented by the dynamically bound RHI. */
@@ -1049,8 +1059,39 @@ public:
 
 	static FD3D12DynamicRHI* GetD3DRHI() { return SingleD3DRHI; }
 
+private:
 	/** Global D3D12 lock list */
 	TMap<FD3D12LockedKey, FD3D12LockedData> OutstandingLocks;
+	FCriticalSection OutstandingLocksCS;
+
+	/** Texture pool size */
+	int64 RequestedTexturePoolSize;
+
+public:
+
+	inline void ClearOutstandingLocks()
+	{
+		FScopeLock Lock(&OutstandingLocksCS);
+		OutstandingLocks.Empty();
+	}
+
+	inline void AddToOutstandingLocks(FD3D12LockedKey& Key, FD3D12LockedData& Value)
+	{
+		FScopeLock Lock(&OutstandingLocksCS);
+		OutstandingLocks.Add(Key, Value);
+	}
+
+	inline void RemoveFromOutstandingLocks(FD3D12LockedKey& Key)
+	{
+		FScopeLock Lock(&OutstandingLocksCS);
+		OutstandingLocks.Remove(Key);
+	}
+
+	inline FD3D12LockedData* FindInOutstandingLocks(FD3D12LockedKey& Key)
+	{
+		FScopeLock Lock(&OutstandingLocksCS);
+		return OutstandingLocks.Find(Key);
+	}
 
 	/** Initialization constructor. */
 	FD3D12DynamicRHI(IDXGIFactory4* InDXGIFactory, FD3D12Adapter& InAdapter);
@@ -1165,6 +1206,67 @@ public:
 	virtual class IRHIComputeContext* RHIGetDefaultAsyncComputeContext() final override;
 	virtual class IRHICommandContextContainer* RHIGetCommandContextContainer() final override;
 
+
+	//
+	// The Following functions are the _RenderThread version of the above functions. They allow the RHI to control the thread synchronization for greater efficiency.
+	// These will be un-commented as they are implemented.
+	//
+	virtual FVertexBufferRHIRef CreateVertexBuffer_RenderThread(class FRHICommandListImmediate& RHICmdList, uint32 Size, uint32 InUsage, FRHIResourceCreateInfo& CreateInfo);
+	virtual void* LockVertexBuffer_RenderThread(class FRHICommandListImmediate& RHICmdList, FVertexBufferRHIParamRef VertexBuffer, uint32 Offset, uint32 SizeRHI, EResourceLockMode LockMode);
+	virtual void UnlockVertexBuffer_RenderThread(class FRHICommandListImmediate& RHICmdList, FVertexBufferRHIParamRef VertexBuffer);
+	// virtual FTexture2DRHIRef AsyncReallocateTexture2D_RenderThread(class FRHICommandListImmediate& RHICmdList, FTexture2DRHIParamRef Texture2D, int32 NewMipCount, int32 NewSizeX, int32 NewSizeY, FThreadSafeCounter* RequestStatus);
+	// virtual ETextureReallocationStatus FinalizeAsyncReallocateTexture2D_RenderThread(class FRHICommandListImmediate& RHICmdList, FTexture2DRHIParamRef Texture2D, bool bBlockUntilCompleted);
+	// virtual ETextureReallocationStatus CancelAsyncReallocateTexture2D_RenderThread(class FRHICommandListImmediate& RHICmdList, FTexture2DRHIParamRef Texture2D, bool bBlockUntilCompleted);
+	virtual FIndexBufferRHIRef CreateIndexBuffer_RenderThread(class FRHICommandListImmediate& RHICmdList, uint32 Stride, uint32 Size, uint32 InUsage, FRHIResourceCreateInfo& CreateInfo);
+	virtual void* LockIndexBuffer_RenderThread(class FRHICommandListImmediate& RHICmdList, FIndexBufferRHIParamRef IndexBuffer, uint32 Offset, uint32 SizeRHI, EResourceLockMode LockMode);
+	virtual void UnlockIndexBuffer_RenderThread(class FRHICommandListImmediate& RHICmdList, FIndexBufferRHIParamRef IndexBuffer);
+	// virtual FVertexDeclarationRHIRef CreateVertexDeclaration_RenderThread(class FRHICommandListImmediate& RHICmdList, const FVertexDeclarationElementList& Elements);
+	// virtual void UpdateTexture2D_RenderThread(class FRHICommandListImmediate& RHICmdList, FTexture2DRHIParamRef Texture, uint32 MipIndex, const struct FUpdateTextureRegion2D& UpdateRegion, uint32 SourcePitch, const uint8* SourceData);
+	// virtual void* LockTexture2D_RenderThread(class FRHICommandListImmediate& RHICmdList, FTexture2DRHIParamRef Texture, uint32 MipIndex, EResourceLockMode LockMode, uint32& DestStride, bool bLockWithinMiptail, bool bNeedsDefaultRHIFlush = true);
+	// virtual void UnlockTexture2D_RenderThread(class FRHICommandListImmediate& RHICmdList, FTexture2DRHIParamRef Texture, uint32 MipIndex, bool bLockWithinMiptail, bool bNeedsDefaultRHIFlush = true);
+	// 
+	// virtual FTexture2DRHIRef RHICreateTexture2D_RenderThread(class FRHICommandListImmediate& RHICmdList, uint32 SizeX, uint32 SizeY, uint8 Format, uint32 NumMips, uint32 NumSamples, uint32 Flags, FRHIResourceCreateInfo& CreateInfo);
+	// virtual FTexture2DArrayRHIRef RHICreateTexture2DArray_RenderThread(class FRHICommandListImmediate& RHICmdList, uint32 SizeX, uint32 SizeY, uint32 SizeZ, uint8 Format, uint32 NumMips, uint32 Flags, FRHIResourceCreateInfo& CreateInfo);
+	// virtual FTexture3DRHIRef RHICreateTexture3D_RenderThread(class FRHICommandListImmediate& RHICmdList, uint32 SizeX, uint32 SizeY, uint32 SizeZ, uint8 Format, uint32 NumMips, uint32 Flags, FRHIResourceCreateInfo& CreateInfo);
+	virtual FUnorderedAccessViewRHIRef RHICreateUnorderedAccessView_RenderThread(class FRHICommandListImmediate& RHICmdList, FStructuredBufferRHIParamRef StructuredBuffer, bool bUseUAVCounter, bool bAppendBuffer);
+	virtual FUnorderedAccessViewRHIRef RHICreateUnorderedAccessView_RenderThread(class FRHICommandListImmediate& RHICmdList, FTextureRHIParamRef Texture, uint32 MipLevel);
+	virtual FUnorderedAccessViewRHIRef RHICreateUnorderedAccessView_RenderThread(class FRHICommandListImmediate& RHICmdList, FVertexBufferRHIParamRef VertexBuffer, uint8 Format);
+	virtual FShaderResourceViewRHIRef RHICreateShaderResourceView_RenderThread(class FRHICommandListImmediate& RHICmdList, FTexture2DRHIParamRef Texture2DRHI, uint8 MipLevel);
+	virtual FShaderResourceViewRHIRef RHICreateShaderResourceView_RenderThread(class FRHICommandListImmediate& RHICmdList, FTexture2DRHIParamRef Texture2DRHI, uint8 MipLevel, uint8 NumMipLevels, uint8 Format);
+	virtual FShaderResourceViewRHIRef RHICreateShaderResourceView_RenderThread(class FRHICommandListImmediate& RHICmdList, FTexture3DRHIParamRef Texture3DRHI, uint8 MipLevel);
+	virtual FShaderResourceViewRHIRef RHICreateShaderResourceView_RenderThread(class FRHICommandListImmediate& RHICmdList, FTexture2DArrayRHIParamRef Texture2DArrayRHI, uint8 MipLevel);
+	virtual FShaderResourceViewRHIRef RHICreateShaderResourceView_RenderThread(class FRHICommandListImmediate& RHICmdList, FTextureCubeRHIParamRef TextureCubeRHI, uint8 MipLevel);
+	virtual FShaderResourceViewRHIRef RHICreateShaderResourceView_RenderThread(class FRHICommandListImmediate& RHICmdList, FVertexBufferRHIParamRef VertexBuffer, uint32 Stride, uint8 Format);
+	virtual FShaderResourceViewRHIRef RHICreateShaderResourceView_RenderThread(class FRHICommandListImmediate& RHICmdList, FStructuredBufferRHIParamRef StructuredBuffer);
+	// virtual FTextureCubeRHIRef RHICreateTextureCube_RenderThread(class FRHICommandListImmediate& RHICmdList, uint32 Size, uint8 Format, uint32 NumMips, uint32 Flags, FRHIResourceCreateInfo& CreateInfo);
+	// virtual FTextureCubeRHIRef RHICreateTextureCubeArray_RenderThread(class FRHICommandListImmediate& RHICmdList, uint32 Size, uint32 ArraySize, uint8 Format, uint32 NumMips, uint32 Flags, FRHIResourceCreateInfo& CreateInfo);
+	// virtual FRenderQueryRHIRef RHICreateRenderQuery_RenderThread(class FRHICommandListImmediate& RHICmdList, ERenderQueryType QueryType);
+
+	FD3D12ResourceLocation* CreateBuffer(FRHICommandListImmediate* RHICmdList, const D3D12_RESOURCE_DESC Desc, uint32 Size, uint32 InUsage, FRHIResourceCreateInfo& CreateInfo, uint32 Alignment);
+
+	template<class BufferType>
+	void* LockBuffer(FRHICommandListImmediate* RHICmdList, BufferType* Buffer, uint32 Offset, uint32 Size, EResourceLockMode LockMode);
+
+	template<class BufferType>
+	void UnlockBuffer(FRHICommandListImmediate* RHICmdList, BufferType* Buffer);
+
+	inline bool ShouldDeferBufferLockOperation(FRHICommandList* RHICmdList)
+	{
+		if (RHICmdList == nullptr)
+		{
+			return false;
+		}
+
+		if (RHICmdList->Bypass() || !GRHIThread)
+		{
+			return false;
+		}
+
+		return true;
+	}
+
+	void UpdateBuffer(FD3D12Resource* Dest, uint32 DestOffset, FD3D12Resource* Source, uint32 SourceOffset, uint32 NumBytes);
+
 #if UE_BUILD_DEBUG	
 	uint32 SubmissionLockStalls;
 	uint32 DrawCount;
@@ -1182,6 +1284,8 @@ public:
 		PerRHISetup(MainDevice);
 	}
 
+	inline void UpdataTextureMemorySize(int64 TextureSizeInKiloBytes) { FPlatformAtomics::InterlockedAdd(&GCurrentTextureMemorySize, TextureSizeInKiloBytes); }
+
 	/** Determine if an two views intersect */
 	template <class LeftT, class RightT>
 	static inline bool ResourceViewsIntersect(FD3D12View<LeftT>* pLeftView, FD3D12View<RightT>* pRightView)
@@ -1226,6 +1330,7 @@ public:
 	/** Transition a resource's state based on a Render target view */
 	static inline void TransitionResource(FD3D12CommandListHandle& hCommandList, FD3D12RenderTargetView* pView, D3D12_RESOURCE_STATES after)
 	{
+#if USE_D3D12RHI_RESOURCE_STATE_TRACKING || ASSERT_RESOURCE_STATES
 		FD3D12Resource* pResource = pView->GetResource();
 
 		const D3D12_RENDER_TARGET_VIEW_DESC &desc = pView->GetDesc();
@@ -1251,15 +1356,17 @@ public:
 			check(false);	// Need to update this code to include the view type
 			break;
 		}
+#endif // USE_D3D12RHI_RESOURCE_STATE_TRACKING || ASSERT_RESOURCE_STATES
 	}
 
 	/** Transition a resource's state based on a Depth stencil view's desc flags */
 	static inline void TransitionResource(FD3D12CommandListHandle& hCommandList, FD3D12DepthStencilView* pView)
 	{
+#if USE_D3D12RHI_RESOURCE_STATE_TRACKING || ASSERT_RESOURCE_STATES
 		// Determine the required subresource states from the view desc
 		const D3D12_DEPTH_STENCIL_VIEW_DESC &desc = pView->GetDesc();
-		const D3D12_RESOURCE_STATES depthState =	(desc.Flags & D3D12_DSV_FLAG_READ_ONLY_DEPTH)	? D3D12_RESOURCE_STATE_DEPTH_READ : D3D12_RESOURCE_STATE_DEPTH_WRITE;
-		const D3D12_RESOURCE_STATES stencilState =	(desc.Flags & D3D12_DSV_FLAG_READ_ONLY_STENCIL)	? D3D12_RESOURCE_STATE_DEPTH_READ : D3D12_RESOURCE_STATE_DEPTH_WRITE;
+		const D3D12_RESOURCE_STATES depthState = (desc.Flags & D3D12_DSV_FLAG_READ_ONLY_DEPTH) ? D3D12_RESOURCE_STATE_DEPTH_READ : D3D12_RESOURCE_STATE_DEPTH_WRITE;
+		const D3D12_RESOURCE_STATES stencilState = (desc.Flags & D3D12_DSV_FLAG_READ_ONLY_STENCIL) ? D3D12_RESOURCE_STATE_DEPTH_READ : D3D12_RESOURCE_STATE_DEPTH_WRITE;
 		const bool bSameState = depthState == stencilState;
 
 		FD3D12Resource* pResource = pView->GetResource();
@@ -1287,11 +1394,13 @@ public:
 			TransitionResource(hCommandList, pResource, depthState, pView->GetDepthOnlyViewSubresourceSubset());
 			TransitionResource(hCommandList, pResource, stencilState, pView->GetStencilOnlyViewSubresourceSubset());
 		}
+#endif // USE_D3D12RHI_RESOURCE_STATE_TRACKING || ASSERT_RESOURCE_STATES
 	}
 
 	/** Transition a resource's state based on a Depth stencil view */
 	static inline void TransitionResource(FD3D12CommandListHandle& hCommandList, FD3D12DepthStencilView* pView, D3D12_RESOURCE_STATES after)
 	{
+#if USE_D3D12RHI_RESOURCE_STATE_TRACKING || ASSERT_RESOURCE_STATES
 		FD3D12Resource* pResource = pView->GetResource();
 
 		const D3D12_DEPTH_STENCIL_VIEW_DESC &desc = pView->GetDesc();
@@ -1325,11 +1434,13 @@ public:
 			check(false);	// Need to update this code to include the view type
 			break;
 		}
+#endif // USE_D3D12RHI_RESOURCE_STATE_TRACKING || ASSERT_RESOURCE_STATES
 	}
 
 	/** Transition a resource's state based on a Unordered access view */
 	static inline void TransitionResource(FD3D12CommandListHandle& hCommandList, FD3D12UnorderedAccessView* pView, D3D12_RESOURCE_STATES after)
 	{
+#if USE_D3D12RHI_RESOURCE_STATE_TRACKING || ASSERT_RESOURCE_STATES
 		FD3D12Resource* pResource = pView->GetResource();
 
 		const D3D12_UNORDERED_ACCESS_VIEW_DESC &desc = pView->GetDesc();
@@ -1361,12 +1472,14 @@ public:
 			check(false);	// Need to update this code to include the view type
 			break;
 		}
+#endif // USE_D3D12RHI_RESOURCE_STATE_TRACKING || ASSERT_RESOURCE_STATES
 	}
 
 	/** Transition a resource's state based on a Shader resource view */
-	static inline void TransitionResource(FD3D12CommandListHandle& hCommandList, FD3D12ShaderResourceView* pSRView, D3D12_RESOURCE_STATES after)
+	static inline void TransitionResource(FD3D12CommandListHandle& hCommandList, FD3D12ShaderResourceView* pView, D3D12_RESOURCE_STATES after)
 	{
-		FD3D12Resource* pResource = pSRView->GetResource();
+#if USE_D3D12RHI_RESOURCE_STATE_TRACKING || ASSERT_RESOURCE_STATES
+		FD3D12Resource* pResource = pView->GetResource();
 
 		if (!pResource || !pResource->RequiresResourceStateTracking())
 		{
@@ -1375,9 +1488,9 @@ public:
 		}
 
 		const D3D12_RESOURCE_DESC &ResDesc = pResource->GetDesc();
-		const CViewSubresourceSubset &subresourceSubset = pSRView->GetViewSubresourceSubset();
+		const CViewSubresourceSubset &subresourceSubset = pView->GetViewSubresourceSubset();
 
-		const D3D12_SHADER_RESOURCE_VIEW_DESC &desc = pSRView->GetDesc();
+		const D3D12_SHADER_RESOURCE_VIEW_DESC &desc = pView->GetDesc();
 		switch (desc.ViewDimension)
 		{
 		default:
@@ -1397,18 +1510,23 @@ public:
 			break;
 		}
 		}
+#endif // USE_D3D12RHI_RESOURCE_STATE_TRACKING || ASSERT_RESOURCE_STATES
 	}
 
 	// Transition a specific subresource to the after state.
 	static inline void TransitionResource(FD3D12CommandListHandle& hCommandList, FD3D12Resource* pResource, D3D12_RESOURCE_STATES after, uint32 subresource)
 	{
+#if USE_D3D12RHI_RESOURCE_STATE_TRACKING || ASSERT_RESOURCE_STATES
 		TransitionResourceWithTracking(hCommandList, pResource, after, subresource);
+#endif // USE_D3D12RHI_RESOURCE_STATE_TRACKING || ASSERT_RESOURCE_STATES
 	}
 
 	// Transition a subset of subresources to the after state.
 	static inline void TransitionResource(FD3D12CommandListHandle& hCommandList, FD3D12Resource* pResource, D3D12_RESOURCE_STATES after, const CViewSubresourceSubset& subresourceSubset)
 	{
+#if USE_D3D12RHI_RESOURCE_STATE_TRACKING || ASSERT_RESOURCE_STATES
 		TransitionResourceWithTracking(hCommandList, pResource, after, subresourceSubset);
+#endif USE_D3D12RHI_RESOURCE_STATE_TRACKING || ASSERT_RESOURCE_STATES
 	}
 
 	// Transition subresources from one state to another
@@ -1462,6 +1580,7 @@ public:
 	// Transition a subresource from current to a new state, using resource state tracking.
 	static void TransitionResourceWithTracking(FD3D12CommandListHandle& hCommandList, FD3D12Resource* pResource, D3D12_RESOURCE_STATES after, uint32 subresource)
 	{
+#if USE_D3D12RHI_RESOURCE_STATE_TRACKING
 		check(pResource->RequiresResourceStateTracking());
 
 		CResourceState& ResourceState = hCommandList.GetResourceState(pResource);
@@ -1530,11 +1649,17 @@ public:
 				ResourceState.SetSubresourceState(subresource, after);
 			}
 		}
+#endif // USE_D3D12RHI_RESOURCE_STATE_TRACKING
+
+#if ASSERT_RESOURCE_STATES
+		AssertResourceState(hCommandList.CommandList(), pResource, after, subresource);
+#endif // ASSERT_RESOURCE_STATES
 	}
 
 	// Transition subresources from current to a new state, using resource state tracking.
 	static void TransitionResourceWithTracking(FD3D12CommandListHandle& hCommandList, FD3D12Resource* pResource, D3D12_RESOURCE_STATES after, const CViewSubresourceSubset& subresourceSubset)
 	{
+#if USE_D3D12RHI_RESOURCE_STATE_TRACKING
 		check(pResource->RequiresResourceStateTracking());
 
 		D3D12_RESOURCE_STATES before;
@@ -1625,6 +1750,11 @@ public:
 				ResourceState.SetResourceState(after);
 			}
 		}
+#endif // USE_D3D12RHI_RESOURCE_STATE_TRACKING
+
+#if ASSERT_RESOURCE_STATES
+		AssertResourceState(hCommandList.CommandList(), pResource, after, subresourceSubset);
+#endif // ASSERT_RESOURCE_STATES
 	}
 
 	void GetLocalVideoMemoryInfo(DXGI_QUERY_VIDEO_MEMORY_INFO* LocalVideoMemoryInfo);
diff --git a/Engine/Source/Runtime/Windows/D3D12RHI/Private/D3D12StateCache.cpp b/Engine/Source/Runtime/Windows/D3D12RHI/Private/D3D12StateCache.cpp
--- a/Engine/Source/Runtime/Windows/D3D12RHI/Private/D3D12StateCache.cpp
+++ b/Engine/Source/Runtime/Windows/D3D12RHI/Private/D3D12StateCache.cpp
@@ -50,84 +50,6 @@ static FD3D12ToggleStateCacheExecHelper GD3D12ToggleStateCacheExecHelper;
 
 #endif	// D3D12_STATE_CACHE_RUNTIME_TOGGLE
 
-// FD3D12DynamicHeapAllocator
-FD3D12DynamicHeapAllocator::FD3D12DynamicHeapAllocator(FD3D12Device* InParent, FString Name, eBuddyAllocationStrategy allocationStrategy,
-		uint32 MaxSizeForPooling,
-		uint32 maxBlockSize,
-		uint32 minBlockSize)
-	: CurrentCommandContext(nullptr)
-	, FD3D12ResourceAllocator(InParent, Name,
-		allocationStrategy, 
-		D3D12_HEAP_TYPE_UPLOAD, 
-		D3D12_HEAP_FLAG_ALLOW_ONLY_BUFFERS,
-		D3D12_RESOURCE_FLAG_NONE,
-		MaxSizeForPooling,
-		maxBlockSize, 
-		minBlockSize)
-{
-}
-
-void* FD3D12DynamicHeapAllocator::AllocUploadResource(uint32 size, uint32 alignment, FD3D12ResourceLocation* ResourceLocation)
-{
-	ResourceLocation->Clear();
-
-	//TODO: For some reason 0 sized buffers are being created and then expected to have a resource
-	if (size == 0)
-	{
-		size = 16;
-	}
-
-	//Work loads like infiltrator create enourmous amounts of buffer space in setup
-	//clean up as we go as it can even run out of memory before the first frame.
-	if (GetParentDevice()->FirstFrameSeen == false)
-	{
-		GetParentDevice()->GetDeferredDeletionQueue().ReleaseResources(true);
-		CleanUpAllocations();
-	}
-
-	if (size <= MaximumAllocationSizeForPooling)
-	{
-		FD3D12ResourceBlockInfo* Block = TryAllocate(size, alignment, nullptr);
-
-		if (Block)
-		{
-			FD3D12CommandListManager& CommandListManager = GetParentDevice()->GetCommandListManager();
-			Block->FrameFence = CommandListManager.GetFence(FT_Frame).GetCurrentFence();
-			check(Block->FrameFence);
-
-			FD3D12StateCache* pStateCache = CurrentCommandContext ? &CurrentCommandContext->StateCache : nullptr;
-			ResourceLocation->SetBlockInfo(Block, pStateCache);
-			ResourceLocation->EffectiveBufferSize = size;
-
-			check(Block->Address);
-			check(Block->ResourceHeap);
-
-			return Block->Address;
-		}
-	}
-
-	//Allocate Standalone
-
-	// Allocations are 64k aligned
-	if (alignment)
-	{
-		alignment = (D3D_BUFFER_ALIGNMENT % alignment) == 0 ? 0 : alignment;
-	}
-
-	const uint64 BufferSize = size + alignment;
-
-	VERIFYD3D11RESULT(GetParentDevice()->GetResourceHelper().CreateBuffer(D3D12_HEAP_TYPE_UPLOAD, BufferSize, ResourceLocation->Resource.GetInitReference()));
-	SetName(ResourceLocation->Resource, L"Stand Alone Upload Buffer");
-
-	void* Data = nullptr;
-	VERIFYD3D11RESULT(ResourceLocation->Resource->GetResource()->Map(0, nullptr, &Data));
-
-	ResourceLocation->EffectiveBufferSize = size;
-	ResourceLocation->Offset = alignment;
-
-	return Data;
-}
-
 void* FD3D12FastAllocator::Allocate(uint32 size, uint32 alignment, class FD3D12ResourceLocation* ResourceLocation)
 {
 	return AllocateInternal(size, alignment, ResourceLocation);
@@ -209,151 +131,6 @@ void* FD3D12ThreadSafeFastAllocator::Allocate(uint32 size, uint32 alignment, cla
 	return AllocateInternal(size, alignment, ResourceLocation);
 }
 
-FD3D12DefaultBufferPool::FD3D12DefaultBufferPool(FD3D12Device* InParent, FString Name, eBuddyAllocationStrategy allocationStrategy,
-		uint32 MaxSizeForPooling,
-		D3D12_RESOURCE_FLAGS flags,
-		uint32 maxBlockSize,
-		uint32 minBlockSize) :
-	FD3D12ResourceAllocator(InParent, Name,
-		allocationStrategy, 
-		D3D12_HEAP_TYPE_DEFAULT,
-		D3D12_HEAP_FLAG_ALLOW_ONLY_BUFFERS,
-		flags, MaxSizeForPooling, 
-		maxBlockSize,
-		minBlockSize)
-{
-}
-
-FD3D12DefaultBufferPool::FD3D12DefaultBufferPool()
-	: FD3D12ResourceAllocator(nullptr, FString(L""),
-		kManualSubAllocationStrategy,
-		D3D12_HEAP_TYPE_DEFAULT,
-		D3D12_HEAP_FLAG_ALLOW_ONLY_BUFFERS,
-		D3D12_RESOURCE_FLAG_NONE,
-		0, 
-		0)
-{
-}
-
-// Grab a buffer from the available buffers or create a new buffer if none are available
-void FD3D12DefaultBufferPool::AllocDefaultResource(const D3D12_RESOURCE_DESC& Desc, D3D12_SUBRESOURCE_DATA* InitialData, FD3D12ResourceLocation* ResourceLocation, uint32 Alignment)
-{
-	// If the resource location owns a block, this will deallocate it.
-	ResourceLocation->Clear();
-
-	// Cleanup the pool
-	CleanUpAllocations();
-
-	if (Desc.Width == 0) return;
-
-	const void* Data = (InitialData) ? InitialData->pData : nullptr;
-
-	const bool PoolResource = Desc.Width < MaximumAllocationSizeForPooling &&
-		((Desc.Width % (1024 * 64)) != 0) ;
-
-	if (PoolResource)
-	{
-		// Ensure we're allocating from the correct pool
-		check(Desc.Flags == ResourceFlags);
-
-		FD3D12ResourceBlockInfo* Block = TryAllocate(Desc.Width, Alignment, Data);
-
-		if (Block)
-		{
-			ResourceLocation->SetBlockInfo(Block, &GetParentDevice()->GetDefaultCommandContext().StateCache);
-
-			ResourceLocation->EffectiveBufferSize = Desc.Width;
-
-			if (Alignment > 0)
-			{
-				uint32 AlignmentMismatch = ResourceLocation->GetOffset() % Alignment;
-
-				if (AlignmentMismatch != 0)
-				{
-					ResourceLocation->SetPadding(Alignment - AlignmentMismatch);
-
-					check(ResourceLocation->GetOffset() % Alignment == 0);
-					check(ResourceLocation->GetOffset() % 4 == 0);
-				}
-			}
-
-			return;
-		}
-	}
-
-	//Allocate Standalone
-
-	//Allocations are 64k aligned
-	if (Alignment)
-	{
-		Alignment = (D3D_BUFFER_ALIGNMENT % Alignment) == 0 ? 0 : Alignment;
-	}
-
-	const uint64 Size = (Desc.Width + Alignment);
-	VERIFYD3D11RESULT(GetParentDevice()->GetResourceHelper().CreateBuffer(D3D12_HEAP_TYPE_DEFAULT, Size, ResourceLocation->Resource.GetInitReference(), ResourceFlags));
-	SetName(ResourceLocation->Resource, L"Stand Alone Default Buffer");
-
-	ResourceLocation->EffectiveBufferSize = Desc.Width;
-	ResourceLocation->Offset = Alignment;
-
-	if (Data)
-	{
-		InitializeDefaultBuffer(ResourceLocation->Resource.GetReference(), ResourceLocation->Offset, Data, InitialData->RowPitch);
-	}
-}
-
-// Grab a buffer from the available buffers or create a new buffer if none are available
-HRESULT FD3D12DefaultBufferAllocator::AllocDefaultResource(const D3D12_RESOURCE_DESC& Desc, D3D12_SUBRESOURCE_DATA* InitialData, FD3D12ResourceLocation* ResourceLocation, uint32 Alignment)
-{
-	//NOTE: Indexing based on the resource flags looks weird but is necessary e.g. the flags dictate if the resource
-	//      can be used as a UAV. So each type of buffer has to come from a separate pool.
-
-	check ((uint32)Desc.Flags < MAX_DEFAULT_POOLS);
-	if (DefaultBufferPools[Desc.Flags] == nullptr)
-	{
-		// D3D12_RESOURCE_FLAG_DENY_SHADER_RESOURCE are used much more frequently
-		uint32 HeapSize = (Desc.Flags == D3D12_RESOURCE_FLAG_DENY_SHADER_RESOURCE) ? DEFAULT_BUFFER_POOL_SIZE : 4 * 1024 * 1024;
-		uint32 MaxAllocSize = (Desc.Flags == D3D12_RESOURCE_FLAG_DENY_SHADER_RESOURCE) ? DEFAULT_BUFFER_POOL_MAX_ALLOC_SIZE : 1024 * 64;
-		DefaultBufferPools[Desc.Flags] = new FD3D12DefaultBufferPool(GetParentDevice(), FString(L"Default Buffer Allocator"), kManualSubAllocationStrategy, MaxAllocSize,Desc.Flags, HeapSize, 32);
-	}
-
-	DefaultBufferPools[Desc.Flags]->AllocDefaultResource(Desc, InitialData, ResourceLocation, Alignment);
-
-	return S_OK;
-}
-
-HRESULT FD3D12DefaultBufferAllocator::AllocDefaultResource(const D3D12_RESOURCE_DESC& Desc, D3D12_SUBRESOURCE_DATA* InitialData, FD3D12ResourceLocation** ResourceLocation, uint32 Alignment)
-{
-	*ResourceLocation = new FD3D12ResourceLocation(GetParentDevice(), Desc.Width);
-	(*ResourceLocation)->AddRef();
-
-	return AllocDefaultResource(Desc, InitialData, *ResourceLocation, Alignment);
-}
-
-void FD3D12DefaultBufferAllocator::FreeDefaultBufferPools()
-{
-	for (uint32 i = 0; i < MAX_DEFAULT_POOLS; ++i)
-	{
-		if (DefaultBufferPools[i])
-		{
-			DefaultBufferPools[i]->CleanUpAllocations();
-		}
-		delete DefaultBufferPools[i];
-		DefaultBufferPools[i] = nullptr;
-	}
-}
-
-void FD3D12DefaultBufferAllocator::CleanupFreeBlocks()
-{
-	for (uint32 i = 0; i < MAX_DEFAULT_POOLS; ++i)
-	{
-		if (DefaultBufferPools[i])
-		{
-			DefaultBufferPools[i]->CleanUpAllocations();
-		}
-	}
-}
-
 HRESULT FD3D12ResourceHelper::CreateCommittedResource(const D3D12_RESOURCE_DESC& Desc, const D3D12_HEAP_PROPERTIES& HeapProps, const D3D12_CLEAR_VALUE* ClearValue, FD3D12Resource** ppResource)
 {
 	return CreateCommittedResource(Desc, HeapProps, D3D12_RESOURCE_STATE_COMMON, ClearValue, ppResource);
@@ -1058,8 +835,8 @@ void FD3D12StateCacheBase::ApplyState(bool IsCompute)
 		bNeedSetConstantBuffers = false;
 	}
 
-#if DEBUG_RESOURCE_STATES
-	bool bSucceeded = VerifyResourceStates(IsCompute);
+#if ASSERT_RESOURCE_STATES
+	bool bSucceeded = AssertResourceStates(IsCompute);
 	check(bSucceeded);
 #endif
 }
@@ -2060,19 +1837,20 @@ FD3D12PipelineStateCache::~FD3D12PipelineStateCache()
 
 }
 
-bool FD3D12StateCacheBase::VerifyResourceStates(const bool IsCompute)
+bool FD3D12StateCacheBase::AssertResourceStates(const bool IsCompute)
 {
 	// Can only verify resource states if the debug layer is used
-	const bool bWithD3DDebug = D3D12RHI_ShouldCreateWithD3DDebug();
+	static const bool bWithD3DDebug = D3D12RHI_ShouldCreateWithD3DDebug();
 	if (!bWithD3DDebug)
 	{
-		UE_LOG(LogD3D12RHI, Log, TEXT("*** VerifyResourceStates requires the debug layer ***"), this);
+		UE_LOG(LogD3D12RHI, Fatal, TEXT("*** AssertResourceStates requires the debug layer ***"));
 		return false;
 	}
 
 	// Get the debug command queue
+	ID3D12CommandList* pCommandList = CmdContext->CommandListHandle.CommandList();
 	TRefCountPtr<ID3D12DebugCommandList> pDebugCommandList;
-	VERIFYD3D11RESULT(CmdContext->CommandListHandle->QueryInterface(pDebugCommandList.GetInitReference()));
+	VERIFYD3D11RESULT(pCommandList->QueryInterface(pDebugCommandList.GetInitReference()));
 
 	// Note: There is nothing special to check when IsCompute = true
 	if (!IsCompute)
@@ -2160,84 +1938,47 @@ bool FD3D12StateCacheBase::VerifyResourceStates(const bool IsCompute)
 			for (uint32 i = 0; i < numRTVs; i++)
 			{
 				FD3D12RenderTargetView* pCurrentView = PipelineState.Graphics.RenderTargetArray[i];
-				if (!VerifyViewState(pDebugCommandList, pCurrentView, D3D12_RESOURCE_STATE_RENDER_TARGET))
+				if (!AssertResourceState(pCommandList, pCurrentView, D3D12_RESOURCE_STATE_RENDER_TARGET))
 				{
 					return false;
 				}
 			}
 		}
+
+		// TODO: Verify vertex buffer, index buffer, and constant buffer state.
 	}
 
 	//
 	// Verify common pipeline state
 	//
 
-	// UAVs
-	{
-		// MSFT: Disabled UAV checks for now. This will fail in Infiltrator during the FXSystem's SortGPUBuffers call
-		// because a single sub-resources is used for UAV and SRV. Need to identify what the correct state should be.
-		/*const uint32 numUAVs = _countof(PipelineState.Common.UnorderedAccessViewArray);
-		for (uint32 i = 0; i < numUAVs; i++)
-		{
-			FD3D12UnorderedAccessView *pCurrentView = PipelineState.Common.UnorderedAccessViewArray[i];
-			if (!VerifyViewState(pDebugCommandList, pCurrentView, D3D12_RESOURCE_STATE_UNORDERED_ACCESS))
-			{
-				return false;
-			}
-		}*/
-	}
-
-	// SRVs
+	const uint32 StartStage = IsCompute ? SF_Compute : 0;
+	const uint32 EndStage = IsCompute ? SF_NumFrequencies : SF_Compute;
+	for (uint32 Stage = StartStage; Stage < EndStage; Stage++)
 	{
-		const uint32 numStages = _countof(PipelineState.Common.CurrentShaderResourceViews);
-		for (uint32 stage = 0; stage < numStages; stage++)
+		// UAVs
 		{
-			const uint32 numSRVs = _countof(PipelineState.Common.CurrentShaderResourceViews[stage]);
-			for (uint32 i = 0; i < numSRVs; i++)
+			const uint32 numUAVs = PipelineState.Common.CurrentShaderUAVCounts[Stage];
+			for (uint32 i = 0; i < numUAVs; i++)
 			{
-				FD3D12ShaderResourceView* pCurrentView = PipelineState.Common.CurrentShaderResourceViews[stage][i];
-				if (!VerifyViewState(pDebugCommandList, pCurrentView, D3D12_RESOURCE_STATE_PIXEL_SHADER_RESOURCE | D3D12_RESOURCE_STATE_NON_PIXEL_SHADER_RESOURCE))
+				FD3D12UnorderedAccessView *pCurrentView = PipelineState.Common.UnorderedAccessViewArray[Stage][i];
+				if (!AssertResourceState(pCommandList, pCurrentView, D3D12_RESOURCE_STATE_UNORDERED_ACCESS))
 				{
 					return false;
 				}
 			}
 		}
-	}
 
-	return true;
-}
-
-template <class TView>
-bool FD3D12StateCacheBase::VerifyViewState(ID3D12DebugCommandList* pDebugCommandList, FD3D12View<TView>* pView, D3D12_RESOURCE_STATES State)
-{
-	// Check the view
-	if (!pView)
-	{
-		// No need to check null views
-		return true;
-	}
-
-	// Get the view's resource
-	FD3D12Resource* pResource = pView->GetResource();
-	if (!pResource)
-	{
-		// Some dynamic SRVs haven't been mapped and updated yet so they actually don't have any backing resources.
-		return true;
-	}
-
-	ID3D12Resource* pD3D12Resource = pResource->GetResource();
-	check(pD3D12Resource);
-
-	// For each subresource in the view...
-	const CViewSubresourceSubset subresourceSubset = pView->GetViewSubresourceSubset();
-	for (CViewSubresourceSubset::CViewSubresourceIterator it = subresourceSubset.begin(); it != subresourceSubset.end(); ++it)
-	{
-		for (uint32 SubresourceIndex = it.StartSubresource(); SubresourceIndex < it.EndSubresource(); SubresourceIndex++)
+		// SRVs
 		{
-			bool bGoodState = !!pDebugCommandList->AssertResourceState(pD3D12Resource, SubresourceIndex, State);
-			if (!bGoodState)
+			const uint32 numSRVs = PipelineState.Common.CurrentShaderSRVCounts[Stage];
+			for (uint32 i = 0; i < numSRVs; i++)
 			{
-				return false;
+				FD3D12ShaderResourceView* pCurrentView = PipelineState.Common.CurrentShaderResourceViews[Stage][i];
+				if (!AssertResourceState(pCommandList, pCurrentView, D3D12_RESOURCE_STATE_PIXEL_SHADER_RESOURCE | D3D12_RESOURCE_STATE_NON_PIXEL_SHADER_RESOURCE))
+				{
+					return false;
+				}
 			}
 		}
 	}
@@ -2752,6 +2493,7 @@ void FD3D12DescriptorCache::SetUAVs(EShaderFrequency ShaderStage, uint32 UAVStar
 	FD3D12CommandListHandle& CommandList = CmdContext->CommandListHandle;
 
 	// Fill heap slots
+	check(UAVStartSlot != -1);	// This should never happen or we'll write past the end of the descriptor heap.
 	for (uint32 i = 0; i < UAVStartSlot; i++)
 	{
 		SrcDescriptors[i] = pNullUAV->GetView();
@@ -3637,368 +3379,6 @@ void FD3D12ResourceBlockInfo::Destroy()
 	ResourceHeap = nullptr;
 }
 
-FD3D12ResourceAllocator::FD3D12ResourceAllocator(FD3D12Device* ParentDevice, FString Name, eBuddyAllocationStrategy allocationStrategy, D3D12_HEAP_TYPE heapType,
-	D3D12_HEAP_FLAGS heapFlags, D3D12_RESOURCE_FLAGS flags, uint32 MaxSizeForPooling, uint32 maxBlockSize, uint32 MinBlockSize)
-	: AllocationStrategy(allocationStrategy)
-	, DebugName(Name)
-	, HeapType(heapType)
-	, ResourceFlags(flags)
-	, HeapFlags(heapFlags)
-	, MaxBlockSize(maxBlockSize)
-	, MinBlockSize(MinBlockSize)
-	, BackingHeap(nullptr)
-	, BaseAddress(nullptr)
-	, Initialized(false)
-	, MaximumAllocationSizeForPooling(MaxSizeForPooling)
-#if defined(UE_BUILD_DEBUG)
-	, PeakUsage(0)
-	, SpaceUsed(0)
-	, InternalFragmentation(0)
-	, NumBlocksInDeferredDeletionQueue(0)
-	, FailedAllocationSpace(0)
-#endif
-	, HeapFullMessageDisplayed(false)
-	, FD3D12DeviceChild(ParentDevice)
-{
-	// maxBlockSize should be evenly dividable by MinBlockSize and  
-	// maxBlockSize / MinBlockSize should be a power of two  
-	check((maxBlockSize / MinBlockSize) * MinBlockSize == maxBlockSize); // Evenly dividable  
-	check(0 == ((maxBlockSize / MinBlockSize) & ((maxBlockSize / MinBlockSize) - 1))); // Power of two  
-
-	MaxOrder = UnitSizeToOrder(SizeToUnitSize(maxBlockSize));
-
-	Reset();
-}
-
-FD3D12ResourceAllocator::~FD3D12ResourceAllocator()
-{
-}
-
-void FD3D12ResourceAllocator::Initialize()
-{
-	if (AllocationStrategy == eBuddyAllocationStrategy::kPlacedResourceStrategy)
-	{
-		D3D12_HEAP_PROPERTIES heapProps = CD3DX12_HEAP_PROPERTIES(HeapType);
-
-		D3D12_HEAP_DESC desc = {};
-		desc.SizeInBytes = MaxBlockSize;
-		desc.Properties = heapProps;
-		desc.Alignment = 0;
-		desc.Flags = HeapFlags;
-
-		VERIFYD3D11RESULT(GetParentDevice()->GetDevice()->CreateHeap(&desc, IID_PPV_ARGS(BackingHeap.GetInitReference())));
-		SetName(BackingHeap, L"Placed Resource Allocator Backing Heap");
-	}
-	else
-	{
-		VERIFYD3D11RESULT(GetParentDevice()->GetResourceHelper().CreateBuffer(HeapType, MaxBlockSize, BackingResource.GetInitReference(), ResourceFlags));
-		SetName(BackingResource, L"Resource Allocator Underlying Buffer");
-
-		if (IsCPUWritable(HeapType))
-		{
-			BackingResource->GetResource()->Map(0, nullptr, &BaseAddress);
-		}
-	}
-}
-
-void FD3D12ResourceAllocator::Destroy()
-{
-	ReleaseAllResources();
-}
-
-uint32 FD3D12ResourceAllocator::AllocateBlock(uint32 order)
-{
-	uint32 offset;
-
-	if (order > MaxOrder)
-	{
-		check(false); // Can't allocate a block that large  
-	}
-
-	if (FreeBlocks[order].Num() == 0)
-	{
-		// No free nodes in the requested pool.  Try to find a higher-order block and split it.  
-		uint32 left = AllocateBlock(order + 1);
-
-		uint32 size = OrderToUnitSize(order);
-
-		uint32 right = left + size;
-
-		FreeBlocks[order].Add(right); // Add the right block to the free pool  
-
-		offset = left; // Return the left block  
-	}
-
-	else
-	{
-		TSet<uint32>::TConstIterator it(FreeBlocks[order]);
-		offset = *it;
-
-		// Remove the block from the free list
-		FreeBlocks[order].Remove(*it);
-	}
-
-	return offset;
-}
-
-void FD3D12ResourceAllocator::DeallocateBlock(uint32 offset, uint32 order)
-{
-	// See if the buddy block is free  
-	uint32 size = OrderToUnitSize(order);
-
-	uint32 buddy = GetBuddyOffset(offset, size);
-
-	uint32* it = FreeBlocks[order].Find(buddy);
-
-	if (it != nullptr)
-	{
-		// Deallocate merged blocks
-		DeallocateBlock(FMath::Min(offset, buddy), order + 1);
-		// Remove the buddy from the free list  
-		FreeBlocks[order].Remove(*it);
-	}
-	else
-	{
-		// Add the block to the free list
-		FreeBlocks[order].Add(offset);
-	}
-}
-
-FD3D12ResourceBlockInfo* FD3D12ResourceAllocator::Allocate(uint32 SizeInBytes, uint32 Alignment, const void* InitialData)
-{
-	FScopeLock Lock(&CS);
-
-	if (Initialized == false)
-	{
-		Initialize();
-		Initialized = true;
-	}
-
-	// Work out what size block is needed and allocate one
-	uint32 unitSize = SizeToUnitSize(SizeInBytes);
-	uint32 order = UnitSizeToOrder(unitSize);
-	uint32 offset = AllocateBlock(order); // This is the offset in MinBlockSize units
-
-	FD3D12ResourceBlockInfo::AllocationTrackingValues alloc = {};
-	alloc.Size = uint32(OrderToUnitSize(order) * MinBlockSize);
-	alloc.AllocationBlockOffset = uint32(offset * MinBlockSize);
-	alloc.UnpaddedSize = SizeInBytes;
-
-#if defined(UE_BUILD_DEBUG)
-	if (Alignment)
-	{
-		check(MinBlockSize % Alignment == 0);
-	}
-#endif
-
-	INCREASE_ALLOC_COUNTER(SpaceUsed, alloc.Size);
-	INCREASE_ALLOC_COUNTER(InternalFragmentation, (alloc.Size - SizeInBytes));
-
-#if defined(UE_BUILD_DEBUG)
-	if (SpaceUsed > PeakUsage)
-	{
-		PeakUsage = SpaceUsed;
-	}
-#endif
-
-	FD3D12ResourceBlockInfo* Block = new FD3D12ResourceBlockInfo(this, alloc.AllocationBlockOffset, alloc);
-
-	uint64 InitialDataOffset = 0;
-	if (AllocationStrategy == eBuddyAllocationStrategy::kManualSubAllocationStrategy)
-	{
-		Block->ResourceHeap = BackingResource;
-		InitialDataOffset = alloc.AllocationBlockOffset;
-
-		if (IsCPUWritable(HeapType))
-		{
-			Block->Address = (byte*)BaseAddress + alloc.AllocationBlockOffset;
-		}
-	}
-
-	if (InitialData)
-	{
-		InitializeDefaultBuffer(Block->ResourceHeap, InitialDataOffset, InitialData, SizeInBytes);
-	}
-
-	return Block;
-}
-
-FD3D12ResourceBlockInfo * FD3D12ResourceAllocator::TryAllocate(uint32 SizeInBytes, uint32 Alignment, const void * InitialData)
-{
-	FScopeLock Lock(&CS);
-
-	if (CanAllocate(SizeInBytes))
-	{
-		return Allocate(SizeInBytes, Alignment, InitialData);
-	}
-	else
-	{
-		if (HeapFullMessageDisplayed == false)
-		{
-			UE_LOG(LogD3D12RHI, Warning, TEXT("Allocation Pool ran out of space. This is ok as the allocation will still succeed, it will just take up more space."));
-			HeapFullMessageDisplayed = true;
-		}
-		INCREASE_ALLOC_COUNTER(FailedAllocationSpace, SizeInBytes);
-		return nullptr;
-	}
-}
-
-void FD3D12ResourceAllocator::InitializeDefaultBuffer(FD3D12Resource* Destination, uint64 DestinationOffset, const void* Data, uint64 DataSize)
-{
-	check(HeapType == D3D12_HEAP_TYPE_DEFAULT);
-	FD3D12CommandListHandle& hCommandList = GetParentDevice()->GetDefaultCommandContext().CommandListHandle;
-	// Queue up a CopySubresource() to handle initial data...
-
-	// Get an upload heap and initialize data
-	FD3D12ResourceLocation SrcResourceLoc(GetParentDevice());
-	void* pData = GetParentDevice()->GetBufferInitFastAllocator().Allocate(DataSize, 4UL, &SrcResourceLoc);
-	check(pData);
-	FMemory::Memcpy(pData, Data, DataSize);
-
-	// Copy from the temporary upload heap to the default resource
-	{
-		// Writable structured bufferes are sometimes initialized with inital data which means they sometimes need tracking.
-		FConditionalScopeResourceBarrier ConditionalScopeResourceBarrier(hCommandList, Destination, D3D12_RESOURCE_STATE_COPY_DEST, 0);
-
-		GetParentDevice()->GetDefaultCommandContext().numCopies++;
-		hCommandList->CopyBufferRegion(
-			Destination->GetResource(),
-			DestinationOffset,
-			SrcResourceLoc.GetResource()->GetResource(),
-			SrcResourceLoc.GetOffset(), DataSize);
-	}
-}
-
-void FD3D12ResourceAllocator::Deallocate(FD3D12ResourceBlockInfo* Block)
-{
-	FScopeLock Lock(&CS);
-
-	Block->FrameFence = GetParentDevice()->GetCommandListManager().GetFence(EFenceType::FT_Frame).GetCurrentFence();
-
-	if (Block->IsPlacedResource)
-	{
-		// Blocks don't have ref counted pointers to their resources so add a ref for the deletion queue
-		// (Otherwise the resource may be deleted before the GPU is finished with it)
-		Block->ResourceHeap->AddRef();
-	}
-	DeferredDeletionQueue.Enqueue(Block);
-	INCREASE_ALLOC_COUNTER(NumBlocksInDeferredDeletionQueue, 1);
-}
-
-void FD3D12ResourceAllocator::DeallocateInternal(FD3D12ResourceBlockInfo* Block)
-{
-	check(IsOwner(*Block));
-
-	uint32 offset = SizeToUnitSize(Block->GetOffset());
-
-	uint32 size = SizeToUnitSize(Block->GetSize());
-
-	uint32 order = UnitSizeToOrder(size);
-
-	// Blocks are cleaned up async so need a lock
-
-	DeallocateBlock(offset, order);
-
-	DECREASE_ALLOC_COUNTER(SpaceUsed, Block->GetSize());
-	DECREASE_ALLOC_COUNTER(InternalFragmentation, (Block->AllocatorValues.Size - Block->AllocatorValues.UnpaddedSize));
-
-	if (AllocationStrategy == eBuddyAllocationStrategy::kPlacedResourceStrategy)
-	{
-		// Release the resource
-		Block->Destroy();
-	}
-
-	delete Block;
-};
-
-void FD3D12ResourceAllocator::CleanUpAllocations()
-{
-	FScopeLock Lock(&CS);
-
-	FD3D12CommandListManager& manager = GetParentDevice()->GetCommandListManager();
-
-	FD3D12ResourceBlockInfo* Block = nullptr;
-
-	while (DeferredDeletionQueue.IsEmpty() == false &&
-		DeferredDeletionQueue.Peek(Block) &&
-		manager.GetFence(FT_Frame).IsFenceComplete(Block->FrameFence))
-	{
-		DeferredDeletionQueue.Dequeue(Block);
-		DeallocateInternal(Block);
-		DECREASE_ALLOC_COUNTER(NumBlocksInDeferredDeletionQueue, 1);
-	}
-}
-
-void FD3D12ResourceAllocator::ReleaseAllResources()
-{
-	FD3D12ResourceBlockInfo* Block = nullptr;
-
-	while (DeferredDeletionQueue.IsEmpty() == false &&
-		DeferredDeletionQueue.Peek(Block))
-	{
-		DeferredDeletionQueue.Dequeue(Block);
-		DeallocateInternal(Block);
-		DECREASE_ALLOC_COUNTER(NumBlocksInDeferredDeletionQueue, 1);
-	}
-
-	if (BackingResource)
-	{
-		check(BackingResource->GetRefCount() == 1);
-		BackingResource = nullptr;
-	}
-
-	if (BackingHeap)
-	{
-		BackingHeap = nullptr;
-	}
-}
-
-void FD3D12ResourceAllocator::DumpAllocatorStats(class FOutputDevice& Ar)
-{
-#if defined(UE_BUILD_DEBUG)
-	FBufferedOutputDevice BufferedOutput;
-	{
-		// This is the memory tracked inside individual allocation pools
-		FD3D12DynamicRHI* D3DRHI = FD3D12DynamicRHI::GetD3DRHI();
-		FName categoryName(&DebugName.GetCharArray()[0]);
-
-		BufferedOutput.CategorizedLogf(categoryName, ELogVerbosity::Log, TEXT(""));
-		BufferedOutput.CategorizedLogf(categoryName, ELogVerbosity::Log, TEXT("Heap Size | MinBlock Size | Space Used | Peak Usage | Unpooled Allocations | Internal Fragmentation | Blocks in Deferred Delete Queue "));
-		BufferedOutput.CategorizedLogf(categoryName, ELogVerbosity::Log, TEXT("----------"));
-
-		BufferedOutput.CategorizedLogf(categoryName, ELogVerbosity::Log, TEXT("% 10i % 10i % 16i % 12i % 13i % 8i % 10I"),
-			MaxBlockSize,
-			MinBlockSize,
-			SpaceUsed,
-			PeakUsage,
-			FailedAllocationSpace,
-			InternalFragmentation,
-			NumBlocksInDeferredDeletionQueue);
-	}
-
-	BufferedOutput.RedirectTo(Ar);
-#endif
-}
-
-bool FD3D12ResourceAllocator::CanAllocate(uint32 size)
-{
-	uint32 blockSize = MaxBlockSize;
-
-	for (int32 i = FreeBlocks.Num() - 1; i >= 0; i--)
-	{
-		if (FreeBlocks[i].Num() && blockSize >= size)
-		{
-			return true;
-		}
-
-		// Halve the block size;
-		blockSize = blockSize >> 1;
-
-		if (blockSize < size) return false;
-	}
-	return false;
-}
-
 FD3D12FastAllocatorPage* FD3D12FastAllocatorPagePool::RequestFastAllocatorPage()
 {
 	return RequestFastAllocatorPageInternal();
diff --git a/Engine/Source/Runtime/Windows/D3D12RHI/Private/D3D12StateCachePrivate.h b/Engine/Source/Runtime/Windows/D3D12RHI/Private/D3D12StateCachePrivate.h
--- a/Engine/Source/Runtime/Windows/D3D12RHI/Private/D3D12StateCachePrivate.h
+++ b/Engine/Source/Runtime/Windows/D3D12RHI/Private/D3D12StateCachePrivate.h
@@ -961,7 +961,7 @@ public:
 
 	FD3D12PipelineStateCache& operator=(const FD3D12PipelineStateCache&);
 
-	static const bool bUseAPILibaries = true;
+	static const bool bUseAPILibaries = false;
 	uint32 DriverShaderBlobs;
 };
 
@@ -1726,7 +1726,7 @@ public:
 	void RestoreState();
 	void DirtyViewDescriptorTables();
 	void DirtySamplerDescriptorTables();
-	bool VerifyResourceStates(const bool IsCompute);
+	bool AssertResourceStates(const bool IsCompute);
 
 	template <class ViewT>
 	bool VerifyViewState(ID3D12DebugCommandList* pDebugCommandList, FD3D12View<ViewT> *pView, D3D12_RESOURCE_STATES State);
diff --git a/Engine/Source/Runtime/Windows/D3D12RHI/Private/D3D12StructuredBuffer.cpp b/Engine/Source/Runtime/Windows/D3D12RHI/Private/D3D12StructuredBuffer.cpp
--- a/Engine/Source/Runtime/Windows/D3D12RHI/Private/D3D12StructuredBuffer.cpp
+++ b/Engine/Source/Runtime/Windows/D3D12RHI/Private/D3D12StructuredBuffer.cpp
@@ -2,15 +2,9 @@
 
 #include "D3D12RHIPrivate.h"
 
-FStructuredBufferRHIRef FD3D12DynamicRHI::RHICreateStructuredBuffer(uint32 Stride, uint32 Size, uint32 InUsage, FRHIResourceCreateInfo& CreateInfo)
+D3D12_RESOURCE_DESC CreateStructuredBufferResourceDesc(uint32 Size, uint32 InUsage)
 {
-	// Explicitly check that the size is nonzero before allowing CreateStructuredBuffer to opaquely fail.
-	check(Size > 0);
-	// Check for values that will cause D3D calls to fail
-	check(Size / Stride > 0 && Size % Stride == 0);
-
 	// Describe the structured buffer.
-	const bool bIsDynamic = (InUsage & BUF_AnyDynamic) ? true : false;
 	D3D12_RESOURCE_DESC Desc = CD3DX12_RESOURCE_DESC::Buffer(Size);
 
 	if ((InUsage & BUF_ShaderResource) == 0)
@@ -31,61 +25,25 @@ FStructuredBufferRHIRef FD3D12DynamicRHI::RHICreateStructuredBuffer(uint32 Strid
 		}
 	}
 
-	// If a resource array was provided for the resource, create the resource pre-populated
-	D3D12_SUBRESOURCE_DATA InitData = { 0 };
-	D3D12_SUBRESOURCE_DATA* pInitData = nullptr;
-	if (CreateInfo.ResourceArray)
-	{
-		check(Size == CreateInfo.ResourceArray->GetResourceDataSize());
-		InitData.pData = CreateInfo.ResourceArray->GetResourceData();
-		InitData.RowPitch = Size;
-		InitData.SlicePitch = 0;
-		pInitData = &InitData;
-	}
-
-	TRefCountPtr<FD3D12ResourceLocation> StructuredBufferResource = new FD3D12ResourceLocation(GetRHIDevice());
-	if (bIsDynamic)
-	{
-		// Use an upload heap for dynamic resources and map its memory for writing.
-		void* pData = GetRHIDevice()->GetDefaultUploadHeapAllocator().AllocUploadResource(Desc.Width, DEFAULT_CONTEXT_UPLOAD_POOL_ALIGNMENT, StructuredBufferResource);
-
-		// Add the lock to the lock map.
-		FD3D12LockedKey LockedKey(StructuredBufferResource.GetReference());
-		FD3D12LockedData LockedData;
-		LockedData.SetData(pData);
-		LockedData.Pitch = Desc.Width;
-		OutstandingLocks.Add(LockedKey, LockedData);
+	return Desc;
+}
 
-		if (pInitData)
-		{
-			memcpy(pData, InitData.pData, Desc.Width);
-		}
-	}
-	else
-	{
-#if SUB_ALLOCATED_DEFAULT_ALLOCATIONS
-		// Structured buffers, non-byte address buffers, need to be aligned to their stride to ensure that they
-		// can be addressed correctly with element based offsets.
-		if ((InUsage & (BUF_ByteAddressBuffer | BUF_DrawIndirect)) == 0)
-		{
-			GetRHIDevice()->GetDefaultBufferAllocator().AllocDefaultResource(Desc, pInitData, StructuredBufferResource.GetInitReference(), Stride);
-		}
-		else
-#endif
-		{
-			GetRHIDevice()->GetDefaultBufferAllocator().AllocDefaultResource(Desc, pInitData, StructuredBufferResource.GetInitReference());
-		}
-	}
+FStructuredBufferRHIRef FD3D12DynamicRHI::RHICreateStructuredBuffer(uint32 Stride, uint32 Size, uint32 InUsage, FRHIResourceCreateInfo& CreateInfo)
+{
+	// Check for values that will cause D3D calls to fail
+	check(Size / Stride > 0 && Size % Stride == 0);
 
-	UpdateBufferStats(StructuredBufferResource, true, D3D12_BUFFER_TYPE_STRUCTURED);
+	const D3D12_RESOURCE_DESC Desc = CreateStructuredBufferResourceDesc(Size, InUsage);
 
-	if (CreateInfo.ResourceArray)
-	{
-		// Discard the resource array's contents.
-		CreateInfo.ResourceArray->Discard();
-	}
+	// Structured buffers, non-byte address buffers, need to be aligned to their stride to ensure that they
+	// can be addressed correctly with element based offsets.
+	const uint32 Alignment = ((InUsage & (BUF_ByteAddressBuffer | BUF_DrawIndirect)) == 0) ? Stride : 4;
 
-	return new FD3D12StructuredBuffer(GetRHIDevice(), StructuredBufferResource, Stride, Size, InUsage);
+	FD3D12ResourceLocation* ResourceLocation = CreateBuffer(nullptr, Desc, Size, InUsage, CreateInfo, Alignment);
+	FD3D12StructuredBuffer* NewBuffer = new FD3D12StructuredBuffer(GetRHIDevice(), ResourceLocation, Stride, Size, InUsage);
+	UpdateBufferStats(NewBuffer->ResourceLocation, true, D3D12_BUFFER_TYPE_STRUCTURED);
+	NewBuffer->BufferAlignment = Alignment;
+	return NewBuffer;
 }
 
 FD3D12StructuredBuffer::~FD3D12StructuredBuffer()
@@ -95,136 +53,27 @@ FD3D12StructuredBuffer::~FD3D12StructuredBuffer()
 	// Dynamic resources have a permanent entry in the OutstandingLocks that needs to be cleaned up. 
 	if (GetUsage() & BUF_AnyDynamic)
 	{
-		FD3D12LockedKey LockedKey(ResourceLocation.GetReference());
-		FD3D12DynamicRHI::GetD3DRHI()->OutstandingLocks.Remove(LockedKey);
+		FD3D12LockedKey LockedKey(this);
+		FD3D12DynamicRHI::GetD3DRHI()->RemoveFromOutstandingLocks(LockedKey);
 	}
 }
 
-void* FD3D12DynamicRHI::RHILockStructuredBuffer(FStructuredBufferRHIParamRef StructuredBufferRHI, uint32 Offset, uint32 Size, EResourceLockMode LockMode)
+void FD3D12StructuredBuffer::Rename(FD3D12ResourceLocation* NewResource)
 {
-	FD3D12StructuredBuffer*  StructuredBuffer = FD3D12DynamicRHI::ResourceCast(StructuredBufferRHI);
-
-	FD3D12CommandContext& DefaultContext = GetRHIDevice()->GetDefaultCommandContext();
+	FD3D12CommandContext& DefaultContext = GetParentDevice()->GetDefaultCommandContext();
 
 	// If this resource is bound to the device, unbind it
-	DefaultContext.ConditionalClearShaderResource(StructuredBuffer->ResourceLocation);
-
-	// Determine whether the Structured buffer is dynamic or not.
-	D3D12_RESOURCE_DESC const& Desc = StructuredBuffer->Resource->GetDesc();
-	const uint32 BufferUsage = StructuredBuffer->GetUsage();
-	const bool bIsDynamic = (BufferUsage & BUF_AnyDynamic) ? true : false;
-
-	FD3D12LockedKey LockedKey(StructuredBuffer->Resource.GetReference());
-	FD3D12LockedData LockedData;
-
-	if (bIsDynamic)
-	{
-		check(LockMode == RLM_WriteOnly);
+	DefaultContext.ConditionalClearShaderResource(ResourceLocation);
 
-		// The buffer is dynamic, so we've already mapped it. Just return the offset pointer.
-		FD3D12LockedData* pLockedData = OutstandingLocks.Find(LockedKey);
-		check(pLockedData != nullptr);
-		return (void*)((uint8*)pLockedData->GetData() + Offset);
-	}
-	else
-	{
-		if (LockMode == RLM_ReadOnly)
-		{
-			// If the static buffer is being locked for reading, create a staging buffer.
-			TRefCountPtr<FD3D12Resource> StagingStructuredBuffer;
-			VERIFYD3D11RESULT(GetRHIDevice()->GetResourceHelper().CreateBuffer(D3D12_HEAP_TYPE_READBACK, Offset + Size, StagingStructuredBuffer.GetInitReference()));
-			LockedData.StagingResource = StagingStructuredBuffer;
-
-			// Copy the contents of the Structured buffer to the staging buffer.
-			DefaultContext.numCopies++;
-			DefaultContext.CommandListHandle->CopyResource(StagingStructuredBuffer->GetResource(), StructuredBuffer->Resource->GetResource());
-
-			GetRHIDevice()->GetDefaultCommandContext().FlushCommands(true);
-
-			// Map the staging buffer's memory for reading.
-			void* pData;
-			VERIFYD3D11RESULT(StagingStructuredBuffer->GetResource()->Map(0, nullptr, &pData));
-			LockedData.SetData(pData);
-			LockedData.Pitch = Offset + Size;
-		}
-		else
-		{
-			// If the static buffer is being locked for writing, allocate memory for the contents to be written to.
-			check(StructuredBuffer->ResourceLocation->GetEffectiveBufferSize() >= Size);
-
-			// Use an upload heap to copy data to a default resource.
-			TRefCountPtr<FD3D12ResourceLocation> UploadBuffer = new FD3D12ResourceLocation(GetRHIDevice());
-			void *pData = GetRHIDevice()->GetDefaultFastAllocator().Allocate(Offset + Size, D3D12_TEXTURE_DATA_PLACEMENT_ALIGNMENT, UploadBuffer);
-
-			// Add the lock to the lock map.
-			FD3D12LockedKey LockedKey(UploadBuffer);
-			LockedData.SetData(pData);
-			LockedData.Pitch = Offset + Size;
-
-			// Keep track of the underlying resource for the upload heap so it can be referenced during Unmap.
-			LockedData.UploadHeapLocation = UploadBuffer;
-		}
-	}
-
-	// Add the lock to the lock map.
-	OutstandingLocks.Add(LockedKey, LockedData);
+	ResourceLocation = NewResource;
+}
 
-	// Return the offset pointer
-	return (void*)((uint8*)LockedData.GetData() + Offset);
+void* FD3D12DynamicRHI::RHILockStructuredBuffer(FStructuredBufferRHIParamRef StructuredBufferRHI, uint32 Offset, uint32 Size, EResourceLockMode LockMode)
+{
+	return LockBuffer(nullptr, FD3D12DynamicRHI::ResourceCast(StructuredBufferRHI), Offset, Size, LockMode);
 }
 
 void FD3D12DynamicRHI::RHIUnlockStructuredBuffer(FStructuredBufferRHIParamRef StructuredBufferRHI)
 {
-	FD3D12StructuredBuffer*  StructuredBuffer = FD3D12DynamicRHI::ResourceCast(StructuredBufferRHI);
-
-	// Determine whether the Structured buffer is dynamic or not.
-	D3D12_RESOURCE_DESC const& Desc = StructuredBuffer->Resource->GetDesc();
-	const uint32 BufferUsage = StructuredBuffer->GetUsage();
-	const bool bIsDynamic = (BufferUsage & BUF_AnyDynamic) ? true : false;
-
-	// Find the outstanding lock for this VB.
-	FD3D12LockedKey LockedKey(StructuredBuffer->Resource.GetReference());
-	FD3D12LockedData* LockedData = OutstandingLocks.Find(LockedKey);
-	check(LockedData);
-
-	if (bIsDynamic)
-	{
-		// If the VB is dynamic, its upload heap memory can always stay mapped. Don't do anything.
-	}
-	else
-	{
-		// If the static VB lock involved a staging resource, it was locked for reading.
-		if (LockedData->StagingResource)
-		{
-			// Unmap the staging buffer's memory.
-			ID3D12Resource* StagingBuffer = LockedData->StagingResource.GetReference()->GetResource();
-			StagingBuffer->Unmap(0, nullptr);
-		}
-		else
-		{
-			check(LockedData->UploadHeapLocation != nullptr);
-
-			// Copy the contents of the temporary memory buffer allocated for writing into the VB.
-			FD3D12ResourceLocation* UploadHeap = LockedData->UploadHeapLocation;
-
-			FD3D12CommandListHandle& hCommandList = GetRHIDevice()->GetDefaultCommandContext().CommandListHandle;
-
-			// Writable structured bufferes are sometimes unlocked which means they sometimes need tracking.
-			FConditionalScopeResourceBarrier ConditionalScopeResourceBarrier(hCommandList, StructuredBuffer->Resource, D3D12_RESOURCE_STATE_COPY_DEST, 0);
-			// Don't need to transition upload heaps
-
-			GetRHIDevice()->GetDefaultCommandContext().numCopies++;
-			hCommandList->CopyBufferRegion(
-				StructuredBuffer->Resource->GetResource(),
-				StructuredBuffer->ResourceLocation->GetOffset(),
-				UploadHeap->GetResource()->GetResource(),
-				UploadHeap->GetOffset(),
-				LockedData->Pitch);
-		}
-	}
-
-	// Remove the FD3D12LockedData from the lock map.
-	// If the lock involved a staging resource, this releases it.
-	// If the lock involved an upload heap resource, this releases it.
-	OutstandingLocks.Remove(LockedKey);
+	UnlockBuffer(nullptr, FD3D12DynamicRHI::ResourceCast(StructuredBufferRHI));
 }
diff --git a/Engine/Source/Runtime/Windows/D3D12RHI/Private/D3D12Texture.cpp b/Engine/Source/Runtime/Windows/D3D12RHI/Private/D3D12Texture.cpp
--- a/Engine/Source/Runtime/Windows/D3D12RHI/Private/D3D12Texture.cpp
+++ b/Engine/Source/Runtime/Windows/D3D12RHI/Private/D3D12Texture.cpp
@@ -75,7 +75,7 @@ namespace D3D12RHI
 			return;
 		}
 
-		int64 AlignedSize = (TextureSize > 0) ? Align(TextureSize, 1024) / 1024 : -(Align(-TextureSize, 1024) / 1024);
+		const int64 AlignedSize = (TextureSize > 0) ? Align(TextureSize, 1024) / 1024 : -(Align(-TextureSize, 1024) / 1024);
 		if (ShouldCountAsTextureMemory(Desc.Flags))
 		{
 			FPlatformAtomics::InterlockedAdd(&GCurrentTextureMemorySize, AlignedSize);
@@ -102,19 +102,17 @@ namespace D3D12RHI
 	{
 		FD3D12Resource* D3D11Texture2D = Texture.GetResource();
 
-		if (D3D11Texture2D)
+		// Ignore placed textures as their memory is already allocated and accounted for
+		if (D3D11Texture2D && D3D11Texture2D->IsPlacedResource() == false)
 		{
-			if ((Texture.Flags & TexCreate_Virtual) == TexCreate_Virtual)
-			{
-				Texture.SetMemorySize(0);
-			}
-			else
+			check((Texture.Flags & TexCreate_Virtual) == 0);
 			{
 				const D3D12_RESOURCE_DESC& Desc = D3D11Texture2D->GetDesc();
-
-				int64 TextureSize = CalcTextureSize(Desc.Width, Desc.Height, Texture.GetFormat(), Desc.MipLevels) * Desc.DepthOrArraySize;
+				const D3D12_RESOURCE_ALLOCATION_INFO AllocationInfo = Texture.GetParentDevice()->GetDevice()->GetResourceAllocationInfo(0, 1, &Desc);
+				const int64 TextureSize = AllocationInfo.SizeInBytes;
 
 				Texture.SetMemorySize(TextureSize);
+
 				UpdateD3D11TextureStats(Desc, TextureSize, false, Texture.IsCubemap());
 			}
 		}
@@ -125,20 +123,12 @@ namespace D3D12RHI
 	{
 		FD3D12Resource* D3D11Texture2D = Texture.GetResource();
 
-		if (D3D11Texture2D)
+		// Ignore placed textures as their memory is already allocated and accounted for
+		if (D3D11Texture2D && D3D11Texture2D->IsPlacedResource() == false)
 		{
 			const D3D12_RESOURCE_DESC& Desc = D3D11Texture2D->GetDesc();
-
-			// When using virtual textures use the current memory size, which is the number of physical pages allocated, not virtual
-			int64 TextureSize = 0;
-			if ((Texture.GetFlags() & TexCreate_Virtual) == TexCreate_Virtual)
-			{
-				TextureSize = Texture.GetMemorySize();
-			}
-			else
-			{
-				TextureSize = CalcTextureSize(Desc.Width, Desc.Height, Texture.GetFormat(), Desc.MipLevels) * Desc.DepthOrArraySize;
-			}
+			const int64 TextureSize = Texture.GetMemorySize();
+			check(TextureSize > 0);
 
 			UpdateD3D11TextureStats(Desc, -TextureSize, false, Texture.IsCubemap());
 		}
@@ -156,8 +146,8 @@ namespace D3D12RHI
 		if (D3D11Texture3D)
 		{
 			const D3D12_RESOURCE_DESC& Desc = D3D11Texture3D->GetDesc();
-
-			int64 TextureSize = CalcTextureSize3D(Desc.Width, Desc.Height, Desc.DepthOrArraySize, Texture.GetFormat(), Desc.MipLevels);
+			const D3D12_RESOURCE_ALLOCATION_INFO AllocationInfo = Texture.GetParentDevice()->GetDevice()->GetResourceAllocationInfo(0, 1, &Desc);
+			const int64 TextureSize = AllocationInfo.SizeInBytes;
 
 			Texture.SetMemorySize(TextureSize);
 
@@ -172,8 +162,8 @@ namespace D3D12RHI
 		if (D3D11Texture3D)
 		{
 			const D3D12_RESOURCE_DESC& Desc = D3D11Texture3D->GetDesc();
-
-			int64 TextureSize = CalcTextureSize3D(Desc.Width, Desc.Height, Desc.DepthOrArraySize, Texture.GetFormat(), Desc.MipLevels);
+			const int64 TextureSize = Texture.GetMemorySize();
+			check(TextureSize > 0);
 
 			UpdateD3D11TextureStats(Desc, -TextureSize, true, false);
 		}
@@ -203,20 +193,56 @@ FD3D12Texture3D::~FD3D12Texture3D()
 
 uint64 FD3D12DynamicRHI::RHICalcTexture2DPlatformSize(uint32 SizeX, uint32 SizeY, uint8 Format, uint32 NumMips, uint32 NumSamples, uint32 Flags, uint32& OutAlign)
 {
-	OutAlign = 0;
-	return CalcTextureSize(SizeX, SizeY, (EPixelFormat)Format, NumMips);
+	D3D12_RESOURCE_DESC Desc = {};
+	Desc.DepthOrArraySize = 1;
+	Desc.Dimension = D3D12_RESOURCE_DIMENSION_TEXTURE2D;
+	Desc.Format = (DXGI_FORMAT)GPixelFormats[Format].PlatformFormat;
+	Desc.Height = SizeY;
+	Desc.Layout = D3D12_TEXTURE_LAYOUT_UNKNOWN;
+	Desc.MipLevels = NumMips;
+	Desc.SampleDesc.Count = NumSamples;
+	Desc.Width = SizeX;
+
+	const D3D12_RESOURCE_ALLOCATION_INFO AllocationInfo = GetRHIDevice()->GetDevice()->GetResourceAllocationInfo(0, 1, &Desc);
+	OutAlign = static_cast<uint32>(AllocationInfo.Alignment);
+
+	return AllocationInfo.SizeInBytes;
 }
 
 uint64 FD3D12DynamicRHI::RHICalcTexture3DPlatformSize(uint32 SizeX, uint32 SizeY, uint32 SizeZ, uint8 Format, uint32 NumMips, uint32 Flags, uint32& OutAlign)
 {
-	OutAlign = 0;
-	return CalcTextureSize3D(SizeX, SizeY, SizeZ, (EPixelFormat)Format, NumMips);
+	D3D12_RESOURCE_DESC Desc = {};
+	Desc.DepthOrArraySize = SizeZ;
+	Desc.Dimension = D3D12_RESOURCE_DIMENSION_TEXTURE3D;
+	Desc.Format = (DXGI_FORMAT)GPixelFormats[Format].PlatformFormat;
+	Desc.Height = SizeY;
+	Desc.Layout = D3D12_TEXTURE_LAYOUT_UNKNOWN;
+	Desc.MipLevels = NumMips;
+	Desc.SampleDesc.Count = 1;
+	Desc.Width = SizeX;
+
+	const D3D12_RESOURCE_ALLOCATION_INFO AllocationInfo = GetRHIDevice()->GetDevice()->GetResourceAllocationInfo(0, 1, &Desc);
+	OutAlign = static_cast<uint32>(AllocationInfo.Alignment);
+
+	return AllocationInfo.SizeInBytes;
 }
 
 uint64 FD3D12DynamicRHI::RHICalcTextureCubePlatformSize(uint32 Size, uint8 Format, uint32 NumMips, uint32 Flags, uint32& OutAlign)
 {
-	OutAlign = 0;
-	return CalcTextureSize(Size, Size, (EPixelFormat)Format, NumMips) * 6;
+	D3D12_RESOURCE_DESC Desc = {};
+	Desc.DepthOrArraySize = 6;
+	Desc.Dimension = D3D12_RESOURCE_DIMENSION_TEXTURE2D;
+	Desc.Format = (DXGI_FORMAT)GPixelFormats[Format].PlatformFormat;
+	Desc.Height = Size;
+	Desc.Layout = D3D12_TEXTURE_LAYOUT_UNKNOWN;
+	Desc.MipLevels = NumMips;
+	Desc.SampleDesc.Count = 1;
+	Desc.Width = Size;
+
+	const D3D12_RESOURCE_ALLOCATION_INFO AllocationInfo = GetRHIDevice()->GetDevice()->GetResourceAllocationInfo(0, 1, &Desc);
+	OutAlign = static_cast<uint32>(AllocationInfo.Alignment);
+
+	return AllocationInfo.SizeInBytes;
 }
 
 /**
@@ -227,64 +253,76 @@ void FD3D12DynamicRHI::RHIGetTextureMemoryStats(FTextureMemoryStats& OutStats)
 	OutStats.DedicatedVideoMemory = FD3D12GlobalStats::GDedicatedVideoMemory;
 	OutStats.DedicatedSystemMemory = FD3D12GlobalStats::GDedicatedSystemMemory;
 	OutStats.SharedSystemMemory = FD3D12GlobalStats::GSharedSystemMemory;
-	bool bIsWARP = FParse::Param(FCommandLine::Get(), TEXT("warp"));
-
-	// Hack to handle a kernel bug where QueryVideomMemoryInfo can't be called on WARP
-	if (!bIsWARP)
-	{
-		DXGI_QUERY_VIDEO_MEMORY_INFO LocalVideoMemoryInfo = {};
-		VERIFYD3D11RESULT(GetRHIDevice()->GetAdapter3()->QueryVideoMemoryInfo(0, DXGI_MEMORY_SEGMENT_GROUP_LOCAL, &LocalVideoMemoryInfo));
-
-		int64 budget = LocalVideoMemoryInfo.Budget;
-		OutStats.TotalGraphicsMemory = budget;
-
-		// Leave 5% unused to account for sudden budget cuts
-		const int64 BudgetPadding = budget * 0.05f;
-		// Note: AvailableSpace can be negative
-		int64 AvailableSpace = budget - int64(LocalVideoMemoryInfo.CurrentUsage);
-		int64 PreviousPoolSize = GTexturePoolSize;
-
-		// Only change the pool size if over-budget, or a reasonable amount of memory is available
-		if (AvailableSpace < 0 || AvailableSpace > BudgetPadding)
+	// OutStats.TotalGraphicsMemory is set below.
+
+	OutStats.AllocatedMemorySize = int64(GCurrentTextureMemorySize) * 1024;
+	// OutStats.TexturePoolSize is set below.
+	OutStats.PendingMemoryAdjustment = 0;
+	
+	DXGI_QUERY_VIDEO_MEMORY_INFO LocalVideoMemoryInfo;
+	VERIFYD3D11RESULT(GetRHIDevice()->GetAdapter3()->QueryVideoMemoryInfo(0, DXGI_MEMORY_SEGMENT_GROUP_LOCAL, &LocalVideoMemoryInfo));
+
+	// Applications must explicitly manage their usage of physical memory and keep usage within the budget 
+	// assigned to the application process. Processes that cannot keep their usage within their assigned budgets 
+	// will likely experience stuttering, as they are intermittently frozen and paged out to allow other processes to run.
+	const int64 TargetBudget = LocalVideoMemoryInfo.Budget * 0.90f;	// Target using 90% of our budget to account for some fragmentation.
+	OutStats.TotalGraphicsMemory = TargetBudget;
+
+	const int64 BudgetPadding = TargetBudget * 0.05f;
+	const int64 AvailableSpace = TargetBudget - int64(LocalVideoMemoryInfo.CurrentUsage);	// Note: AvailableSpace can be negative
+	const int64 PreviousTexturePoolSize = RequestedTexturePoolSize;
+	const bool bOverbudget = AvailableSpace < 0;
+
+	// Only change the pool size if overbudget, or a reasonable amount of memory is available
+	const int64 MinTexturePoolSize = int64(100 * 1024 * 1024);
+	if (bOverbudget)
+	{
+		// Attempt to lower the texture pool size to meet the budget.
+		const bool bOverActualBudget = LocalVideoMemoryInfo.CurrentUsage > LocalVideoMemoryInfo.Budget;
+		UE_CLOG(bOverActualBudget, LogD3D12RHI, Warning, 
+			TEXT("Video memory usage is overbudget by %llu MB (using %lld MB/%lld MB budget). Usage breakdown: %lld MB (Textures), %lld MB (Render targets). Last requested texture pool size is %lld MB. This can cause stuttering due to paging."),
+			(LocalVideoMemoryInfo.CurrentUsage - LocalVideoMemoryInfo.Budget) / 1024ll / 1024ll,
+			LocalVideoMemoryInfo.CurrentUsage / 1024ll / 1024ll,
+			LocalVideoMemoryInfo.Budget / 1024ll / 1024ll,
+			GCurrentTextureMemorySize / 1024ll,
+			GCurrentRendertargetMemorySize / 1024ll,
+			PreviousTexturePoolSize / 1024ll / 1024ll);
+
+		const int64 DesiredTexturePoolSize = PreviousTexturePoolSize + AvailableSpace - BudgetPadding;
+		OutStats.TexturePoolSize = FMath::Max(DesiredTexturePoolSize, MinTexturePoolSize);
+
+		UE_CLOG(bOverActualBudget && (OutStats.TexturePoolSize >= PreviousTexturePoolSize) && (OutStats.TexturePoolSize > MinTexturePoolSize), LogD3D12RHI, Fatal,
+			TEXT("Video memory usage is overbudget by %llu MB and the texture pool size didn't shrink."),
+			(LocalVideoMemoryInfo.CurrentUsage - LocalVideoMemoryInfo.Budget) / 1024ll / 1024ll);
+	}
+	else if (AvailableSpace > BudgetPadding)
+	{
+		// Increase the texture pool size to improve quality if we have a reasonable amount of memory available.
+		int64 DesiredTexturePoolSize = PreviousTexturePoolSize + AvailableSpace - BudgetPadding;
+		if (GPoolSizeVRAMPercentage > 0)
 		{
-			int64 BudgettedPoolSize = PreviousPoolSize > 0 ? PreviousPoolSize + AvailableSpace - BudgetPadding : INT64_MAX;
+			// The texture pool size is a percentage of GTotalGraphicsMemory.
+			const float PoolSize = float(GPoolSizeVRAMPercentage) * 0.01f * float(OutStats.TotalGraphicsMemory);
 
-			float MaxPoolSize = float(GPoolSizeVRAMPercentage) * 0.01f * budget;
-			int64 TruncMaxPoolSize = int64(FGenericPlatformMath::TruncToFloat(MaxPoolSize / 1024.0f / 1024.0f)) * 1024 * 1024;
-
-			if (sizeof(SIZE_T) < 8)
-			{
-				// Clamp to 1 GB if we're less than 64-bit
-				TruncMaxPoolSize = FMath::Min(TruncMaxPoolSize, 1024ll * 1024ll * 1024ll);
-			}
-			else
-			{
-				// Clamp to 1.9 GB if we're 64-bit
-				TruncMaxPoolSize = FMath::Min(TruncMaxPoolSize, 1945ll * 1024ll * 1024ll);
-			}
-
-			// Clamp to the pre-defined texture pool size, if it's been set.
-			TruncMaxPoolSize = (GTexturePoolSize == 0) ? TruncMaxPoolSize : FMath::Min(TruncMaxPoolSize, GTexturePoolSize);
-
-			const int64 MinPoolSize = int64(200 * 1024 * 1024);
-			OutStats.TexturePoolSize = FMath::Max(FMath::Min(BudgettedPoolSize, TruncMaxPoolSize), MinPoolSize);
-		}
-		else
-		{
-			OutStats.TexturePoolSize = PreviousPoolSize;
+			// Truncate texture pool size to MB (but still counted in bytes).
+			DesiredTexturePoolSize = int64(FGenericPlatformMath::TruncToFloat(PoolSize / 1024.0f / 1024.0f)) * 1024 * 1024;
 		}
 
-		// Truncate GTexturePoolSize to MB (but still counted in bytes)
-		OutStats.AllocatedMemorySize = int64(GCurrentTextureMemorySize) * 1024;
-		OutStats.LargestContiguousAllocation = OutStats.AllocatedMemorySize;
-		OutStats.PendingMemoryAdjustment = 0;
+		// Make sure the desired texture pool size doesn't make us go overbudget.
+		const bool bIsLimitedTexturePoolSize = GTexturePoolSize > 0;
+		const int64 LimitedMaxTexturePoolSize = bIsLimitedTexturePoolSize ? GTexturePoolSize : INT64_MAX;
+		const int64 MaxTexturePoolSize = FMath::Min(PreviousTexturePoolSize + AvailableSpace - BudgetPadding, LimitedMaxTexturePoolSize);	// Max texture pool size without going overbudget or the pre-defined max.
+		OutStats.TexturePoolSize = FMath::Min(DesiredTexturePoolSize, MaxTexturePoolSize);
 	}
 	else
 	{
-		OutStats.TotalGraphicsMemory = FD3D12GlobalStats::GTotalGraphicsMemory;
+		// Keep the previous requested texture pool size.
+		OutStats.TexturePoolSize = PreviousTexturePoolSize;
 	}
 
-
+	// Cache the last requested texture pool size.
+	check(OutStats.TexturePoolSize >= MinTexturePoolSize);
+	RequestedTexturePoolSize = OutStats.TexturePoolSize;
 }
 
 /**
@@ -496,67 +534,6 @@ void InitializeTexture2DData(
 	}
 }
 
-HRESULT FD3D12TextureAllocator::AllocateTexture(D3D12_RESOURCE_DESC Desc, const D3D12_CLEAR_VALUE* ClearValue, FD3D12ResourceLocation* TextureLocation)
-{
-	FD3D12Device* Device = GetParentDevice();
-	FD3D12ResourceHelper& ResourceHelper = Device->GetResourceHelper();
-	TRefCountPtr<FD3D12Resource> Resource;
-	HRESULT hr = S_OK;
-
-	TextureLocation->Clear();
-
-	D3D12_RESOURCE_ALLOCATION_INFO Info = Device->GetDevice()->GetResourceAllocationInfo(0, 1, &Desc);
-
-	if (Info.SizeInBytes < D3D12_DEFAULT_RESOURCE_PLACEMENT_ALIGNMENT)
-	{
-		FD3D12ResourceBlockInfo* Block = TryAllocate(Info.SizeInBytes, Info.Alignment, nullptr);
-		if(Block)
-		{
-			hr = FD3D12DynamicRHI::CreatePlacedResource(Desc, BackingHeap, Block->Offset, D3D12_RESOURCE_STATE_COMMON, ClearValue, Resource.GetInitReference());
-
-			Block->ResourceHeap = Resource.GetReference();
-			Block->IsPlacedResource = true;
-			TextureLocation->SetBlockInfo(Block, &Device->GetDefaultCommandContext().StateCache);
-
-			return hr;
-		}
-	}
-
-	// Request default alignment for stand alone textures
-	Desc.Alignment = 0;
-	hr = ResourceHelper.CreateDefaultResource(Desc, ClearValue, Resource.GetInitReference());
-
-	TextureLocation->SetFromD3DResource(Resource, 0, 0);
-
-	return hr;
-}
-
-HRESULT FD3D12TextureAllocatorPool::AllocateTexture(D3D12_RESOURCE_DESC Desc, const D3D12_CLEAR_VALUE* ClearValue, uint8 UEFormat, FD3D12ResourceLocation* TextureLocation)
-{
-	// Multi-Sample texures have much larger alignment requirements (4MB vs 64KB)
-	check(Desc.SampleDesc.Count == 1)
-
-	// 4KB alignment is only available for read only textures
-	if ((Desc.Flags & D3D12_RESOURCE_FLAG_ALLOW_RENDER_TARGET ||
-		Desc.Flags & D3D12_RESOURCE_FLAG_ALLOW_DEPTH_STENCIL ||
-		Desc.Flags & D3D12_RESOURCE_FLAG_ALLOW_UNORDERED_ACCESS) == false)
-	{
-		// The top mip level must be less than 64k
-		if (TextureCanBe4KAligned(Desc, UEFormat))
-		{
-			Desc.Alignment = D3D12_SMALL_RESOURCE_PLACEMENT_ALIGNMENT; // request 4k alignment
-			return ReadOnlyTexturePool.AllocateTexture(Desc, ClearValue, TextureLocation);
-		}
-	}
-
-	TRefCountPtr<FD3D12Resource> Resource = nullptr;
-
-	HRESULT hr = GetParentDevice()->GetResourceHelper().CreateDefaultResource(Desc, ClearValue, Resource.GetInitReference());
-	TextureLocation->SetFromD3DResource(Resource, 0, 0);
-
-	return hr;
-}
-
 /**
  * Creates a 2D texture optionally guarded by a structured exception handler.
  */
@@ -1456,133 +1433,6 @@ void FD3D12DynamicRHI::RHIGetResourceInfo(FTextureRHIParamRef Ref, FRHIResourceI
 	}
 }
 
-FShaderResourceViewRHIRef FD3D12DynamicRHI::RHICreateShaderResourceView(FTexture2DRHIParamRef Texture2DRHI, uint8 MipLevel)
-{
-	FD3D12Texture2D*  Texture2D = FD3D12DynamicRHI::ResourceCast(Texture2DRHI);
-
-	uint64 Offset = 0;
-	const D3D12_RESOURCE_DESC& TextureDesc = Texture2D->GetResource()->GetDesc();
-
-	const bool bSRGB = (Texture2D->GetFlags() & TexCreate_SRGB) != 0;
-	const DXGI_FORMAT PlatformShaderResourceFormat = FindShaderResourceDXGIFormat(TextureDesc.Format, bSRGB);
-
-	// Create a Shader Resource View
-	D3D12_SHADER_RESOURCE_VIEW_DESC SRVDesc ={};
-	SRVDesc.Shader4ComponentMapping = D3D12_DEFAULT_SHADER_4_COMPONENT_MAPPING;
-	SRVDesc.ViewDimension = D3D12_SRV_DIMENSION_TEXTURE2D;
-	SRVDesc.Texture2D.MipLevels = 1;
-	SRVDesc.Texture2D.MostDetailedMip = MipLevel;
-
-	SRVDesc.Format = PlatformShaderResourceFormat;
-
-	SRVDesc.Texture2D.PlaneSlice = GetPlaneSliceFromViewFormat(TextureDesc.Format, SRVDesc.Format);
-
-	return new FD3D12ShaderResourceView(GetRHIDevice(), &SRVDesc, Texture2D->ResourceLocation);
-}
-
-FShaderResourceViewRHIRef FD3D12DynamicRHI::RHICreateShaderResourceView(FTexture2DArrayRHIParamRef Texture2DRHI, uint8 MipLevel)
-{
-	FD3D12Texture2DArray*  Texture2DArray = FD3D12DynamicRHI::ResourceCast(Texture2DRHI);
-
-	uint64 Offset = 0;
-	const D3D12_RESOURCE_DESC& TextureDesc = Texture2DArray->GetResource()->GetDesc();
-
-	const bool bSRGB = (Texture2DArray->GetFlags() & TexCreate_SRGB) != 0;
-	const DXGI_FORMAT PlatformShaderResourceFormat = FindShaderResourceDXGIFormat(TextureDesc.Format, bSRGB);
-
-	// Create a Shader Resource View
-	D3D12_SHADER_RESOURCE_VIEW_DESC SRVDesc = {};
-	SRVDesc.Shader4ComponentMapping = D3D12_DEFAULT_SHADER_4_COMPONENT_MAPPING;
-	SRVDesc.ViewDimension = D3D12_SRV_DIMENSION_TEXTURE2DARRAY;
-	SRVDesc.Texture2DArray.ArraySize = TextureDesc.DepthOrArraySize;
-	SRVDesc.Texture2DArray.FirstArraySlice = 0;
-	SRVDesc.Texture2DArray.MipLevels = 1;
-	SRVDesc.Texture2DArray.MostDetailedMip = MipLevel;
-
-	SRVDesc.Format = PlatformShaderResourceFormat;
-
-	SRVDesc.Texture2DArray.PlaneSlice = GetPlaneSliceFromViewFormat(TextureDesc.Format, SRVDesc.Format);
-
-	return new FD3D12ShaderResourceView(GetRHIDevice(), &SRVDesc, Texture2DArray->ResourceLocation);
-}
-
-FShaderResourceViewRHIRef FD3D12DynamicRHI::RHICreateShaderResourceView(FTextureCubeRHIParamRef TextureCubeRHI, uint8 MipLevel)
-{
-	FD3D12TextureCube*  TextureCube = FD3D12DynamicRHI::ResourceCast(TextureCubeRHI);
-
-	uint64 Offset = 0;
-	const D3D12_RESOURCE_DESC& TextureDesc = TextureCube->GetResource()->GetDesc();
-
-	const bool bSRGB = (TextureCube->GetFlags() & TexCreate_SRGB) != 0;
-	const DXGI_FORMAT PlatformShaderResourceFormat = FindShaderResourceDXGIFormat(TextureDesc.Format, bSRGB);
-
-	// Create a Shader Resource View
-	D3D12_SHADER_RESOURCE_VIEW_DESC SRVDesc = {};
-	SRVDesc.Shader4ComponentMapping = D3D12_DEFAULT_SHADER_4_COMPONENT_MAPPING;
-	SRVDesc.ViewDimension = D3D12_SRV_DIMENSION_TEXTURECUBE;
-	SRVDesc.TextureCube.MipLevels = 1;
-	SRVDesc.TextureCube.MostDetailedMip = MipLevel;
-
-	SRVDesc.Format = PlatformShaderResourceFormat;
-
-	return new FD3D12ShaderResourceView(GetRHIDevice(), &SRVDesc, TextureCube->ResourceLocation);
-}
-
-FShaderResourceViewRHIRef FD3D12DynamicRHI::RHICreateShaderResourceView(FTexture2DRHIParamRef Texture2DRHI, uint8 MipLevel, uint8 NumMipLevels, uint8 Format)
-{
-	FD3D12Texture2D*  Texture2D = FD3D12DynamicRHI::ResourceCast(Texture2DRHI);
-	uint64 Offset = 0;
-	const D3D12_RESOURCE_DESC& TextureDesc = Texture2D->GetResource()->GetDesc();
-
-	const DXGI_FORMAT PlatformResourceFormat = FD3D12DynamicRHI::GetPlatformTextureResourceFormat((DXGI_FORMAT)GPixelFormats[Format].PlatformFormat, Texture2D->GetFlags());
-
-	const bool bSRGB = (Texture2D->GetFlags() & TexCreate_SRGB) != 0;
-	const DXGI_FORMAT PlatformShaderResourceFormat = FindShaderResourceDXGIFormat(PlatformResourceFormat, bSRGB);
-
-	// Create a Shader Resource View
-	D3D12_SHADER_RESOURCE_VIEW_DESC SRVDesc = {};
-	SRVDesc.Shader4ComponentMapping = D3D12_DEFAULT_SHADER_4_COMPONENT_MAPPING;
-
-	if (TextureDesc.SampleDesc.Count > 1)
-	{
-		// MS textures can't have mips apparently, so nothing else to set.
-		SRVDesc.ViewDimension = D3D12_SRV_DIMENSION_TEXTURE2DMS;
-	}
-	else
-	{
-		SRVDesc.ViewDimension = D3D12_SRV_DIMENSION_TEXTURE2D;
-		SRVDesc.Texture2D.MostDetailedMip = MipLevel;
-		SRVDesc.Texture2D.MipLevels = NumMipLevels;
-		SRVDesc.Texture2D.PlaneSlice = GetPlaneSliceFromViewFormat(PlatformResourceFormat, PlatformShaderResourceFormat);
-	}
-
-	SRVDesc.Format = PlatformShaderResourceFormat;
-
-	return new FD3D12ShaderResourceView(GetRHIDevice(), &SRVDesc, Texture2D->ResourceLocation);
-}
-
-FShaderResourceViewRHIRef FD3D12DynamicRHI::RHICreateShaderResourceView(FTexture3DRHIParamRef Texture3DRHI, uint8 MipLevel)
-{
-	FD3D12Texture3D*  Texture3D = FD3D12DynamicRHI::ResourceCast(Texture3DRHI);
-
-	uint64 Offset = 0;
-	const D3D12_RESOURCE_DESC& TextureDesc = Texture3D->GetResource()->GetDesc();
-
-	const bool bSRGB = (Texture3D->GetFlags() & TexCreate_SRGB) != 0;
-	const DXGI_FORMAT PlatformShaderResourceFormat = FindShaderResourceDXGIFormat(TextureDesc.Format, bSRGB);
-
-	// Create a Shader Resource View
-	D3D12_SHADER_RESOURCE_VIEW_DESC SRVDesc = {};
-	SRVDesc.Shader4ComponentMapping = D3D12_DEFAULT_SHADER_4_COMPONENT_MAPPING;
-	SRVDesc.ViewDimension = D3D12_SRV_DIMENSION_TEXTURE3D;
-	SRVDesc.Texture3D.MipLevels = 1;
-	SRVDesc.Texture3D.MostDetailedMip = MipLevel;
-
-	SRVDesc.Format = PlatformShaderResourceFormat;
-
-	return new FD3D12ShaderResourceView(GetRHIDevice(), &SRVDesc, Texture3D->ResourceLocation);
-}
-
 /** Generates mip maps for the surface. */
 void FD3D12DynamicRHI::RHIGenerateMips(FTextureRHIParamRef TextureRHI)
 {
@@ -1807,7 +1657,8 @@ void* TD3D12Texture2D<RHIResourceType>::Lock(uint32 MipIndex, uint32 ArrayIndex,
 	check(DestStride == XBytesAligned);
 
 	// Add the lock to the outstanding lock list.
-	GetParentDevice()->GetOwningRHI()->OutstandingLocks.Add(FD3D12LockedKey(GetResource(), Subresource), LockedData);
+	FD3D12LockedKey Key = FD3D12LockedKey(GetResource(), Subresource);
+	GetParentDevice()->GetOwningRHI()->AddToOutstandingLocks(Key, LockedData);
 
 	return (void*)LockedData.GetData();
 }
@@ -1828,8 +1679,8 @@ void TD3D12Texture2D<RHIResourceType>::Unlock(uint32 MipIndex, uint32 ArrayIndex
 	const uint32 MipSizeY = FMath::Max(GetSizeY() >> MipIndex, BlockSizeY);
 
 	// Find the object that is tracking this lock
-	const FD3D12LockedKey LockedKey(GetResource(), Subresource);
-	FD3D12LockedData* LockedData = GetParentDevice()->GetOwningRHI()->OutstandingLocks.Find(LockedKey);
+	FD3D12LockedKey LockedKey(GetResource(), Subresource);
+	FD3D12LockedData* LockedData = GetParentDevice()->GetOwningRHI()->FindInOutstandingLocks(LockedKey);
 	check(LockedData);
 
 #if PLATFORM_SUPPORTS_VIRTUAL_TEXTURES
@@ -1888,7 +1739,7 @@ void TD3D12Texture2D<RHIResourceType>::Unlock(uint32 MipIndex, uint32 ArrayIndex
 	}
 
 	// Remove the lock from the outstanding lock list.
-	GetParentDevice()->GetOwningRHI()->OutstandingLocks.Remove(LockedKey);
+	GetParentDevice()->GetOwningRHI()->RemoveFromOutstandingLocks(LockedKey);
 }
 
 void* FD3D12DynamicRHI::RHILockTexture2D(FTexture2DRHIParamRef TextureRHI, uint32 MipIndex, EResourceLockMode LockMode, uint32& DestStride, bool bLockWithinMiptail)
diff --git a/Engine/Source/Runtime/Windows/D3D12RHI/Private/D3D12UAV.cpp b/Engine/Source/Runtime/Windows/D3D12RHI/Private/D3D12UAV.cpp
--- a/Engine/Source/Runtime/Windows/D3D12RHI/Private/D3D12UAV.cpp
+++ b/Engine/Source/Runtime/Windows/D3D12RHI/Private/D3D12UAV.cpp
@@ -43,7 +43,7 @@ FUnorderedAccessViewRHIRef FD3D12DynamicRHI::RHICreateUnorderedAccessView(FStruc
 
 	if (bNeedsCounterResource)
 	{
-		GetRHIDevice()->GetResourceHelper().CreateBuffer(D3D12_HEAP_TYPE_READBACK, 4, CounterResource.GetInitReference());
+		GetRHIDevice()->GetResourceHelper().CreateBuffer(D3D12_HEAP_TYPE_DEFAULT, 4, CounterResource.GetInitReference(), D3D12_RESOURCE_FLAG_ALLOW_UNORDERED_ACCESS);
 	}
 
 	if (bByteAccessBuffer)
@@ -242,3 +242,35 @@ void FD3D12CommandContext::RHIClearUAV(FUnorderedAccessViewRHIParamRef Unordered
 
 	OwningRHI.RegisterGPUWork(1);
 }
+
+FUnorderedAccessViewRHIRef FD3D12DynamicRHI::RHICreateUnorderedAccessView_RenderThread(FRHICommandListImmediate& RHICmdList, FStructuredBufferRHIParamRef StructuredBufferRHI, bool bUseUAVCounter, bool bAppendBuffer)
+{
+	FD3D12StructuredBuffer* StructuredBuffer = FD3D12DynamicRHI::ResourceCast(StructuredBufferRHI);
+	// TODO: we have to stall the RHI thread when creating SRVs of dynamic buffers because they get renamed.
+	// perhaps we could do a deferred operation?
+	if (StructuredBuffer->GetUsage() & BUF_AnyDynamic)
+	{
+		FScopedRHIThreadStaller StallRHIThread(RHICmdList);
+		return RHICreateUnorderedAccessView(StructuredBufferRHI, bUseUAVCounter, bAppendBuffer);
+	}
+	return RHICreateUnorderedAccessView(StructuredBufferRHI, bUseUAVCounter, bAppendBuffer);
+}
+
+FUnorderedAccessViewRHIRef FD3D12DynamicRHI::RHICreateUnorderedAccessView_RenderThread(FRHICommandListImmediate& RHICmdList, FTextureRHIParamRef Texture, uint32 MipLevel)
+{
+	return FD3D12DynamicRHI::RHICreateUnorderedAccessView(Texture, MipLevel);
+}
+
+FUnorderedAccessViewRHIRef FD3D12DynamicRHI::RHICreateUnorderedAccessView_RenderThread(FRHICommandListImmediate& RHICmdList, FVertexBufferRHIParamRef VertexBufferRHI, uint8 Format)
+{
+	FD3D12VertexBuffer* VertexBuffer = FD3D12DynamicRHI::ResourceCast(VertexBufferRHI);
+
+	// TODO: we have to stall the RHI thread when creating SRVs of dynamic buffers because they get renamed.
+	// perhaps we could do a deferred operation?
+	if (VertexBuffer->GetUsage() & BUF_AnyDynamic)
+	{
+		FScopedRHIThreadStaller StallRHIThread(RHICmdList);
+		return RHICreateUnorderedAccessView(VertexBufferRHI, Format);
+	}
+	return RHICreateUnorderedAccessView(VertexBufferRHI, Format);
+}
\ No newline at end of file
diff --git a/Engine/Source/Runtime/Windows/D3D12RHI/Private/D3D12Util.cpp b/Engine/Source/Runtime/Windows/D3D12RHI/Private/D3D12Util.cpp
--- a/Engine/Source/Runtime/Windows/D3D12RHI/Private/D3D12Util.cpp
+++ b/Engine/Source/Runtime/Windows/D3D12RHI/Private/D3D12Util.cpp
@@ -27,6 +27,8 @@ D3D12Util.h: D3D RHI utility implementation.
 #endif//D3DERR_WASSTILLDRAWING
 #endif
 
+extern bool D3D12RHI_ShouldCreateWithD3DDebug();
+
 void SetName(ID3D12Object* const Object, const TCHAR* const Name)
 {
 	if (Object)
@@ -612,6 +614,80 @@ void FD3D12SyncPoint::WaitForCompletion() const
 	Fence->WaitForFence(Value);
 }
 
+// Forward declarations are required for the template functions
+template bool AssertResourceState(ID3D12CommandList* pCommandList, FD3D12View<D3D12_RENDER_TARGET_VIEW_DESC>* pView, const D3D12_RESOURCE_STATES& State);
+template bool AssertResourceState(ID3D12CommandList* pCommandList, FD3D12View<D3D12_UNORDERED_ACCESS_VIEW_DESC>* pView, const D3D12_RESOURCE_STATES& State);
+template bool AssertResourceState(ID3D12CommandList* pCommandList, FD3D12View<D3D12_SHADER_RESOURCE_VIEW_DESC>* pView, const D3D12_RESOURCE_STATES& State);
+
+template <class TView>
+bool AssertResourceState(ID3D12CommandList* pCommandList, FD3D12View<TView>* pView, const D3D12_RESOURCE_STATES& State)
+{
+	// Check the view
+	if (!pView)
+	{
+		// No need to check null views
+		return true;
+	}
+
+	return AssertResourceState(pCommandList, pView->GetResource(), State, pView->GetViewSubresourceSubset());
+}
+
+bool AssertResourceState(ID3D12CommandList* pCommandList, FD3D12Resource* pResource, const D3D12_RESOURCE_STATES& State, uint32 Subresource)
+{
+	// Check the resource
+	if (!pResource)
+	{
+		// No need to check null resources
+		// Some dynamic SRVs haven't been mapped and updated yet so they actually don't have any backing resources.
+		return true;
+	}
+
+	CViewSubresourceSubset SubresourceSubset(Subresource, pResource->GetMipLevels(), pResource->GetArraySize(), pResource->GetPlaneCount());
+	return AssertResourceState(pCommandList, pResource, State, SubresourceSubset);
+}
+
+bool AssertResourceState(ID3D12CommandList* pCommandList, FD3D12Resource* pResource, const D3D12_RESOURCE_STATES& State, const CViewSubresourceSubset& SubresourceSubset)
+{
+	// Check the resource
+	if (!pResource)
+	{
+		// No need to check null resources
+		// Some dynamic SRVs haven't been mapped and updated yet so they actually don't have any backing resources.
+		return true;
+	}
+
+	// Can only verify resource states if the debug layer is used
+	static const bool bWithD3DDebug = D3D12RHI_ShouldCreateWithD3DDebug();
+	if (!bWithD3DDebug)
+	{
+		UE_LOG(LogD3D12RHI, Fatal, TEXT("*** AssertResourceState requires the debug layer ***"));
+		return false;
+	}
+
+	// Get the debug command queue
+	TRefCountPtr<ID3D12DebugCommandList> pDebugCommandList;
+	VERIFYD3D11RESULT(pCommandList->QueryInterface(pDebugCommandList.GetInitReference()));
+
+	// Get the underlying resource
+	ID3D12Resource* pD3D12Resource = pResource->GetResource();
+	check(pD3D12Resource);
+
+	// For each subresource in the view...
+	for (CViewSubresourceSubset::CViewSubresourceIterator it = SubresourceSubset.begin(); it != SubresourceSubset.end(); ++it)
+	{
+		for (uint32 SubresourceIndex = it.StartSubresource(); SubresourceIndex < it.EndSubresource(); SubresourceIndex++)
+		{
+			const bool bGoodState = !!pDebugCommandList->AssertResourceState(pD3D12Resource, SubresourceIndex, State);
+			if (!bGoodState)
+			{
+				return false;
+			}
+		}
+	}
+
+	return true;
+}
+
 //
 // Stat declarations.
 //
