Merging //UE4/Release-4.11 to //UE4/Main (Up to CL#2897738)

diff --git a/Engine/Source/Runtime/Windows/D3D12RHI/Private/D3D12VertexBuffer.cpp b/Engine/Source/Runtime/Windows/D3D12RHI/Private/D3D12VertexBuffer.cpp
--- a/Engine/Source/Runtime/Windows/D3D12RHI/Private/D3D12VertexBuffer.cpp
+++ b/Engine/Source/Runtime/Windows/D3D12RHI/Private/D3D12VertexBuffer.cpp
@@ -6,13 +6,9 @@
 
 #include "D3D12RHIPrivate.h"
 
-FVertexBufferRHIRef FD3D12DynamicRHI::RHICreateVertexBuffer(uint32 Size, uint32 InUsage, FRHIResourceCreateInfo& CreateInfo)
+D3D12_RESOURCE_DESC CreateVertexBufferResourceDesc(uint32 Size, uint32 InUsage)
 {
-	// Explicitly check that the size is nonzero before allowing CreateVertexBuffer to opaquely fail.
-	check(Size > 0);
-
 	// Describe the vertex buffer.
-	const bool bIsDynamic = (InUsage & BUF_AnyDynamic) ? true : false;
 	D3D12_RESOURCE_DESC Desc = CD3DX12_RESOURCE_DESC::Buffer(Size);
 
 	if (InUsage & BUF_UnorderedAccess)
@@ -40,44 +36,7 @@ FVertexBufferRHIRef FD3D12DynamicRHI::RHICreateVertexBuffer(uint32 Size, uint32
 		}
 	}
 
-	// If a resource array was provided for the resource, create the resource pre-populated
-	D3D12_SUBRESOURCE_DATA InitData ={0};
-	D3D12_SUBRESOURCE_DATA* pInitData = nullptr;
-	if (CreateInfo.ResourceArray)
-	{
-		check(Size == CreateInfo.ResourceArray->GetResourceDataSize());
-		InitData.pData = CreateInfo.ResourceArray->GetResourceData();
-		InitData.RowPitch = Size;
-		InitData.SlicePitch = 0;
-		pInitData = &InitData;
-	}
-
-	TRefCountPtr<FD3D12ResourceLocation> ResourceLocation = new FD3D12ResourceLocation(GetRHIDevice());
-	if (bIsDynamic)
-	{
-		void* pData = GetRHIDevice()->GetDefaultUploadHeapAllocator().AllocUploadResource(Size, DEFAULT_CONTEXT_UPLOAD_POOL_ALIGNMENT, ResourceLocation);
-
-		if (pInitData)
-		{
-			// Handle initial data
-			FMemory::Memcpy(pData, InitData.pData, Size);
-		}
-		check(ResourceLocation->GetEffectiveBufferSize() == Size);
-	}
-	else
-	{
-		// Create an empty resource location
-		VERIFYD3D11RESULT(GetRHIDevice()->GetDefaultBufferAllocator().AllocDefaultResource(Desc, pInitData, ResourceLocation.GetInitReference()));
-		check(ResourceLocation->GetEffectiveBufferSize() == Size);
-	}
-
-	if (CreateInfo.ResourceArray)
-	{
-		// Discard the resource array's contents.
-		CreateInfo.ResourceArray->Discard();
-	}
-
-	return new FD3D12VertexBuffer(GetRHIDevice(), ResourceLocation, Size, InUsage);
+	return Desc;
 }
 
 FD3D12VertexBuffer::~FD3D12VertexBuffer()
@@ -90,18 +49,19 @@ FD3D12VertexBuffer::~FD3D12VertexBuffer()
 	// Dynamic resources have a permanente entry in the OutstandingLocks that needs to be cleaned up. 
 	if (GetUsage() & BUF_AnyDynamic)
 	{
-		FD3D12LockedKey LockedKey(ResourceLocation.GetReference());
-		FD3D12DynamicRHI::GetD3DRHI()->OutstandingLocks.Remove(LockedKey);
+		FD3D12LockedKey LockedKey(this);
+		FD3D12DynamicRHI::GetD3DRHI()->RemoveFromOutstandingLocks(LockedKey);
 	}
 }
 
-void* FD3D12VertexBuffer::DynamicLock()
+void FD3D12VertexBuffer::Rename(FD3D12ResourceLocation* NewResource)
 {
-	void* pData = nullptr;
-	check((GetUsage() & BUF_AnyDynamic) ? true : false);
+	FD3D12CommandContext& DefaultContext = GetParentDevice()->GetDefaultCommandContext();
+
+	// If this resource is bound to the device, unbind it
+	DefaultContext.ConditionalClearShaderResource(ResourceLocation);
 
-	FD3D12DynamicHeapAllocator &Allocator = GetParentDevice()->GetDefaultUploadHeapAllocator();
-	pData = Allocator.AllocUploadResource(ResourceLocation->GetEffectiveBufferSize(), DEFAULT_CONTEXT_UPLOAD_POOL_ALIGNMENT, ResourceLocation);
+	ResourceLocation = NewResource;
 
 	if (DynamicSRV != nullptr)
 	{
@@ -110,168 +70,58 @@ void* FD3D12VertexBuffer::DynamicLock()
 		Desc.ptr = 0;
 		DynamicSRV->Rename(ResourceLocation, Desc, 0);
 	}
+}
 
-	check(pData);
-
-	return pData;
+FVertexBufferRHIRef FD3D12DynamicRHI::RHICreateVertexBuffer(uint32 Size, uint32 InUsage, FRHIResourceCreateInfo& CreateInfo)
+{
+	const D3D12_RESOURCE_DESC Desc = CreateVertexBufferResourceDesc(Size, InUsage);
+	const uint32 Alignment = 4;
+
+	FD3D12ResourceLocation* ResourceLocation = CreateBuffer(nullptr, Desc, Size, InUsage, CreateInfo, Alignment);
+	FD3D12VertexBuffer* NewBuffer = new FD3D12VertexBuffer(GetRHIDevice(), ResourceLocation, Size, InUsage);
+	UpdateBufferStats(NewBuffer->ResourceLocation, true, D3D12_BUFFER_TYPE_VERTEX);
+	NewBuffer->BufferAlignment = Alignment;
+	return NewBuffer;
 }
 
 void* FD3D12DynamicRHI::RHILockVertexBuffer(FVertexBufferRHIParamRef VertexBufferRHI, uint32 Offset, uint32 Size, EResourceLockMode LockMode)
 {
-	// TEMP: Disable check due to engine bug where it passes Size = 0 when doing parallel rendering.
-	//check(Size > 0);
-
-	FD3D12VertexBuffer*  VertexBuffer = FD3D12DynamicRHI::ResourceCast(VertexBufferRHI);
-
-	FD3D12CommandContext& DefaultContext = GetRHIDevice()->GetDefaultCommandContext();
-	FD3D12CommandListHandle& hCommandList = DefaultContext.CommandListHandle;
-
-	// If this resource is bound to the device, unbind it
-	DefaultContext.ConditionalClearShaderResource(VertexBuffer->ResourceLocation);
-
-	// Determine whether the vertex buffer is dynamic or not.
-	const bool bIsDynamic = (VertexBuffer->GetUsage() & BUF_AnyDynamic) ? true : false;
-
-	FD3D12LockedKey LockedKey(VertexBuffer->ResourceLocation.GetReference());
-	FD3D12LockedData LockedData;
-
-	if (bIsDynamic)
-	{
-		check(LockMode == RLM_WriteOnly);
-
-		void* pData = VertexBuffer->DynamicLock();
-
-		// Add the lock to the lock map.
-		LockedData.SetData(pData);
-		LockedData.Pitch = Size;
-		OutstandingLocks.Add(LockedKey, LockedData);
-
-		return (void*)((uint8*)pData + Offset);
-	}
-	else
-	{
-		FD3D12Resource* pResource = VertexBuffer->ResourceLocation->GetResource();
-		uint32 ResourceSize = VertexBuffer->ResourceLocation->GetEffectiveBufferSize();
-
-		if (LockMode == RLM_ReadOnly)
-		{
-			// If the static buffer is being locked for reading, create a staging buffer.
-			TRefCountPtr<FD3D12Resource> StagingVertexBuffer;
-			VERIFYD3D11RESULT(GetRHIDevice()->GetResourceHelper().CreateBuffer(D3D12_HEAP_TYPE_READBACK, Offset + Size, StagingVertexBuffer.GetInitReference()));
-			LockedData.StagingResource = StagingVertexBuffer;
-
-			// Copy the contents of the vertex buffer to the staging buffer.
-			{
-				FScopeResourceBarrier ScopeResourceBarrierSource(hCommandList, pResource, pResource->GetDefaultResourceState(), D3D12_RESOURCE_STATE_COPY_SOURCE, 0);
-				// Don't need to transition upload heaps
-
-				DefaultContext.numCopies++;
-				hCommandList->CopyBufferRegion(
-					StagingVertexBuffer->GetResource(),
-					0,
-					pResource->GetResource(),
-					Offset, Size);
-			}
-
-			GetRHIDevice()->GetDefaultCommandContext().FlushCommands(true);
-
-			// Map the staging buffer's memory for reading.
-			void* pData;
-			VERIFYD3D11RESULT(StagingVertexBuffer->GetResource()->Map(0, nullptr, &pData));
-			LockedData.SetData(pData);
-			LockedData.Pitch = Offset + Size;
-		}
-		else
-		{
-			// If the static buffer is being locked for writing, allocate memory for the contents to be written to.
-
-			// Use an upload heap to copy data to a default resource.
-
-			// Use an upload heap for dynamic resources and map its memory for writing.
-			TRefCountPtr<FD3D12ResourceLocation> pUploadBufferLocation = new FD3D12ResourceLocation(GetRHIDevice());
-			void* pData = GetRHIDevice()->GetDefaultFastAllocator().Allocate(Offset + Size, D3D12_TEXTURE_DATA_PLACEMENT_ALIGNMENT, pUploadBufferLocation);
-
-			if (nullptr == pData)
-			{
-				return nullptr;
-			}
-
-
-			// Add the lock to the lock map.
-			FD3D12LockedData LockedData;
-			LockedData.SetData(pData);
-			LockedData.Pitch = Offset + Size;
-
-			// Keep track of the underlying resource for the upload heap so it can be referenced during Unmap.
-			LockedData.UploadHeapLocation = pUploadBufferLocation;
-			OutstandingLocks.Add(LockedKey, LockedData);
-
-			return (void*)((uint8*)pData + Offset);
-		}
-	}
-
-	// Add the lock to the lock map.
-	OutstandingLocks.Add(LockedKey, LockedData);
-
-	// Return the offset pointer
-	return (void*)((uint8*)LockedData.GetData() + Offset);
+	return LockBuffer(nullptr, FD3D12DynamicRHI::ResourceCast(VertexBufferRHI), Offset, Size, LockMode);
 }
 
 void FD3D12DynamicRHI::RHIUnlockVertexBuffer(FVertexBufferRHIParamRef VertexBufferRHI)
 {
-	FD3D12VertexBuffer*  VertexBuffer = FD3D12DynamicRHI::ResourceCast(VertexBufferRHI);
-
-	// Determine whether the vertex buffer is dynamic or not.
-	const bool bIsDynamic = (VertexBuffer->GetUsage() & BUF_AnyDynamic) ? true : false;
-
-	// Find the outstanding lock for this VB.
-	FD3D12LockedKey LockedKey(VertexBuffer->ResourceLocation);
-	FD3D12LockedData* LockedData = OutstandingLocks.Find(LockedKey);
-	check(LockedData);
-
-	if (bIsDynamic)
-	{
-		// If the VB is dynamic, its upload heap memory can always stay mapped. Don't do anything.
-	}
-	else
-	{
-		// If the static VB lock involved a staging resource, it was locked for reading.
-		if (LockedData->StagingResource)
-		{
-			// Unmap the staging buffer's memory.
-			ID3D12Resource* StagingBuffer = LockedData->StagingResource.GetReference()->GetResource();
-			StagingBuffer->Unmap(0, nullptr);
-		}
-		else
-		{
-			FD3D12Resource* pResource = VertexBuffer->ResourceLocation->GetResource();
-
-			// Copy the contents of the temporary memory buffer allocated for writing into the IB.
-			FD3D12ResourceLocation* pUploadHeapLocation = LockedData->UploadHeapLocation.GetReference();
-			FD3D12Resource* UploadHeap = pUploadHeapLocation->GetResource();
+	UnlockBuffer(nullptr, FD3D12DynamicRHI::ResourceCast(VertexBufferRHI));
+}
 
-			FD3D12CommandContext& defaultContext = GetRHIDevice()->GetDefaultCommandContext();
-			FD3D12CommandListHandle& hCommandList = defaultContext.CommandListHandle;
+FVertexBufferRHIRef FD3D12DynamicRHI::CreateVertexBuffer_RenderThread(FRHICommandListImmediate& RHICmdList, uint32 Size, uint32 InUsage, FRHIResourceCreateInfo& CreateInfo)
+{
+	const D3D12_RESOURCE_DESC Desc = CreateVertexBufferResourceDesc(Size, InUsage);
+	const uint32 Alignment = 4;
+
+	FD3D12ResourceLocation* ResourceLocation = CreateBuffer(&RHICmdList, Desc, Size, InUsage, CreateInfo, Alignment);
+	FD3D12VertexBuffer* NewBuffer = new FD3D12VertexBuffer(GetRHIDevice(), ResourceLocation, Size, InUsage);
+	UpdateBufferStats(NewBuffer->ResourceLocation, true, D3D12_BUFFER_TYPE_VERTEX);
+	NewBuffer->BufferAlignment = Alignment;
+	return NewBuffer;
+}
 
-			{
-				FScopeResourceBarrier ScopeResourceBarrierDest(hCommandList, pResource, pResource->GetDefaultResourceState(), D3D12_RESOURCE_STATE_COPY_DEST, 0);
-				// Don't need to transition upload heaps
+void* FD3D12DynamicRHI::LockVertexBuffer_RenderThread(class FRHICommandListImmediate& RHICmdList, FVertexBufferRHIParamRef VertexBufferRHI, uint32 Offset, uint32 SizeRHI, EResourceLockMode LockMode)
+{
+	// Pull down the above RHI implementation so that we can flush only when absolutely necessary
+	QUICK_SCOPE_CYCLE_COUNTER(STAT_FDynamicRHI_LockVertexBuffer_RenderThread);
+	check(IsInRenderingThread());
 
-				defaultContext.numCopies++;
-				hCommandList->CopyBufferRegion(
-					pResource->GetResource(),
-					VertexBuffer->ResourceLocation->GetOffset(),
-					UploadHeap->GetResource(),
-					pUploadHeapLocation->GetOffset(), LockedData->Pitch);
-			}
+	return LockBuffer(&RHICmdList, FD3D12DynamicRHI::ResourceCast(VertexBufferRHI), Offset, SizeRHI, LockMode);
+}
 
-			DEBUG_RHI_EXECUTE_COMMAND_LIST(this);
-		}
-	}
+void FD3D12DynamicRHI::UnlockVertexBuffer_RenderThread(FRHICommandListImmediate& RHICmdList, FVertexBufferRHIParamRef VertexBufferRHI)
+{
+	// Pull down the above RHI implementation so that we can flush only when absolutely necessary
+	QUICK_SCOPE_CYCLE_COUNTER(STAT_FDynamicRHI_UnlockVertexBuffer_RenderThread);
+	check(IsInRenderingThread());
 
-	// Remove the FD3D12LockedData from the lock map.
-	// If the lock involved a staging resource, this releases it.
-	OutstandingLocks.Remove(LockedKey);
+	UnlockBuffer(&RHICmdList, FD3D12DynamicRHI::ResourceCast(VertexBufferRHI));
 }
 
 void FD3D12DynamicRHI::RHICopyVertexBuffer(FVertexBufferRHIParamRef SourceBufferRHI, FVertexBufferRHIParamRef DestBufferRHI)
diff --git a/Engine/Source/Runtime/Windows/D3D12RHI/Private/D3D12View.cpp b/Engine/Source/Runtime/Windows/D3D12RHI/Private/D3D12View.cpp
new file mode 100644
--- /dev/null
+++ b/Engine/Source/Runtime/Windows/D3D12RHI/Private/D3D12View.cpp
@@ -0,0 +1,187 @@
+// Copyright 1998-2016 Epic Games, Inc. All Rights Reserved.
+
+/*=============================================================================
+	D3D12View.cpp: 
+=============================================================================*/
+
+#include "D3D12RHIPrivate.h"
+
+FShaderResourceViewRHIRef FD3D12DynamicRHI::RHICreateShaderResourceView(FTexture2DRHIParamRef Texture2DRHI, uint8 MipLevel)
+{
+	FD3D12Texture2D*  Texture2D = FD3D12DynamicRHI::ResourceCast(Texture2DRHI);
+
+	uint64 Offset = 0;
+	const D3D12_RESOURCE_DESC& TextureDesc = Texture2D->GetResource()->GetDesc();
+
+	const bool bSRGB = (Texture2D->GetFlags() & TexCreate_SRGB) != 0;
+	const DXGI_FORMAT PlatformShaderResourceFormat = FindShaderResourceDXGIFormat(TextureDesc.Format, bSRGB);
+
+	// Create a Shader Resource View
+	D3D12_SHADER_RESOURCE_VIEW_DESC SRVDesc = {};
+	SRVDesc.Shader4ComponentMapping = D3D12_DEFAULT_SHADER_4_COMPONENT_MAPPING;
+	SRVDesc.ViewDimension = D3D12_SRV_DIMENSION_TEXTURE2D;
+	SRVDesc.Texture2D.MipLevels = 1;
+	SRVDesc.Texture2D.MostDetailedMip = MipLevel;
+
+	SRVDesc.Format = PlatformShaderResourceFormat;
+
+	SRVDesc.Texture2D.PlaneSlice = GetPlaneSliceFromViewFormat(TextureDesc.Format, SRVDesc.Format);
+
+	return new FD3D12ShaderResourceView(GetRHIDevice(), &SRVDesc, Texture2D->ResourceLocation);
+}
+
+FShaderResourceViewRHIRef FD3D12DynamicRHI::RHICreateShaderResourceView(FTexture2DArrayRHIParamRef Texture2DRHI, uint8 MipLevel)
+{
+	FD3D12Texture2DArray*  Texture2DArray = FD3D12DynamicRHI::ResourceCast(Texture2DRHI);
+
+	uint64 Offset = 0;
+	const D3D12_RESOURCE_DESC& TextureDesc = Texture2DArray->GetResource()->GetDesc();
+
+	const bool bSRGB = (Texture2DArray->GetFlags() & TexCreate_SRGB) != 0;
+	const DXGI_FORMAT PlatformShaderResourceFormat = FindShaderResourceDXGIFormat(TextureDesc.Format, bSRGB);
+
+	// Create a Shader Resource View
+	D3D12_SHADER_RESOURCE_VIEW_DESC SRVDesc = {};
+	SRVDesc.Shader4ComponentMapping = D3D12_DEFAULT_SHADER_4_COMPONENT_MAPPING;
+	SRVDesc.ViewDimension = D3D12_SRV_DIMENSION_TEXTURE2DARRAY;
+	SRVDesc.Texture2DArray.ArraySize = TextureDesc.DepthOrArraySize;
+	SRVDesc.Texture2DArray.FirstArraySlice = 0;
+	SRVDesc.Texture2DArray.MipLevels = 1;
+	SRVDesc.Texture2DArray.MostDetailedMip = MipLevel;
+
+	SRVDesc.Format = PlatformShaderResourceFormat;
+
+	SRVDesc.Texture2DArray.PlaneSlice = GetPlaneSliceFromViewFormat(TextureDesc.Format, SRVDesc.Format);
+
+	return new FD3D12ShaderResourceView(GetRHIDevice(), &SRVDesc, Texture2DArray->ResourceLocation);
+}
+
+FShaderResourceViewRHIRef FD3D12DynamicRHI::RHICreateShaderResourceView(FTextureCubeRHIParamRef TextureCubeRHI, uint8 MipLevel)
+{
+	FD3D12TextureCube*  TextureCube = FD3D12DynamicRHI::ResourceCast(TextureCubeRHI);
+
+	uint64 Offset = 0;
+	const D3D12_RESOURCE_DESC& TextureDesc = TextureCube->GetResource()->GetDesc();
+
+	const bool bSRGB = (TextureCube->GetFlags() & TexCreate_SRGB) != 0;
+	const DXGI_FORMAT PlatformShaderResourceFormat = FindShaderResourceDXGIFormat(TextureDesc.Format, bSRGB);
+
+	// Create a Shader Resource View
+	D3D12_SHADER_RESOURCE_VIEW_DESC SRVDesc = {};
+	SRVDesc.Shader4ComponentMapping = D3D12_DEFAULT_SHADER_4_COMPONENT_MAPPING;
+	SRVDesc.ViewDimension = D3D12_SRV_DIMENSION_TEXTURECUBE;
+	SRVDesc.TextureCube.MipLevels = 1;
+	SRVDesc.TextureCube.MostDetailedMip = MipLevel;
+
+	SRVDesc.Format = PlatformShaderResourceFormat;
+
+	return new FD3D12ShaderResourceView(GetRHIDevice(), &SRVDesc, TextureCube->ResourceLocation);
+}
+
+FShaderResourceViewRHIRef FD3D12DynamicRHI::RHICreateShaderResourceView(FTexture2DRHIParamRef Texture2DRHI, uint8 MipLevel, uint8 NumMipLevels, uint8 Format)
+{
+	FD3D12Texture2D*  Texture2D = FD3D12DynamicRHI::ResourceCast(Texture2DRHI);
+	uint64 Offset = 0;
+	const D3D12_RESOURCE_DESC& TextureDesc = Texture2D->GetResource()->GetDesc();
+
+	const DXGI_FORMAT PlatformResourceFormat = FD3D12DynamicRHI::GetPlatformTextureResourceFormat((DXGI_FORMAT)GPixelFormats[Format].PlatformFormat, Texture2D->GetFlags());
+
+	const bool bSRGB = (Texture2D->GetFlags() & TexCreate_SRGB) != 0;
+	const DXGI_FORMAT PlatformShaderResourceFormat = FindShaderResourceDXGIFormat(PlatformResourceFormat, bSRGB);
+
+	// Create a Shader Resource View
+	D3D12_SHADER_RESOURCE_VIEW_DESC SRVDesc = {};
+	SRVDesc.Shader4ComponentMapping = D3D12_DEFAULT_SHADER_4_COMPONENT_MAPPING;
+
+	if (TextureDesc.SampleDesc.Count > 1)
+	{
+		// MS textures can't have mips apparently, so nothing else to set.
+		SRVDesc.ViewDimension = D3D12_SRV_DIMENSION_TEXTURE2DMS;
+	}
+	else
+	{
+		SRVDesc.ViewDimension = D3D12_SRV_DIMENSION_TEXTURE2D;
+		SRVDesc.Texture2D.MostDetailedMip = MipLevel;
+		SRVDesc.Texture2D.MipLevels = NumMipLevels;
+		SRVDesc.Texture2D.PlaneSlice = GetPlaneSliceFromViewFormat(PlatformResourceFormat, PlatformShaderResourceFormat);
+	}
+
+	SRVDesc.Format = PlatformShaderResourceFormat;
+
+	return new FD3D12ShaderResourceView(GetRHIDevice(), &SRVDesc, Texture2D->ResourceLocation);
+}
+
+FShaderResourceViewRHIRef FD3D12DynamicRHI::RHICreateShaderResourceView(FTexture3DRHIParamRef Texture3DRHI, uint8 MipLevel)
+{
+	FD3D12Texture3D*  Texture3D = FD3D12DynamicRHI::ResourceCast(Texture3DRHI);
+
+	uint64 Offset = 0;
+	const D3D12_RESOURCE_DESC& TextureDesc = Texture3D->GetResource()->GetDesc();
+
+	const bool bSRGB = (Texture3D->GetFlags() & TexCreate_SRGB) != 0;
+	const DXGI_FORMAT PlatformShaderResourceFormat = FindShaderResourceDXGIFormat(TextureDesc.Format, bSRGB);
+
+	// Create a Shader Resource View
+	D3D12_SHADER_RESOURCE_VIEW_DESC SRVDesc = {};
+	SRVDesc.Shader4ComponentMapping = D3D12_DEFAULT_SHADER_4_COMPONENT_MAPPING;
+	SRVDesc.ViewDimension = D3D12_SRV_DIMENSION_TEXTURE3D;
+	SRVDesc.Texture3D.MipLevels = 1;
+	SRVDesc.Texture3D.MostDetailedMip = MipLevel;
+
+	SRVDesc.Format = PlatformShaderResourceFormat;
+
+	return new FD3D12ShaderResourceView(GetRHIDevice(), &SRVDesc, Texture3D->ResourceLocation);
+}
+
+FShaderResourceViewRHIRef FD3D12DynamicRHI::RHICreateShaderResourceView_RenderThread(FRHICommandListImmediate& RHICmdList, FTexture2DRHIParamRef Texture2DRHI, uint8 MipLevel)
+{
+	return RHICreateShaderResourceView(Texture2DRHI, MipLevel);
+}
+
+FShaderResourceViewRHIRef FD3D12DynamicRHI::RHICreateShaderResourceView_RenderThread(FRHICommandListImmediate& RHICmdList, FTexture2DRHIParamRef Texture2DRHI, uint8 MipLevel, uint8 NumMipLevels, uint8 Format)
+{
+	return RHICreateShaderResourceView(Texture2DRHI, MipLevel, NumMipLevels, Format);
+}
+
+FShaderResourceViewRHIRef FD3D12DynamicRHI::RHICreateShaderResourceView_RenderThread(FRHICommandListImmediate& RHICmdList, FTexture3DRHIParamRef Texture3DRHI, uint8 MipLevel)
+{
+	return RHICreateShaderResourceView(Texture3DRHI, MipLevel);
+}
+
+FShaderResourceViewRHIRef FD3D12DynamicRHI::RHICreateShaderResourceView_RenderThread(FRHICommandListImmediate& RHICmdList, FTexture2DArrayRHIParamRef Texture2DArrayRHI, uint8 MipLevel)
+{
+	return RHICreateShaderResourceView(Texture2DArrayRHI, MipLevel);
+}
+
+FShaderResourceViewRHIRef FD3D12DynamicRHI::RHICreateShaderResourceView_RenderThread(FRHICommandListImmediate& RHICmdList, FTextureCubeRHIParamRef TextureCubeRHI, uint8 MipLevel)
+{
+	return RHICreateShaderResourceView(TextureCubeRHI, MipLevel);
+}
+
+FShaderResourceViewRHIRef FD3D12DynamicRHI::RHICreateShaderResourceView_RenderThread(FRHICommandListImmediate& RHICmdList, FVertexBufferRHIParamRef VertexBufferRHI, uint32 Stride, uint8 Format)
+{
+	FD3D12VertexBuffer*  VertexBuffer = FD3D12DynamicRHI::ResourceCast(VertexBufferRHI);
+
+	// TODO: we have to stall the RHI thread when creating SRVs of dynamic buffers because they get renamed.
+	// perhaps we could do a deferred operation?
+	if (VertexBuffer->GetUsage() & BUF_AnyDynamic)
+	{
+		FScopedRHIThreadStaller StallRHIThread(RHICmdList);
+		return RHICreateShaderResourceView(VertexBufferRHI, Stride, Format);
+	}
+	return RHICreateShaderResourceView(VertexBufferRHI, Stride, Format);
+}
+
+FShaderResourceViewRHIRef FD3D12DynamicRHI::RHICreateShaderResourceView_RenderThread(FRHICommandListImmediate& RHICmdList, FStructuredBufferRHIParamRef StructuredBufferRHI)
+{
+	FD3D12StructuredBuffer*  StructuredBuffer = FD3D12DynamicRHI::ResourceCast(StructuredBufferRHI);
+
+	// TODO: we have to stall the RHI thread when creating SRVs of dynamic buffers because they get renamed.
+	// perhaps we could do a deferred operation?
+	if (StructuredBuffer->GetUsage() & BUF_AnyDynamic)
+	{
+		FScopedRHIThreadStaller StallRHIThread(RHICmdList);
+		return RHICreateShaderResourceView(StructuredBufferRHI);
+	}
+	return RHICreateShaderResourceView(StructuredBufferRHI);
+}
\ No newline at end of file
diff --git a/Engine/Source/Runtime/Windows/D3D12RHI/Private/D3D12Viewport.cpp b/Engine/Source/Runtime/Windows/D3D12RHI/Private/D3D12Viewport.cpp
--- a/Engine/Source/Runtime/Windows/D3D12RHI/Private/D3D12Viewport.cpp
+++ b/Engine/Source/Runtime/Windows/D3D12RHI/Private/D3D12Viewport.cpp
@@ -165,6 +165,8 @@ FD3D12Texture2D* GetSwapChainSurface(FD3D12Device* Parent, EPixelFormat PixelFor
 	D3D11TextureAllocated2D(*NewTexture);
 
 	NewTexture->DoNoDeferDelete();
+	pWrappedShaderResourceView->DoNoDeferDelete();
+	BackBufferRenderTargetView->DoNoDeferDelete();
 
 	return NewTexture;
 }
@@ -314,7 +316,7 @@ bool FD3D12Viewport::PresentChecked(int32 SyncInterval)
 		// Signal the frame is complete.
 		GetParentDevice()->GetCommandListManager().SignalFrameComplete();
 
-#if UE_BUILD_DEBUG
+#if 0 //UE_BUILD_DEBUG
 		if (RHIConsoleVariables::DumpStatsEveryNFrames > 0)
 		{
 			GetParentDevice()->GetOwningRHI()->DrawCount = 0;
@@ -322,9 +324,7 @@ bool FD3D12Viewport::PresentChecked(int32 SyncInterval)
 
 			if (GetParentDevice()->GetOwningRHI()->PresentCount % RHIConsoleVariables::DumpStatsEveryNFrames == 0)
 			{
-				UE_LOG(LogD3D12RHI, Log, TEXT("*** PRESENT ***"), this);
-				FOutputDeviceRedirector* pOutputDevice = FOutputDeviceRedirector::Get();
-				GetParentDevice()->GetDefaultUploadHeapAllocator().DumpAllocatorStats(*pOutputDevice);
+				//TODO: Dump allocator stats.
 			}
 		}
 #endif
diff --git a/Engine/Source/Runtime/Windows/D3D12RHI/Private/Windows/WindowsD3D12Device.cpp b/Engine/Source/Runtime/Windows/D3D12RHI/Private/Windows/WindowsD3D12Device.cpp
--- a/Engine/Source/Runtime/Windows/D3D12RHI/Private/Windows/WindowsD3D12Device.cpp
+++ b/Engine/Source/Runtime/Windows/D3D12RHI/Private/Windows/WindowsD3D12Device.cpp
@@ -13,6 +13,7 @@
 #pragma comment(lib, "d3d12.lib")
 
 extern bool D3D12RHI_ShouldCreateWithD3DDebug();
+extern bool D3D12RHI_ShouldCreateWithWarp();
 extern bool D3D12RHI_ShouldAllowAsyncResourceCreation();
 
 static TAutoConsoleVariable<int32> CVarGraphicsAdapter(
@@ -259,6 +260,7 @@ void FD3D12DynamicRHIModule::FindAdapter()
 	bool bIsAnyAMD = false;
 	bool bIsAnyIntel = false;
 	bool bIsAnyNVIDIA = false;
+	bool bRequestedWARP = D3D12RHI_ShouldCreateWithWarp();
 
 	// Enumerate the DXGIFactory's adapters.
 	for (uint32 AdapterIndex = 0; DXGIFactory->EnumAdapters(AdapterIndex, TempAdapter.GetInitReference()) != DXGI_ERROR_NOT_FOUND; ++AdapterIndex)
@@ -291,7 +293,7 @@ void FD3D12DynamicRHIModule::FindAdapter()
 				bool bIsAMD = AdapterDesc.VendorId == 0x1002;
 				bool bIsIntel = AdapterDesc.VendorId == 0x8086;
 				bool bIsNVIDIA = AdapterDesc.VendorId == 0x10DE;
-				bool bIsWARP = FParse::Param(FCommandLine::Get(), TEXT("warp"));
+				bool bIsWARP = AdapterDesc.VendorId == 0x1414;
 
 				if (bIsAMD) bIsAnyAMD = true;
 				if (bIsIntel) bIsAnyIntel = true;
@@ -304,6 +306,12 @@ void FD3D12DynamicRHIModule::FindAdapter()
 
 				FD3D12Adapter CurrentAdapter(AdapterIndex, ActualFeatureLevel);
 
+				if (bRequestedWARP && !bIsWARP)
+				{
+					// Requested WARP, reject all other adapters.
+					continue;
+				}
+
 				if (!OutputCount && !bIsWARP)
 				{
 					// Add special check to support WARP, which does not have an output associated with it.
@@ -391,8 +399,6 @@ void FD3D12DynamicRHI::PostInit()
 
 void FD3D12DynamicRHI::PerRHISetup(FD3D12Device* MainDevice)
 {
-	bool bIsWARP = FParse::Param(FCommandLine::Get(), TEXT("warp"));
-
 	check(!GIsRHIInitialized);
 
 	DXGI_ADAPTER_DESC* AdapterDesc = MainDevice->GetD3DAdapterDesc();
@@ -404,6 +410,16 @@ void FD3D12DynamicRHI::PerRHISetup(FD3D12Device* MainDevice)
     GRHIVendorId = AdapterDesc->VendorId;
 	GRHIDeviceId = AdapterDesc->DeviceId;
 
+	// Copied from the D3D11 RHI but disabled for now as this doesn't exist for UT.
+	//// get driver version (todo: share with other RHIs)
+	//{
+	//	FPlatformMisc::GetGPUDriverInfo(GRHIAdapterName, GRHIAdapterInternalDriverVersion, GRHIAdapterUserDriverVersion, GRHIAdapterDriverDate);
+
+	//	UE_LOG(LogD3D12RHI, Log, TEXT("    Adapter Name: %s"), *GRHIAdapterName);
+	//	UE_LOG(LogD3D12RHI, Log, TEXT("  Driver Version: %s (internal %s)"), *GRHIAdapterUserDriverVersion, *GRHIAdapterInternalDriverVersion);
+	//	UE_LOG(LogD3D12RHI, Log, TEXT("     Driver Date: %s"), *GRHIAdapterDriverDate);
+	//}
+
 	// Issue: 32bit windows doesn't report 64bit value, we take what we get.
 	FD3D12GlobalStats::GDedicatedVideoMemory = int64(AdapterDesc->DedicatedVideoMemory);
 	FD3D12GlobalStats::GDedicatedSystemMemory = int64(AdapterDesc->DedicatedSystemMemory);
@@ -416,38 +432,10 @@ void FD3D12DynamicRHI::PerRHISetup(FD3D12Device* MainDevice)
 	int64 ConsideredSharedSystemMemory = FMath::Min(FD3D12GlobalStats::GSharedSystemMemory / 2ll, TotalPhysicalMemory / 4ll);
 
 	IDXGIAdapter3* DxgiAdapter3 = MainDevice->GetAdapter3();
-	TRefCountPtr<IDXGIAdapter> EnumAdapter;
-
-	// Hack to handle a kernel bug where QueryVideomMemoryInfo can't be called on WARP
-	DXGI_QUERY_VIDEO_MEMORY_INFO LocalVideoMemoryInfo ={};
-	if (!bIsWARP)
-	{
-		FD3D12GlobalStats::GTotalGraphicsMemory = 0;
-		if (DXGIFactory->EnumAdapters(MainDevice->GetAdapterIndex(), EnumAdapter.GetInitReference()) != DXGI_ERROR_NOT_FOUND)
-		{
-			if (EnumAdapter)
-			{
-				VERIFYD3D11RESULT(DxgiAdapter3->QueryVideoMemoryInfo(0, DXGI_MEMORY_SEGMENT_GROUP_LOCAL, &LocalVideoMemoryInfo));
-				FD3D12GlobalStats::GTotalGraphicsMemory = LocalVideoMemoryInfo.Budget;
-			}
-		}
-
-	}
-	else
-	{
-		FD3D12GlobalStats::GTotalGraphicsMemory = FD3D12GlobalStats::GDedicatedSystemMemory;
-	}
-
-	GMaxTextureMipCount = FMath::CeilLogTwo(GMaxTextureDimensions) + 1;
-	GMaxTextureMipCount = FMath::Min<int32>(MAX_TEXTURE_MIP_COUNT, GMaxTextureMipCount);
-	// If we only have <= 2GB of GPU memory to work with, we should cut down on the mip count
-	// "GMaxTextureMipCount = 9" allows Intel to run the 4.8 Elemental demo, ideally this handling 
-	// is unnecessary as the upper engine should be able to look at GTotalGraphicsMemory to determine 
-	// the appropriate mip level but that doesn't appear to happen correctly on start-up.
-	if (IsRHIDeviceIntel() && LocalVideoMemoryInfo.Budget <= 2ll * 1024ll * 1024ll * 1024ll)
-	{
-		GMaxTextureMipCount = FMath::Min<int32>(GMaxTextureMipCount, 9);
-	}
+	DXGI_QUERY_VIDEO_MEMORY_INFO LocalVideoMemoryInfo;
+	VERIFYD3D11RESULT(DxgiAdapter3->QueryVideoMemoryInfo(0, DXGI_MEMORY_SEGMENT_GROUP_LOCAL, &LocalVideoMemoryInfo));
+	const int64 TargetBudget = LocalVideoMemoryInfo.Budget * 0.90f;	// Target using 90% of our budget to account for some fragmentation.
+	FD3D12GlobalStats::GTotalGraphicsMemory = TargetBudget;
 
 	if (sizeof(SIZE_T) < 8)
 	{
@@ -460,13 +448,6 @@ void FD3D12DynamicRHI::PerRHISetup(FD3D12Device* MainDevice)
 		FD3D12GlobalStats::GTotalGraphicsMemory = FMath::Min(FD3D12GlobalStats::GTotalGraphicsMemory, 1945ll * 1024ll * 1024ll);
 	}
 
-	if (!bIsWARP)
-	{
-		VERIFYD3D11RESULT(DxgiAdapter3->SetVideoMemoryReservation(0, DXGI_MEMORY_SEGMENT_GROUP_LOCAL, FMath::Min((int64)LocalVideoMemoryInfo.AvailableForReservation, FD3D12GlobalStats::GTotalGraphicsMemory)));
-	}
-
-
-
 	if (GPoolSizeVRAMPercentage > 0)
 	{
 		float PoolSize = float(GPoolSizeVRAMPercentage) * 0.01f * float(FD3D12GlobalStats::GTotalGraphicsMemory);
@@ -480,6 +461,10 @@ void FD3D12DynamicRHI::PerRHISetup(FD3D12Device* MainDevice)
 			FD3D12GlobalStats::GTotalGraphicsMemory / 1024 / 1024);
 	}
 
+	RequestedTexturePoolSize = GTexturePoolSize;
+
+	VERIFYD3D11RESULT(DxgiAdapter3->SetVideoMemoryReservation(0, DXGI_MEMORY_SEGMENT_GROUP_LOCAL, FMath::Min((int64)LocalVideoMemoryInfo.AvailableForReservation, FD3D12GlobalStats::GTotalGraphicsMemory)));
+
 #if (UE_BUILD_SHIPPING && WITH_EDITOR) && PLATFORM_WINDOWS && !PLATFORM_64BITS
 	// Disable PIX for windows in the shipping editor builds
 	D3DPERF_SetOptions(1);
@@ -493,14 +478,6 @@ void FD3D12DynamicRHI::PerRHISetup(FD3D12Device* MainDevice)
 	GShaderPlatformForFeatureLevel[ERHIFeatureLevel::SM4] = SP_PCD3D_SM4;
 	GShaderPlatformForFeatureLevel[ERHIFeatureLevel::SM5] = SP_PCD3D_SM5;
 
-	if (IsRHIDeviceIntel())
-	{
-		// MS: To work around current known issue on Intel, TextureQuality be '0'.
-		// SetQualityLevels responsible for setting this value, must be re-run now after device creation.
-		Scalability::FQualityLevels currentQualityLevels = Scalability::GetQualityLevels();
-		Scalability::SetQualityLevels(currentQualityLevels);
-	}
-
 	// Notify all initialized FRenderResources that there's a valid RHI device to create their RHI resources for now.
 	for (TLinkedList<FRenderResource*>::TIterator ResourceIt(FRenderResource::GetResourceList()); ResourceIt; ResourceIt.Next())
 	{
@@ -517,6 +494,9 @@ void FD3D12DynamicRHI::PerRHISetup(FD3D12Device* MainDevice)
 	GRHISupportsTextureStreaming = true;
 	GRHISupportsFirstInstance = true;
 
+	// Indicate that the RHI needs to use the engine's deferred deletion queue.
+	GRHINeedsExtraDeletionLatency = true;
+
 	// Set the RHI initialized flag.
 	GIsRHIInitialized = true;
 }
@@ -614,30 +594,12 @@ void FD3D12Device::InitD3DDevice()
 			check(!"Internal error, EnumAdapters() failed but before it worked")
 		}
 
-		D3D_FEATURE_LEVEL ActualFeatureLevel = GetFeatureLevel();
-
-		if (FParse::Param(FCommandLine::Get(), TEXT("warp")))
-		{
-			TRefCountPtr<IDXGIAdapter> WarpAdapter;
-			VERIFYD3D11RESULT(DXGIFactory->EnumWarpAdapter(IID_PPV_ARGS(WarpAdapter.GetInitReference())));
-
-			// Creating the Direct3D WARP device.
-			VERIFYD3D11RESULT(D3D12CreateDevice(
-				WarpAdapter,
-				GetFeatureLevel(),
-				IID_PPV_ARGS(Direct3DDevice.GetInitReference())
-				));
-
-		}
-		else
-		{
-			// Creating the Direct3D device.
-			VERIFYD3D11RESULT(D3D12CreateDevice(
-				Adapter,
-				GetFeatureLevel(),
-				IID_PPV_ARGS(Direct3DDevice.GetInitReference())
-				));
-		}
+		// Creating the Direct3D device.
+		VERIFYD3D11RESULT(D3D12CreateDevice(
+			Adapter,
+			GetFeatureLevel(),
+			IID_PPV_ARGS(Direct3DDevice.GetInitReference())
+			));
 
 #if UE_BUILD_DEBUG	
 		//break on debug
@@ -654,9 +616,6 @@ void FD3D12Device::InitD3DDevice()
 		}
 #endif
 
-		// We should get the feature level we asked for as earlier we checked to ensure it is supported.
-		check(ActualFeatureLevel == GetFeatureLevel());
-
 		D3D12_FEATURE_DATA_D3D12_OPTIONS D3D12Caps;
 		VERIFYD3D11RESULT(GetDevice()->CheckFeatureSupport(D3D12_FEATURE_D3D12_OPTIONS, &D3D12Caps, sizeof(D3D12Caps)));
 		ResourceHeapTier = D3D12Caps.ResourceHeapTier;
diff --git a/Engine/Source/Runtime/Windows/D3D12RHI/Public/D3D12Resources.h b/Engine/Source/Runtime/Windows/D3D12RHI/Public/D3D12Resources.h
--- a/Engine/Source/Runtime/Windows/D3D12RHI/Public/D3D12Resources.h
+++ b/Engine/Source/Runtime/Windows/D3D12RHI/Public/D3D12Resources.h
@@ -658,6 +658,8 @@ private:
 	uint16 SubresourceCount;
 	CResourceState* pResourceState;
 	D3D12_RESOURCE_STATES DefaultResourceState;
+	D3D12_RESOURCE_STATES ReadableState;
+	D3D12_RESOURCE_STATES WritableState;
 	bool bRequiresResourceStateTracking;
 	D3D12_HEAP_TYPE HeapType;
 	FName DebugName;
@@ -726,6 +728,8 @@ public:
 		return pResourceState;
 	}
 	D3D12_RESOURCE_STATES GetDefaultResourceState() const { check(!bRequiresResourceStateTracking); return DefaultResourceState; }
+	D3D12_RESOURCE_STATES GetWritableState() const { return WritableState; }
+	D3D12_RESOURCE_STATES GetReadableState() const { return ReadableState; }
 	bool RequiresResourceStateTracking() const { return bRequiresResourceStateTracking; }
 
 	void SetName(const TCHAR* Name)
@@ -734,11 +738,17 @@ public:
 		::SetName(Resource, Name);
 	}
 
+	FName GetName() const
+	{
+		return DebugName;
+	}
+
+	inline bool IsPlacedResource() const { return Heap.GetReference() != nullptr; }
 private:
 	void InitalizeResourceState()
 	{
 		SubresourceCount = GetMipLevels() * GetArraySize() * GetPlaneCount();
-		ComputeDefaultResourceState();
+		DetermineResourceStates();
 
 		if (bRequiresResourceStateTracking)
 		{
@@ -749,20 +759,44 @@ private:
 		}
 	}
 
-	void ComputeDefaultResourceState()
+	void DetermineResourceStates()
 	{
 		const bool bSRV = (Desc.Flags & D3D12_RESOURCE_FLAG_DENY_SHADER_RESOURCE) == 0;
 		const bool bDSV = (Desc.Flags & D3D12_RESOURCE_FLAG_ALLOW_DEPTH_STENCIL) != 0;
 		const bool bRTV = (Desc.Flags & D3D12_RESOURCE_FLAG_ALLOW_RENDER_TARGET) != 0;
 		const bool bUAV = (Desc.Flags & D3D12_RESOURCE_FLAG_ALLOW_UNORDERED_ACCESS) != 0;
 
-		const bool bWriteable = bDSV || bRTV || bUAV;
-		const bool bSRVOnly = bSRV && !bWriteable;
+		const bool bWritable = bDSV || bRTV || bUAV;
+		const bool bSRVOnly = bSRV && !bWritable;
 		const bool bBuffer = Desc.Dimension == D3D12_RESOURCE_DIMENSION_BUFFER;
 
+		if (bWritable)
+		{
+			// Determine the resource's write/read states.
+			if (bRTV)
+			{
+				// Note: The resource could also be used as a UAV however we don't store that writable state. UAV's are handled in a separate RHITransitionResources() specially for UAVs so we know the writeable state in that case should be UAV.
+				check(!bDSV && bSRV && !bBuffer);
+				WritableState = D3D12_RESOURCE_STATE_RENDER_TARGET;
+				ReadableState = D3D12_RESOURCE_STATE_NON_PIXEL_SHADER_RESOURCE | D3D12_RESOURCE_STATE_PIXEL_SHADER_RESOURCE;
+			}
+			else if (bDSV)
+			{
+				check(!bRTV && !bUAV && bSRV && !bBuffer);
+				WritableState = D3D12_RESOURCE_STATE_DEPTH_WRITE;
+				ReadableState = D3D12_RESOURCE_STATE_DEPTH_READ | D3D12_RESOURCE_STATE_NON_PIXEL_SHADER_RESOURCE | D3D12_RESOURCE_STATE_PIXEL_SHADER_RESOURCE;
+			}
+			else
+			{
+				check(bUAV && !bRTV && !bDSV && bSRV)
+				WritableState = D3D12_RESOURCE_STATE_UNORDERED_ACCESS;
+				ReadableState = bBuffer ? D3D12_RESOURCE_STATE_GENERIC_READ : D3D12_RESOURCE_STATE_NON_PIXEL_SHADER_RESOURCE | D3D12_RESOURCE_STATE_PIXEL_SHADER_RESOURCE;
+			}
+		}
+
 		if (bBuffer)
 		{
-			if (!bWriteable)
+			if (!bWritable)
 			{
 				// Buffer used for input, like Vertex/Index buffer.
 				// Don't bother tracking state for this resource.
@@ -796,6 +830,9 @@ struct FD3D12ResourceBlockInfo
 {
 private:
 	friend class FD3D12ResourceAllocator;
+	friend class FD3D12BuddyAllocator;
+	friend class FD3D12LinearAllocator;
+	friend class FD3D12BucketAllocator;
 
 	inline uint64 GetOffset() const { return AllocatorValues.AllocationBlockOffset; }
 	inline uint32 GetSize() const { return AllocatorValues.Size; }
@@ -803,10 +840,12 @@ private:
 	// Internal values used by the allocator
 	struct AllocationTrackingValues
 	{
-		AllocationTrackingValues() : AllocationBlockOffset(0), Size(0), UnpaddedSize(0) {};
+		AllocationTrackingValues() : AllocationBlockOffset(0), Size(0), UnpaddedSize(0), AllocatorPrivateData(nullptr){};
 		uint64 AllocationBlockOffset;
 		uint32 Size;
 		uint32 UnpaddedSize;
+
+		void* AllocatorPrivateData;
 	} AllocatorValues;
 
 public:
@@ -850,7 +889,7 @@ protected:
 	friend class FD3D12DefaultBufferPool;
 	friend class FD3D12FastAllocator;
 
-	virtual void InternalReleaseResource();
+	void InternalReleaseResource();
 
 public:
 	FD3D12ResourceLocation()
@@ -908,7 +947,7 @@ public:
 		}
 	}
 
-	virtual ~FD3D12ResourceLocation();
+	~FD3D12ResourceLocation();
 
 	void SetAsFastAllocatedSubresource()
 	{
@@ -1088,10 +1127,7 @@ public:
 
 	void Clear()
 	{
-		// Make sure the queue is empty after a clear.
-		// TODO: Perhaps just ignore the frame fence and clear all resources?
-		const bool bQueueIsEmpty = ReleaseResources(true);
-		check(bQueueIsEmpty);
+		ReleaseResources(true);
 	}
 
 	FD3D12DeferredDeletionQueue(FD3D12Device* InParent);
@@ -1645,15 +1681,18 @@ class FD3D12BaseShaderResource : public FD3D12DeviceChild, public IRefCountedObj
 {
 public:
 	TRefCountPtr<FD3D12ResourceLocation> ResourceLocation;
+	uint32 BufferAlignment;
 
 public:
 	FD3D12BaseShaderResource(FD3D12ResourceLocation *InResourceLocation, FD3D12Device* InParent) :
+		BufferAlignment(0),
 		FD3D12DeviceChild(InParent)
 	{
 		ResourceLocation = InResourceLocation;
 	}
 
 	FD3D12BaseShaderResource(FD3D12Resource *InResource, FD3D12Device* InParent) :
+		BufferAlignment(0),
 		FD3D12DeviceChild(InParent)
 	{
 		ResourceLocation = new FD3D12ResourceLocation(GetParentDevice(), InResource);
@@ -3168,6 +3207,8 @@ public:
 
 	virtual ~FD3D12IndexBuffer();
 
+	void Rename(FD3D12ResourceLocation* NewResource);
+
 	// IRefCountedObject interface.
 	virtual uint32 AddRef() const
 	{
@@ -3202,6 +3243,8 @@ public:
 		, Resource(InResource)
 	{}
 
+	void Rename(FD3D12ResourceLocation* NewResource);
+
 	virtual ~FD3D12StructuredBuffer();
 
 	// IRefCountedObject interface.
@@ -3235,7 +3278,8 @@ public:
 
 	virtual ~FD3D12VertexBuffer();
 
-	void* DynamicLock();
+	void Rename(FD3D12ResourceLocation* NewResource);
+
 	void SetDynamicSRV(FD3D12ShaderResourceView* InSRV)
 	{
 		DynamicSRV = InSRV;
@@ -3268,152 +3312,6 @@ namespace D3D12RHI
 	void ReleasePooledTextures();
 }
 
-// Allocates blocks from a fixed range using buddy allocation method.
-// Buddy allocation allows reasonably fast allocation of arbitrary size blocks
-// with minimal fragmentation and provides efficient reuse of freed ranges.
-// When a block is de-allocated an attempt is made to merge it with it's 
-// neighbour (buddy) if it is contiguous and free.
-// Based on reference implementation by MSFT: billkris
-
-// Unfortunately the api restricts the minimum size of a placed buffer resource to 64k
-#define MIN_PLACED_BUFFER_SIZE (64 * 1024)
-#define D3D_BUFFER_ALIGNMENT (64 * 1024)
-
-#if defined(UE_BUILD_DEBUG)
-#define INCREASE_ALLOC_COUNTER(A, B) (A = A + B);
-#define DECREASE_ALLOC_COUNTER(A, B) (A = A - B);
-#else
-#define INCREASE_ALLOC_COUNTER(A, B)
-#define DECREASE_ALLOC_COUNTER(A, B)
-#endif
-
-enum eBuddyAllocationStrategy
-{
-	// This strategy uses Placed Resources to sub-allocate a buffer out of an underlying ID3D12Heap.
-	// The benefit of this is that each buffer can have it's own resource state and can be treated
-	// as any other buffer. The downside of this strategy is the API limitiation which enforces
-	// the minimum buffer size to 64k leading to large internal fragmentation in the allocator
-	kPlacedResourceStrategy,
-	// The alternative is to manualy sub-allocate out of a single large buffer which allows block
-	// allocation granularity down to 1 byte. However, this strategy is only really valid for buffers which
-	// will be treated as read-only after their creation (i.e. most Index and Vertex buffers). This 
-	// is because the underlying resource can only have one state at a time.
-	kManualSubAllocationStrategy
-};
-
-class FD3D12ResourceAllocator : public FD3D12DeviceChild
-{
-public:
-
-	FD3D12ResourceAllocator(FD3D12Device* ParentDevice,
-		FString Name,
-		eBuddyAllocationStrategy allocationStrategy,
-		D3D12_HEAP_TYPE heapType,
-		D3D12_HEAP_FLAGS heapFlags,
-		D3D12_RESOURCE_FLAGS flags,
-		uint32 MaxSizeForPooling,
-		uint32 maxBlockSize,
-		uint32 minBlockSize = MIN_PLACED_BUFFER_SIZE);
-
-	~FD3D12ResourceAllocator();
-
-	void Deallocate(FD3D12ResourceBlockInfo* Block);
-
-	void Initialize();
-
-	void Destroy();
-
-	inline bool IsOwner(const FD3D12ResourceBlockInfo &block)
-	{
-		return block.Allocator == this;
-	}
-
-	inline void Reset()
-	{
-		// Clear the free blocks collection
-		FreeBlocks.Empty();
-
-		// Initialize the pool with a free inner block of max inner block size
-		FreeBlocks.SetNum(MaxOrder + 1);
-		FreeBlocks[MaxOrder].Add((uint32)0);
-	}
-
-	void CleanUpAllocations();
-
-	void DumpAllocatorStats(class FOutputDevice& Ar);
-
-	void ReleaseAllResources();
-
-private:
-	const FString DebugName;
-
-	TRefCountPtr<FD3D12Resource> BackingResource;
-	void* BaseAddress;
-
-	bool Initialized;
-
-	const D3D12_HEAP_TYPE HeapType;
-	const D3D12_HEAP_FLAGS HeapFlags;
-
-	TQueue<FD3D12ResourceBlockInfo*> DeferredDeletionQueue;
-	TArray<TSet<uint32>> FreeBlocks;
-	uint32 MaxOrder;
-
-	const uint32 MaxBlockSize;
-	const uint32 MinBlockSize;
-
-	const eBuddyAllocationStrategy AllocationStrategy;
-
-	inline uint32 SizeToUnitSize(uint32 size) const
-	{
-		return (size + (MinBlockSize - 1)) / MinBlockSize;
-	}
-
-	inline uint32 UnitSizeToOrder(uint32 size) const
-	{
-		return uint32(ceil(log2f(float(size))));
-	}
-
-	inline uint32 GetBuddyOffset(const uint32 &offset, const uint32 &size)
-	{
-		return offset ^ size;
-	}
-
-	void DeallocateInternal(FD3D12ResourceBlockInfo* Block);
-
-	uint32 OrderToUnitSize(uint32 order) const { return ((uint32)1) << order; }
-	uint32 AllocateBlock(uint32 order);
-	void DeallocateBlock(uint32 offset, uint32 order);
-
-	bool CanAllocate(uint32 size);
-	FD3D12ResourceBlockInfo* Allocate(uint32 SizeInBytes, uint32 Alignment, const void* InitialData = nullptr);
-
-protected:
-
-	FD3D12ResourceBlockInfo* TryAllocate(uint32 SizeInBytes, uint32 Alignment, const void* InitialData = nullptr);
-
-	TRefCountPtr<ID3D12Heap> BackingHeap;
-
-	// Any allocation larger than this just gets straight up allocated (i.e. not pooled).
-	// These large allocations should be infrequent so the CPU overhead should be minimal
-	const uint32 MaximumAllocationSizeForPooling;
-
-	void InitializeDefaultBuffer(FD3D12Resource* Destination, uint64 DestinationOffset, const void* Data, uint64 DataSize);
-
-	FCriticalSection CS;
-	D3D12_RESOURCE_FLAGS ResourceFlags;
-
-#if defined(UE_BUILD_DEBUG)
-	uint32 SpaceUsed;
-	uint32 InternalFragmentation;
-	uint32 NumBlocksInDeferredDeletionQueue;
-	uint32 PeakUsage;
-	uint32 FailedAllocationSpace;
-#endif
-
-	bool HeapFullMessageDisplayed;
-};
-
 struct FD3D12FastAllocatorPage
 {
 	FD3D12FastAllocatorPage() :
@@ -3525,128 +3423,6 @@ private:
 	FCriticalSection CS;
 };
 
-// FD3D12DynamicHeapAllocator
-//
-// This is designed for allocation of scratch memory such as temporary staging buffers
-// or shadow buffers for dynamic resources.
-class FD3D12DynamicHeapAllocator : public FD3D12ResourceAllocator
-{
-	// The currently paired command context.  If none is available (as with async texture loads),
-	// the current command list handle can be provided instead.  The command list handle is the
-	// critical part, but the command context can always return the current list handle.
-	// I.e. you need one but not both of these to be valid
-	class FD3D12CommandContext* CurrentCommandContext;
-	FD3D12CommandListHandle CurrentCommandListHandle;
-
-public:
-
-	FD3D12DynamicHeapAllocator(FD3D12Device* InParent, FString Name, eBuddyAllocationStrategy allocationStrategy,
-		uint32 MaxSizeForPooling,
-		uint32 maxBlockSize,
-		uint32 minBlockSize);
-
-	// These are mutually exclusive calls!
-	inline void SetCurrentCommandContext(class FD3D12CommandContext* CmdContext)
-	{
-		CurrentCommandContext = CmdContext;
-		check(CurrentCommandListHandle == nullptr);
-	}
-
-	inline void SetCurrentCommandListHandle(FD3D12CommandListHandle CommandListHandle)
-	{
-		check(CurrentCommandContext == nullptr);
-		CurrentCommandListHandle = CommandListHandle;
-	}
-
-	// Allocates <size> bytes from the end of an available resource heap.
-	void* AllocUploadResource(uint32 size, uint32 alignment, class FD3D12ResourceLocation* ResourceLocation);
-};
-
-// FD3D12DefaultBufferPool
-//
-class FD3D12DefaultBufferPool : public FD3D12ResourceAllocator
-{
-public:
-	FD3D12DefaultBufferPool(FD3D12Device* InParent, FString Name, eBuddyAllocationStrategy allocationStrategy,
-		uint32 MaxSizeForPooling,
-		D3D12_RESOURCE_FLAGS flags,
-		uint32 maxBlockSize,
-		uint32 minBlockSize);
-	FD3D12DefaultBufferPool();
-
-	// Grab a buffer from the available buffers or create a new buffer if none are available
-	void AllocDefaultResource(const D3D12_RESOURCE_DESC& Desc, D3D12_SUBRESOURCE_DATA* pInitialData, class FD3D12ResourceLocation* ResourceLocation, uint32 Alignment);
-};
-
-// FD3D12DefaultBufferAllocator
-//
-class FD3D12DefaultBufferAllocator : public FD3D12DeviceChild
-{
-	static const uint32 MAX_DEFAULT_POOLS = 16; // Should match the max D3D12_RESOURCE_FLAG_FLAG combinations.
-	FD3D12DefaultBufferPool* DefaultBufferPools[MAX_DEFAULT_POOLS];
-	HRESULT AllocateWriteableBuffer(const D3D12_RESOURCE_DESC& Desc, D3D12_SUBRESOURCE_DATA* pInitialData, class FD3D12ResourceLocation* ResourceLocation, uint32 Alignment = 0);
-
-	bool BufferIsWriteable(const D3D12_RESOURCE_DESC& Desc) 
-	{ 
-		const bool bDSV = (Desc.Flags & D3D12_RESOURCE_FLAG_ALLOW_DEPTH_STENCIL) != 0;
-		const bool bRTV = (Desc.Flags & D3D12_RESOURCE_FLAG_ALLOW_RENDER_TARGET) != 0;
-		const bool bUAV = (Desc.Flags & D3D12_RESOURCE_FLAG_ALLOW_UNORDERED_ACCESS) != 0;
-
-		// Buffer Depth Stencils are invalid
-		check(bDSV == false);
-		const bool bWriteable = bDSV || bRTV || bUAV;
-		return bWriteable;
-	}
-
-public:
-	HRESULT AllocDefaultResource(const D3D12_RESOURCE_DESC& Desc, D3D12_SUBRESOURCE_DATA* InitialData, class FD3D12ResourceLocation** ResourceLocation, uint32 Alignment = 0);
-	// Grab a buffer from the available buffers or create a new buffer if none are available
-	HRESULT AllocDefaultResource(const D3D12_RESOURCE_DESC& pDesc, D3D12_SUBRESOURCE_DATA* InitialData, class FD3D12ResourceLocation* ResourceLocation, uint32 Alignment = 0);
-	void FreeDefaultBufferPools();
-	void CleanupFreeBlocks();
-
-	FD3D12DefaultBufferAllocator(FD3D12Device* InParent) :
-		FD3D12DeviceChild(InParent)
-	{
-		FMemory::Memset(DefaultBufferPools, 0);
-	}
-};
-
-class FD3D12TextureAllocator : public FD3D12ResourceAllocator
-{
-public:
-	FD3D12TextureAllocator(FD3D12Device* Device, FString Name, uint32 HeapSize, D3D12_HEAP_FLAGS Flags) :
-		FD3D12ResourceAllocator(Device, Name,
-			kPlacedResourceStrategy,
-			D3D12_HEAP_TYPE_DEFAULT,
-			Flags | D3D12_HEAP_FLAG_DENY_BUFFERS,
-			D3D12_RESOURCE_FLAG_NONE,
-			D3D12_DEFAULT_RESOURCE_PLACEMENT_ALIGNMENT,
-			HeapSize,
-			D3D12_SMALL_RESOURCE_PLACEMENT_ALIGNMENT) {}
-
-	HRESULT AllocateTexture(D3D12_RESOURCE_DESC Desc, const D3D12_CLEAR_VALUE* ClearValue, FD3D12ResourceLocation* TextureLocation);
-};
-
-#define TEXTURE_POOL_SIZE_READABLE (64 * 1024 * 1024)
-
-class FD3D12TextureAllocatorPool : public FD3D12DeviceChild
-{
-public:
-	FD3D12TextureAllocatorPool(FD3D12Device* Device) :
-		ReadOnlyTexturePool(Device, FString(L"Small Read-Only Texture allocator"), TEXTURE_POOL_SIZE_READABLE, D3D12_HEAP_FLAG_ALLOW_ONLY_NON_RT_DS_TEXTURES),
-		FD3D12DeviceChild(Device)
-	{};
-
-	HRESULT AllocateTexture(D3D12_RESOURCE_DESC Desc, const D3D12_CLEAR_VALUE* ClearValue, uint8 UEFormat, FD3D12ResourceLocation* TextureLocation);
-
-	void CleanUpAllocations() { ReadOnlyTexturePool.CleanUpAllocations(); }
-
-	void Destroy() { ReadOnlyTexturePool.Destroy(); }
-
-private:
-		FD3D12TextureAllocator ReadOnlyTexturePool;
-};
 
 #if SUPPORTS_MEMORY_RESIDENCY
 class FD3D12ResourceResidencyManager : public FD3D12DeviceChild
diff --git a/Engine/Source/Runtime/Windows/D3D12RHI/Public/D3D12Util.h b/Engine/Source/Runtime/Windows/D3D12RHI/Public/D3D12Util.h
--- a/Engine/Source/Runtime/Windows/D3D12RHI/Public/D3D12Util.h
+++ b/Engine/Source/Runtime/Windows/D3D12RHI/Public/D3D12Util.h
@@ -293,6 +293,11 @@ public:
 	FD3D12LockedKey(class FD3D12ResourceLocation* source, uint32 subres = 0) : SourceObject((void*)source)
 		, Subresource(subres)
 	{}
+
+	template<class ClassType>
+	FD3D12LockedKey(ClassType* source, uint32 subres = 0) : SourceObject((void*)source)
+		, Subresource(subres)
+	{}
 	bool operator==(const FD3D12LockedKey& Other) const
 	{
 		return SourceObject == Other.SourceObject && Subresource == Other.Subresource;
@@ -1118,3 +1123,13 @@ static bool TextureCanBe4KAligned(D3D12_RESOURCE_DESC& Desc, uint8 UEFormat)
 
 	return TilesNeeded <= NUM_4K_BLOCKS_PER_64K_PAGE;
 }
+
+template <class TView>
+class FD3D12View;
+class CViewSubresourceSubset;
+
+template <class TView>
+bool AssertResourceState(ID3D12CommandList* pCommandList, FD3D12View<TView>* pView, const D3D12_RESOURCE_STATES& State);
+
+bool AssertResourceState(ID3D12CommandList* pCommandList, FD3D12Resource* pResource, const D3D12_RESOURCE_STATES& State, uint32 Subresource);
+bool AssertResourceState(ID3D12CommandList* pCommandList, FD3D12Resource* pResource, const D3D12_RESOURCE_STATES& State, const CViewSubresourceSubset& SubresourceSubset);
diff --git a/Engine/Source/UE4Editor.Target.cs b/Engine/Source/UE4Editor.Target.cs
--- a/Engine/Source/UE4Editor.Target.cs
+++ b/Engine/Source/UE4Editor.Target.cs
@@ -174,7 +174,12 @@ public class UE4EditorTarget : TargetRules
         NonCodeProjectNames.Add("MorphTargets", DesktopPlats);
         NonCodeProjectNames.Add("PostProcessMatinee", DesktopPlats);
         NonCodeProjectNames.Add("SciFiHallway", DesktopPlats);
-		NonCodeProjectNames.Add("ShowdownDemo", DesktopPlats);
+
+        // Windows only due to VR elements
+        if (HostPlatform == UnrealTargetPlatform.Win64)
+        {
+            NonCodeProjectNames.Add("ShowdownDemo", new List<UnrealTargetPlatform> { HostPlatform, UnrealTargetPlatform.Win32 });
+        }
 
         //NonCodeProjectNames.Add("BlackJack", DesktopAndMobilePlats);
         //NonCodeProjectNames.Add("MemoryGame", DesktopAndMobilePlats);
diff --git a/Templates/TP_Vehicle/Source/TP_Vehicle/TP_VehiclePawn.cpp b/Templates/TP_Vehicle/Source/TP_Vehicle/TP_VehiclePawn.cpp
--- a/Templates/TP_Vehicle/Source/TP_Vehicle/TP_VehiclePawn.cpp
+++ b/Templates/TP_Vehicle/Source/TP_Vehicle/TP_VehiclePawn.cpp
@@ -71,7 +71,7 @@ ATP_VehiclePawn::ATP_VehiclePawn()
 	Camera->FieldOfView = 90.f;
 
 	// Create In-Car camera component 
-	InternalCameraOrigin = FVector(8.0f, -40.0f, 130.0f);
+	InternalCameraOrigin = FVector(0.0f, -40.0f, 120.0f);
 
 	InternalCameraBase = CreateDefaultSubobject<USceneComponent>(TEXT("InternalCameraBase"));
 	InternalCameraBase->SetRelativeLocation(InternalCameraOrigin);
@@ -83,8 +83,19 @@ ATP_VehiclePawn::ATP_VehiclePawn()
 	InternalCamera->FieldOfView = 90.f;
 	InternalCamera->AttachTo(InternalCameraBase);
 
+	//Setup TextRenderMaterial
+	UPROPERTY(EditAnywhere)
+		UMaterialInterface* Material;
+	static ConstructorHelpers::FObjectFinder<UMaterial> TextMaterial(TEXT("Material'/Engine/EngineMaterials/AntiAliasedTextMaterialTranslucent.AntiAliasedTextMaterialTranslucent'"));
+	
+	Material = TextMaterial.Object;
+	
+	
+
+
 	// Create text render component for in car speed display
 	InCarSpeed = CreateDefaultSubobject<UTextRenderComponent>(TEXT("IncarSpeed"));
+	InCarSpeed->SetTextMaterial(Material);
 	InCarSpeed->SetRelativeLocation(FVector(70.0f, -75.0f, 99.0f));
 	InCarSpeed->SetRelativeRotation(FRotator(18.0f, 180.0f, 0.0f));
 	InCarSpeed->AttachTo(GetMesh());
@@ -92,6 +103,7 @@ ATP_VehiclePawn::ATP_VehiclePawn()
 
 	// Create text render component for in car gear display
 	InCarGear = CreateDefaultSubobject<UTextRenderComponent>(TEXT("IncarGear"));
+	InCarGear->SetTextMaterial(Material);
 	InCarGear->SetRelativeLocation(FVector(66.0f, -9.0f, 95.0f));	
 	InCarGear->SetRelativeRotation(FRotator(25.0f, 180.0f,0.0f));
 	InCarGear->SetRelativeScale3D(FVector(1.0f, 0.4f, 0.4f));
diff --git a/Templates/TP_VehicleAdv/Source/TP_VehicleAdv/TP_VehicleAdvPawn.cpp b/Templates/TP_VehicleAdv/Source/TP_VehicleAdv/TP_VehicleAdvPawn.cpp
--- a/Templates/TP_VehicleAdv/Source/TP_VehicleAdv/TP_VehicleAdvPawn.cpp
+++ b/Templates/TP_VehicleAdv/Source/TP_VehicleAdv/TP_VehicleAdvPawn.cpp
@@ -130,7 +130,7 @@ ATP_VehicleAdvPawn::ATP_VehicleAdvPawn()
 	Camera->FieldOfView = 90.f;
 
 	// Create In-Car camera component 
-	InternalCameraOrigin = FVector(-34.0f, 0.0f, 50.0f);
+	InternalCameraOrigin = FVector(-34.0f, -10.0f, 50.0f);
 	InternalCameraBase = CreateDefaultSubobject<USceneComponent>(TEXT("InternalCameraBase"));
 	InternalCameraBase->SetRelativeLocation(InternalCameraOrigin);
 	InternalCameraBase->AttachTo(GetMesh());
