Copying //UE4/Dev-Framework to //UE4/Dev-Main (Source: //UE4/Dev-Framework @ 2964666)

diff --git a/Engine/Source/Runtime/Engine/Private/SceneView.cpp b/Engine/Source/Runtime/Engine/Private/SceneView.cpp
--- a/Engine/Source/Runtime/Engine/Private/SceneView.cpp
+++ b/Engine/Source/Runtime/Engine/Private/SceneView.cpp
@@ -1395,9 +1395,9 @@ void FSceneView::EndFinalPostprocessSettings(const FSceneViewInitOptions& ViewIn
 	}
 
 	{
-		static const auto CVar = IConsoleManager::Get().FindTConsoleVariableDataInt(TEXT("r.SceneColorFringeQuality"));
+		static const auto SceneColorFringeQualityCVar = IConsoleManager::Get().FindTConsoleVariableDataInt(TEXT("r.SceneColorFringeQuality"));
 
-		int32 FringeQuality = CVar->GetValueOnGameThread();
+		int32 FringeQuality = SceneColorFringeQualityCVar->GetValueOnGameThread();
 		if (FringeQuality <= 0)
 		{
 			FinalPostProcessSettings.SceneFringeIntensity = 0;
@@ -1405,9 +1405,9 @@ void FSceneView::EndFinalPostprocessSettings(const FSceneViewInitOptions& ViewIn
 	}
 
 	{
-		static const auto CVar = IConsoleManager::Get().FindTConsoleVariableDataInt(TEXT("r.BloomQuality"));
+		static const auto BloomQualityCVar = IConsoleManager::Get().FindTConsoleVariableDataInt(TEXT("r.BloomQuality"));
 
-		int Value = CVar->GetValueOnGameThread();
+		int Value = BloomQualityCVar->GetValueOnGameThread();
 
 		if(Value <= 0)
 		{
@@ -1451,9 +1451,9 @@ void FSceneView::EndFinalPostprocessSettings(const FSceneViewInitOptions& ViewIn
 	}
 
 	{
-		static const auto CVar = IConsoleManager::Get().FindTConsoleVariableDataInt(TEXT("r.DepthOfFieldQuality"));
+		static const auto DepthOfFieldQualityCVar = IConsoleManager::Get().FindTConsoleVariableDataInt(TEXT("r.DepthOfFieldQuality"));
 
-		int Value = CVar->GetValueOnGameThread();
+		int Value = DepthOfFieldQualityCVar->GetValueOnGameThread();
 
 		if(Value <= 0)
 		{
@@ -1531,9 +1531,9 @@ void FSceneView::EndFinalPostprocessSettings(const FSceneViewInitOptions& ViewIn
 	}
 
 	{
-		static const auto CVar = IConsoleManager::Get().FindTConsoleVariableDataFloat(TEXT("r.ScreenPercentage"));
+		static const auto ScreenPercentageCVar = IConsoleManager::Get().FindTConsoleVariableDataFloat(TEXT("r.ScreenPercentage"));
 
-		float Value = CVar->GetValueOnGameThread();
+		float Value = ScreenPercentageCVar->GetValueOnGameThread();
 
 		if(Value >= 0.0)
 		{
@@ -1562,9 +1562,9 @@ void FSceneView::EndFinalPostprocessSettings(const FSceneViewInitOptions& ViewIn
 #endif
 
 	{
-		static const auto CVar = IConsoleManager::Get().FindTConsoleVariableDataFloat(TEXT("r.AmbientOcclusionStaticFraction"));
+		static const auto AmbientOcclusionStaticFractionCVar = IConsoleManager::Get().FindTConsoleVariableDataFloat(TEXT("r.AmbientOcclusionStaticFraction"));
 
-		float Value = CVar->GetValueOnGameThread();
+		float Value = AmbientOcclusionStaticFractionCVar->GetValueOnGameThread();
 
 		if(Value >= 0.0)
 		{
@@ -1586,9 +1586,9 @@ void FSceneView::EndFinalPostprocessSettings(const FSceneViewInitOptions& ViewIn
 	}
 
 	{
-		static const auto CVar = IConsoleManager::Get().FindTConsoleVariableDataFloat(TEXT("r.AmbientOcclusionRadiusScale"));
+		static const auto AmbientOcclusionRadiusScaleCVar = IConsoleManager::Get().FindTConsoleVariableDataFloat(TEXT("r.AmbientOcclusionRadiusScale"));
 
-		float Scale = FMath::Clamp(CVar->GetValueOnGameThread(), 0.1f, 15.0f);
+		float Scale = FMath::Clamp(AmbientOcclusionRadiusScaleCVar->GetValueOnGameThread(), 0.1f, 15.0f);
 		
 		FinalPostProcessSettings.AmbientOcclusionRadius *= Scale;
 	}
@@ -1649,12 +1649,12 @@ void FSceneView::EndFinalPostprocessSettings(const FSceneViewInitOptions& ViewIn
 
 	// Anti-Aliasing
 	{
-		static const auto CVar = IConsoleManager::Get().FindTConsoleVariableDataInt(TEXT("r.PostProcessAAQuality")); 
+		static const auto PostProcessAAQualityCVar = IConsoleManager::Get().FindTConsoleVariableDataInt(TEXT("r.PostProcessAAQuality")); 
 		static auto* MobileHDRCvar = IConsoleManager::Get().FindTConsoleVariableDataInt(TEXT("r.MobileHDR"));
 		static auto* MobileMSAACvar = IConsoleManager::Get().FindTConsoleVariableDataInt(TEXT("r.MobileMSAA"));
 		static uint32 MSAAValue = GShaderPlatformForFeatureLevel[SceneViewFeatureLevel] == SP_OPENGL_ES2_IOS ? 1 : MobileMSAACvar->GetValueOnGameThread();
 
-		int32 Quality = FMath::Clamp(CVar->GetValueOnGameThread(), 0, 6);
+		int32 Quality = FMath::Clamp(PostProcessAAQualityCVar->GetValueOnGameThread(), 0, 6);
 		const bool bWillApplyTemporalAA = Family->EngineShowFlags.PostProcessing || bIsPlanarReflection;
 
 		if( !bWillApplyTemporalAA || !Family->EngineShowFlags.AntiAliasing || Quality <= 0
diff --git a/Engine/Source/Runtime/Engine/Private/SkeletalRenderGPUSkin.cpp b/Engine/Source/Runtime/Engine/Private/SkeletalRenderGPUSkin.cpp
--- a/Engine/Source/Runtime/Engine/Private/SkeletalRenderGPUSkin.cpp
+++ b/Engine/Source/Runtime/Engine/Private/SkeletalRenderGPUSkin.cpp
@@ -930,7 +930,7 @@ void FSkeletalMeshObjectGPUSkin::FSkeletalMeshObjectLOD::GetVertexBuffers(
 void FSkeletalMeshObjectGPUSkin::FVertexFactoryData::InitVertexFactories(
 	const FVertexFactoryBuffers& VertexBuffers, 
 	const TArray<FSkelMeshChunk>& Chunks, 
-	ERHIFeatureLevel::Type FeatureLevel)
+	ERHIFeatureLevel::Type InFeatureLevel)
 {
 	// first clear existing factories (resources assumed to have been released already)
 	// then [re]create the factories
@@ -941,12 +941,12 @@ void FSkeletalMeshObjectGPUSkin::FVertexFactoryData::InitVertexFactories(
 		{
 			if (VertexBuffers.VertexBufferGPUSkin->HasExtraBoneInfluences())
 			{
-				auto* VertexFactory = CreateVertexFactory< FGPUBaseSkinVertexFactory, TGPUSkinVertexFactory<true> >(VertexFactories, VertexBuffers, FeatureLevel);
+				auto* VertexFactory = CreateVertexFactory< FGPUBaseSkinVertexFactory, TGPUSkinVertexFactory<true> >(VertexFactories, VertexBuffers, InFeatureLevel);
 				CreatePassthroughVertexFactory<TGPUSkinVertexFactory<true>>(PassthroughVertexFactories, VertexFactory);
 			}
 			else
 			{
-				auto* VertexFactory = CreateVertexFactory< FGPUBaseSkinVertexFactory, TGPUSkinVertexFactory<false> >(VertexFactories, VertexBuffers, FeatureLevel);
+				auto* VertexFactory = CreateVertexFactory< FGPUBaseSkinVertexFactory, TGPUSkinVertexFactory<false> >(VertexFactories, VertexBuffers, InFeatureLevel);
 				CreatePassthroughVertexFactory<TGPUSkinVertexFactory<false>>(PassthroughVertexFactories, VertexFactory);
 			}
 		}
@@ -1050,7 +1050,7 @@ void FSkeletalMeshObjectGPUSkin::FVertexFactoryData::ReleaseAPEXClothVertexFacto
  * @param bUseLocalVertexFactory - use non-gpu skinned factory when rendering in ref pose
  * @param MeshLODInfo - information about the state of the bone influence swapping
  */
-void FSkeletalMeshObjectGPUSkin::FSkeletalMeshObjectLOD::InitResources(const FSkelMeshObjectLODInfo& MeshLODInfo, ERHIFeatureLevel::Type FeatureLevel)
+void FSkeletalMeshObjectGPUSkin::FSkeletalMeshObjectLOD::InitResources(const FSkelMeshObjectLODInfo& MeshLODInfo, ERHIFeatureLevel::Type InFeatureLevel)
 {
 	check(SkelMeshResource);
 	check(SkelMeshResource->LODModels.IsValidIndex(LODIndex));
@@ -1063,10 +1063,10 @@ void FSkeletalMeshObjectGPUSkin::FSkeletalMeshObjectLOD::InitResources(const FSk
 	GetVertexBuffers(VertexBuffers,LODModel,MeshLODInfo);
 
 	// init gpu skin factories
-	GPUSkinVertexFactories.InitVertexFactories(VertexBuffers,LODModel.Chunks, FeatureLevel);
+	GPUSkinVertexFactories.InitVertexFactories(VertexBuffers,LODModel.Chunks, InFeatureLevel);
 	if ( LODModel.HasApexClothData() )
 	{
-		GPUSkinVertexFactories.InitAPEXClothVertexFactories(VertexBuffers,LODModel.Chunks, FeatureLevel);
+		GPUSkinVertexFactories.InitAPEXClothVertexFactories(VertexBuffers,LODModel.Chunks, InFeatureLevel);
 	}
 }
 
@@ -1082,7 +1082,7 @@ void FSkeletalMeshObjectGPUSkin::FSkeletalMeshObjectLOD::ReleaseResources()
 	GPUSkinVertexFactories.ReleaseAPEXClothVertexFactories();
 }
 
-void FSkeletalMeshObjectGPUSkin::FSkeletalMeshObjectLOD::InitMorphResources(const FSkelMeshObjectLODInfo& MeshLODInfo, bool bInUsePerBoneMotionBlur, ERHIFeatureLevel::Type FeatureLevel)
+void FSkeletalMeshObjectGPUSkin::FSkeletalMeshObjectLOD::InitMorphResources(const FSkelMeshObjectLODInfo& MeshLODInfo, bool bInUsePerBoneMotionBlur, ERHIFeatureLevel::Type InFeatureLevel)
 {
 	check(SkelMeshResource);
 	check(SkelMeshResource->LODModels.IsValidIndex(LODIndex));
@@ -1097,7 +1097,7 @@ void FSkeletalMeshObjectGPUSkin::FSkeletalMeshObjectLOD::InitMorphResources(cons
 	FVertexFactoryBuffers VertexBuffers;
 	GetVertexBuffers(VertexBuffers,LODModel,MeshLODInfo);
 	// init morph skin factories
-	GPUSkinVertexFactories.InitMorphVertexFactories(VertexBuffers, LODModel.Chunks, bInUsePerBoneMotionBlur, FeatureLevel);
+	GPUSkinVertexFactories.InitMorphVertexFactories(VertexBuffers, LODModel.Chunks, bInUsePerBoneMotionBlur, InFeatureLevel);
 }
 
 /** 
diff --git a/Engine/Source/Runtime/Engine/Private/Slate/SlateGameResources.cpp b/Engine/Source/Runtime/Engine/Private/Slate/SlateGameResources.cpp
--- a/Engine/Source/Runtime/Engine/Private/Slate/SlateGameResources.cpp
+++ b/Engine/Source/Runtime/Engine/Private/Slate/SlateGameResources.cpp
@@ -56,7 +56,7 @@ const FSlateBrush* FSlateGameResources::GetBrush( const FName PropertyName, cons
 	return FSlateStyleSet::GetBrush(PropertyName, Specifier);
 }
 
-const FSlateBrush* FSlateGameResources::GetOptionalBrush(const FName PropertyName, const ANSICHAR* Specifier, const FSlateBrush* const DefaultBrush) const
+const FSlateBrush* FSlateGameResources::GetOptionalBrush(const FName PropertyName, const ANSICHAR* Specifier, const FSlateBrush* const InDefaultBrush) const
 {
 	ensureMsgf(Specifier == NULL, TEXT("Attempting to look up resource (%s, %s). \n Specifiers not supported by Slate Resource Sets loaded from content browser."), *PropertyName.ToString(), Specifier);
 	FName CleanName = GetCleanName(PropertyName);
@@ -65,9 +65,9 @@ const FSlateBrush* FSlateGameResources::GetOptionalBrush(const FName PropertyNam
 	{
 		const USlateBrushAsset* BrushAsset = Cast<USlateBrushAsset>(*Resource);
 		ensureMsgf(BrushAsset, TEXT("Could not find resource '%s'"), *CleanName.ToString());
-		return BrushAsset ? &BrushAsset->Brush : DefaultBrush;
+		return BrushAsset ? &BrushAsset->Brush : InDefaultBrush;
 	}
-	return FSlateStyleSet::GetOptionalBrush(PropertyName, Specifier, DefaultBrush);
+	return FSlateStyleSet::GetOptionalBrush(PropertyName, Specifier, InDefaultBrush);
 }
 
 UCurveFloat* FSlateGameResources::GetCurveFloat( const FName AssetName ) const
diff --git a/Engine/Source/Runtime/Engine/Private/SoundCue.cpp b/Engine/Source/Runtime/Engine/Private/SoundCue.cpp
--- a/Engine/Source/Runtime/Engine/Private/SoundCue.cpp
+++ b/Engine/Source/Runtime/Engine/Private/SoundCue.cpp
@@ -238,10 +238,10 @@ FString USoundCue::GetDesc()
 	FString Description = TEXT( "" );
 
 	// Display duration
-	const float Duration = GetDuration();
-	if( Duration < INDEFINITELY_LOOPING_DURATION )
+	const float CueDuration = GetDuration();
+	if( CueDuration < INDEFINITELY_LOOPING_DURATION )
 	{
-		Description = FString::Printf( TEXT( "%3.2fs" ), Duration );
+		Description = FString::Printf( TEXT( "%3.2fs" ), CueDuration );
 	}
 	else
 	{
diff --git a/Engine/Source/Runtime/Engine/Private/SoundNodeConcatenator.cpp b/Engine/Source/Runtime/Engine/Private/SoundNodeConcatenator.cpp
--- a/Engine/Source/Runtime/Engine/Private/SoundNodeConcatenator.cpp
+++ b/Engine/Source/Runtime/Engine/Private/SoundNodeConcatenator.cpp
@@ -145,23 +145,27 @@ void USoundNodeConcatenator::ParseNodes( FAudioDevice* AudioDevice, const UPTRIN
 			UpdatedParams.NotifyBufferFinishedHooks.AddNotify(this, NodeWaveInstanceHash);
 			UpdatedParams.VolumeMultiplier *= InputVolume[LocalPayload.NodeIndex];
 
+			const int32 NumWaveInstancesBeforeParse = WaveInstances.Num();
+
 			const UPTRINT ChildNodeWaveInstanceHash = GetNodeWaveInstanceHash(NodeWaveInstanceHash, ChildNode, LocalPayload.NodeIndex);
 			ChildNode->ParseNodes(AudioDevice, ChildNodeWaveInstanceHash, ActiveSound, UpdatedParams, WaveInstances);
 
-			// Update the payload with the number of sounds played and update our local copy
-			RETRIEVE_SOUNDNODE_PAYLOAD(sizeof(FSoundNodeConcatenatorPayload));
-			DECLARE_SOUNDNODE_ELEMENT(FSoundNodeConcatenatorPayload, ConcatenatorPayload);
-
-			ConcatenatorPayload = LocalPayload;
-			ConcatenatorPayload.CurrentChildNodeNumSound = ChildNode->GetNumSounds(ChildNodeWaveInstanceHash, ActiveSound);
+			// Check to see if we actually added any wave instances when parsing the concat node.
+			if (WaveInstances.Num() > NumWaveInstancesBeforeParse)
+			{
+				// Update the payload with the number of sounds played and update our local copy
+				RETRIEVE_SOUNDNODE_PAYLOAD(sizeof(FSoundNodeConcatenatorPayload));
+				DECLARE_SOUNDNODE_ELEMENT(FSoundNodeConcatenatorPayload, ConcatenatorPayload);
 
-			break;
-		}
-		else
-		{
-			LocalPayload.CurrentChildNodeSoundIndex = 0;
-			++LocalPayload.NodeIndex;
+				ConcatenatorPayload = LocalPayload;
+				ConcatenatorPayload.CurrentChildNodeNumSound = ChildNode->GetNumSounds(ChildNodeWaveInstanceHash, ActiveSound);
+				break;
+			}
 		}
+
+		// If no child node, or if no wave instances played for this node index
+		LocalPayload.CurrentChildNodeSoundIndex = 0;
+		++LocalPayload.NodeIndex;
 	}
 }
 
diff --git a/Engine/Source/Runtime/Engine/Private/StaticMesh.cpp b/Engine/Source/Runtime/Engine/Private/StaticMesh.cpp
--- a/Engine/Source/Runtime/Engine/Private/StaticMesh.cpp
+++ b/Engine/Source/Runtime/Engine/Private/StaticMesh.cpp
@@ -1496,6 +1496,17 @@ FBox UStaticMesh::GetBoundingBox() const
 	return ExtendedBounds.GetBox();
 }
 
+int32 UStaticMesh::GetNumSections(int32 InLOD) const
+{
+	int32 NumSections = 0;
+	if (RenderData != NULL && RenderData->LODResources.IsValidIndex(InLOD))
+	{
+		const FStaticMeshLODResources& LOD = RenderData->LODResources[InLOD];
+		NumSections = LOD.Sections.Num();
+	}
+	return NumSections;
+}
+
 float UStaticMesh::GetStreamingTextureFactor(int32 RequestedUVIndex) const
 {
 	check(RequestedUVIndex >= 0);
diff --git a/Engine/Source/Runtime/Engine/Private/Tests/EngineAutomationTests.cpp b/Engine/Source/Runtime/Engine/Private/Tests/EngineAutomationTests.cpp
--- a/Engine/Source/Runtime/Engine/Private/Tests/EngineAutomationTests.cpp
+++ b/Engine/Source/Runtime/Engine/Private/Tests/EngineAutomationTests.cpp
@@ -258,8 +258,8 @@ bool FLoadAllMapsInGameTest::RunTest(const FString& Parameters)
 	{
 		//Generate the screen shot name and path
 		FString ScreenshotFileName;
-		const FString TestName = FString::Printf(TEXT("LoadAllMaps_Game/%s"), *FPaths::GetBaseFilename(MapName));
-		AutomationCommon::GetScreenshotPath(TestName, ScreenshotFileName, true);
+		const FString LoadAllMapsTestName = FString::Printf(TEXT("LoadAllMaps_Game/%s"), *FPaths::GetBaseFilename(MapName));
+		AutomationCommon::GetScreenshotPath(LoadAllMapsTestName, ScreenshotFileName, true);
 
 		//Give the map some time to load
 		ADD_LATENT_AUTOMATION_COMMAND(FEngineWaitLatentCommand(1.5f));
diff --git a/Engine/Source/Runtime/Engine/Private/TextureRenderTarget2D.cpp b/Engine/Source/Runtime/Engine/Private/TextureRenderTarget2D.cpp
--- a/Engine/Source/Runtime/Engine/Private/TextureRenderTarget2D.cpp
+++ b/Engine/Source/Runtime/Engine/Private/TextureRenderTarget2D.cpp
@@ -145,7 +145,7 @@ FString UTextureRenderTarget2D::GetDesc()
 	return FString::Printf( TEXT("Render to Texture %dx%d[%s]"), SizeX, SizeY, GPixelFormats[GetFormat()].Name );
 }
 
-UTexture2D* UTextureRenderTarget2D::ConstructTexture2D(UObject* Outer, const FString& NewTexName, EObjectFlags ObjectFlags, uint32 Flags, TArray<uint8>* AlphaOverride)
+UTexture2D* UTextureRenderTarget2D::ConstructTexture2D(UObject* Outer, const FString& NewTexName, EObjectFlags InObjectFlags, uint32 Flags, TArray<uint8>* AlphaOverride)
 {
 	UTexture2D* Result = NULL;
 #if WITH_EDITOR
@@ -157,7 +157,7 @@ UTexture2D* UTextureRenderTarget2D::ConstructTexture2D(UObject* Outer, const FSt
 
 	const EPixelFormat PixelFormat = GetFormat();
 	ETextureSourceFormat TextureFormat = TSF_Invalid;
-	TextureCompressionSettings CompressionSettings = TC_Default;
+	TextureCompressionSettings CompressionSettingsForTexture = TC_Default;
 	switch (PixelFormat)
 	{
 		case PF_B8G8R8A8:
@@ -165,7 +165,7 @@ UTexture2D* UTextureRenderTarget2D::ConstructTexture2D(UObject* Outer, const FSt
 		break;
 		case PF_FloatRGBA:
 			TextureFormat = TSF_RGBA16F;
-			CompressionSettings = TC_HDR;
+			CompressionSettingsForTexture = TC_HDR;
 		break;
 	}
 
@@ -176,7 +176,7 @@ UTexture2D* UTextureRenderTarget2D::ConstructTexture2D(UObject* Outer, const FSt
 	}
 
 	// create the 2d texture
-	Result = NewObject<UTexture2D>(Outer, FName(*NewTexName), ObjectFlags);
+	Result = NewObject<UTexture2D>(Outer, FName(*NewTexName), InObjectFlags);
 	// init to the same size as the 2d texture
 	Result->Source.Init(SizeX, SizeY, 1, 1, TextureFormat);
 
@@ -265,7 +265,7 @@ UTexture2D* UTextureRenderTarget2D::ConstructTexture2D(UObject* Outer, const FSt
 		Result->MipGenSettings = TMGS_NoMipmaps;
 	}
 
-	Result->CompressionSettings = CompressionSettings;
+	Result->CompressionSettings = CompressionSettingsForTexture;
 	if (Flags & CTF_Compress)
 	{
 		// Set compression options.
@@ -324,15 +324,15 @@ void FTextureRenderTarget2DResource::InitDynamicRHI()
 {
 	if( TargetSizeX > 0 && TargetSizeY > 0 )
 	{
-		bool bSRGB=true;
+		bool bUseSRGB=true;
 		// if render target gamma used was 1.0 then disable SRGB for the static texture
 		if( FMath::Abs(GetDisplayGamma() - 1.0f) < KINDA_SMALL_NUMBER )
 		{
-			bSRGB = false;
+			bUseSRGB = false;
 		}
 
 		// Create the RHI texture. Only one mip is used and the texture is targetable for resolve.
-		uint32 TexCreateFlags = bSRGB ? TexCreate_SRGB : 0;
+		uint32 TexCreateFlags = bUseSRGB ? TexCreate_SRGB : 0;
 		FRHIResourceCreateInfo CreateInfo = FRHIResourceCreateInfo(FClearValueBinding(ClearColor));
 
 		if (Owner->bAutoGenerateMips)
diff --git a/Engine/Source/Runtime/Engine/Private/TextureRenderTargetCube.cpp b/Engine/Source/Runtime/Engine/Private/TextureRenderTargetCube.cpp
--- a/Engine/Source/Runtime/Engine/Private/TextureRenderTargetCube.cpp
+++ b/Engine/Source/Runtime/Engine/Private/TextureRenderTargetCube.cpp
@@ -202,15 +202,15 @@ void FTextureRenderTargetCubeResource::InitDynamicRHI()
 {
 	if(Owner->SizeX > 0)
 	{
-		bool bSRGB = true;
+		bool bIsSRGB = true;
 		// if render target gamma used was 1.0 then disable SRGB for the static texture
 		if(FMath::Abs(GetDisplayGamma() - 1.0f) < KINDA_SMALL_NUMBER)
 		{
-			bSRGB = false;
+			bIsSRGB = false;
 		}
 
 		// Create the RHI texture. Only one mip is used and the texture is targetable for resolve.
-		uint32 TexCreateFlags = bSRGB ? TexCreate_SRGB : 0;
+		uint32 TexCreateFlags = bIsSRGB ? TexCreate_SRGB : 0;
 		{
 			FRHIResourceCreateInfo CreateInfo;
 			RHICreateTargetableShaderResourceCube(
diff --git a/Engine/Source/Runtime/Engine/Private/TickTaskManager.cpp b/Engine/Source/Runtime/Engine/Private/TickTaskManager.cpp
--- a/Engine/Source/Runtime/Engine/Private/TickTaskManager.cpp
+++ b/Engine/Source/Runtime/Engine/Private/TickTaskManager.cpp
@@ -830,7 +830,7 @@ public:
 				if ((CumulativeCooldown + ComparisonTickFunction->RelativeTickCooldown) > CooldownTime)
 				{
 					FTickFunction* TickFunction = TickFunctionsToReschedule[RescheduleIndex].TickFunction;
-					if (TickFunction->TickState == FTickFunction::ETickState::Enabled)
+					if (TickFunction->TickState != FTickFunction::ETickState::Disabled)
 					{
 						if (TickFunctionsToReschedule[RescheduleIndex].bDeferredRemove)
 						{
@@ -865,7 +865,7 @@ public:
 			for ( ; RescheduleIndex < TickFunctionsToReschedule.Num(); ++RescheduleIndex)
 			{
 				FTickFunction* TickFunction = TickFunctionsToReschedule[RescheduleIndex].TickFunction;
-				if (TickFunction->TickState == FTickFunction::ETickState::Enabled)
+				if (TickFunction->TickState != FTickFunction::ETickState::Disabled)
 				{
 					if (TickFunctionsToReschedule[RescheduleIndex].bDeferredRemove)
 					{
diff --git a/Engine/Source/Runtime/Engine/Private/TileRendering.cpp b/Engine/Source/Runtime/Engine/Private/TileRendering.cpp
--- a/Engine/Source/Runtime/Engine/Private/TileRendering.cpp
+++ b/Engine/Source/Runtime/Engine/Private/TileRendering.cpp
@@ -84,24 +84,24 @@ public:
 	/** Default constructor. */
 	FTileVertexFactory()
 	{
-		FLocalVertexFactory::FDataType Data;
+		FLocalVertexFactory::FDataType VertexData;
 		// position
-		Data.PositionComponent = FVertexStreamComponent(
+		VertexData.PositionComponent = FVertexStreamComponent(
 			&GTileRendererVertexBuffer,STRUCT_OFFSET(FMaterialTileVertex,Position),sizeof(FMaterialTileVertex),VET_Float3);
 		// tangents
-		Data.TangentBasisComponents[0] = FVertexStreamComponent(
+		VertexData.TangentBasisComponents[0] = FVertexStreamComponent(
 			&GTileRendererVertexBuffer,STRUCT_OFFSET(FMaterialTileVertex,TangentX),sizeof(FMaterialTileVertex),VET_PackedNormal);
-		Data.TangentBasisComponents[1] = FVertexStreamComponent(
+		VertexData.TangentBasisComponents[1] = FVertexStreamComponent(
 			&GTileRendererVertexBuffer,STRUCT_OFFSET(FMaterialTileVertex,TangentZ),sizeof(FMaterialTileVertex),VET_PackedNormal);
 		// color
-		Data.ColorComponent = FVertexStreamComponent(
+		VertexData.ColorComponent = FVertexStreamComponent(
 			&GTileRendererVertexBuffer,STRUCT_OFFSET(FMaterialTileVertex,Color),sizeof(FMaterialTileVertex),VET_Color);
 		// UVs
-		Data.TextureCoordinates.Add(FVertexStreamComponent(
+		VertexData.TextureCoordinates.Add(FVertexStreamComponent(
 			&GTileRendererVertexBuffer,STRUCT_OFFSET(FMaterialTileVertex,U),sizeof(FMaterialTileVertex),VET_Float2));
 
 		// update the data
-		SetData(Data);
+		SetData(VertexData);
 	}
 };
 TGlobalResource<FTileVertexFactory> GTileVertexFactory;
diff --git a/Engine/Source/Runtime/Engine/Private/TriangleRendering.cpp b/Engine/Source/Runtime/Engine/Private/TriangleRendering.cpp
--- a/Engine/Source/Runtime/Engine/Private/TriangleRendering.cpp
+++ b/Engine/Source/Runtime/Engine/Private/TriangleRendering.cpp
@@ -83,24 +83,24 @@ public:
 /** Default constructor. */
 FTriangleVertexFactory()
 {
-	FLocalVertexFactory::FDataType Data;
+	FLocalVertexFactory::FDataType VertexData;
 	// position
-	Data.PositionComponent = FVertexStreamComponent(
+	VertexData.PositionComponent = FVertexStreamComponent(
 		&GTriangleRendererVertexBuffer, STRUCT_OFFSET(FMaterialTriangleVertex, Position), sizeof(FMaterialTriangleVertex), VET_Float3);
 		// tangents
-		Data.TangentBasisComponents[0] = FVertexStreamComponent(
+		VertexData.TangentBasisComponents[0] = FVertexStreamComponent(
 		&GTriangleRendererVertexBuffer, STRUCT_OFFSET(FMaterialTriangleVertex, TangentX), sizeof(FMaterialTriangleVertex), VET_PackedNormal);
-		Data.TangentBasisComponents[1] = FVertexStreamComponent(
+		VertexData.TangentBasisComponents[1] = FVertexStreamComponent(
 		&GTriangleRendererVertexBuffer, STRUCT_OFFSET(FMaterialTriangleVertex, TangentZ), sizeof(FMaterialTriangleVertex), VET_PackedNormal);
 		// color
-		Data.ColorComponent = FVertexStreamComponent(
+		VertexData.ColorComponent = FVertexStreamComponent(
 		&GTriangleRendererVertexBuffer, STRUCT_OFFSET(FMaterialTriangleVertex, Color), sizeof(FMaterialTriangleVertex), VET_Color);
 		// UVs
-		Data.TextureCoordinates.Add(FVertexStreamComponent(
+		VertexData.TextureCoordinates.Add(FVertexStreamComponent(
 		&GTriangleRendererVertexBuffer, STRUCT_OFFSET(FMaterialTriangleVertex, U), sizeof(FMaterialTriangleVertex), VET_Float2));
 		
 		// update the data
-		SetData(Data);
+		SetData(VertexData);
 	}
 };
 
diff --git a/Engine/Source/Runtime/Engine/Private/TriggerActors.cpp b/Engine/Source/Runtime/Engine/Private/TriggerActors.cpp
--- a/Engine/Source/Runtime/Engine/Private/TriggerActors.cpp
+++ b/Engine/Source/Runtime/Engine/Private/TriggerActors.cpp
@@ -28,9 +28,9 @@ ATriggerCapsule::ATriggerCapsule(const FObjectInitializer& ObjectInitializer)
 	bCollideWhenPlacing = true;
 	SpawnCollisionHandlingMethod = ESpawnActorCollisionHandlingMethod::AdjustIfPossibleButDontSpawnIfColliding;
 
-	if (UBillboardComponent* SpriteComponent = GetSpriteComponent())
+	if (UBillboardComponent* TriggerSpriteComponent = GetSpriteComponent())
 	{
-		SpriteComponent->SetupAttachment(CapsuleCollisionComponent);
+		TriggerSpriteComponent->SetupAttachment(CapsuleCollisionComponent);
 	}
 }
 
@@ -93,9 +93,9 @@ ATriggerBox::ATriggerBox(const FObjectInitializer& ObjectInitializer)
 	BoxCollisionComponent->InitBoxExtent(FVector(40.0f, 40.0f, 40.0f));
 	BoxCollisionComponent->SetCollisionProfileName(TriggerCollisionProfileName);
 
-	if (UBillboardComponent* SpriteComponent = GetSpriteComponent())
+	if (UBillboardComponent* TriggerSpriteComponent = GetSpriteComponent())
 	{
-		SpriteComponent->SetupAttachment(BoxCollisionComponent);
+		TriggerSpriteComponent->SetupAttachment(BoxCollisionComponent);
 	}
 }
 
@@ -134,9 +134,9 @@ ATriggerSphere::ATriggerSphere(const FObjectInitializer& ObjectInitializer)
 	SphereCollisionComponent->InitSphereRadius(+40.0f);
 	SphereCollisionComponent->SetCollisionProfileName(TriggerCollisionProfileName);
 
-	if (UBillboardComponent* SpriteComponent = GetSpriteComponent())
+	if (UBillboardComponent* TriggerSpriteComponent = GetSpriteComponent())
 	{
-		SpriteComponent->SetupAttachment(SphereCollisionComponent);
+		TriggerSpriteComponent->SetupAttachment(SphereCollisionComponent);
 	}
 }
 
diff --git a/Engine/Source/Runtime/Engine/Private/UserInterface/CanvasItem.cpp b/Engine/Source/Runtime/Engine/Private/UserInterface/CanvasItem.cpp
--- a/Engine/Source/Runtime/Engine/Private/UserInterface/CanvasItem.cpp
+++ b/Engine/Source/Runtime/Engine/Private/UserInterface/CanvasItem.cpp
@@ -1039,8 +1039,8 @@ void FCanvasTextItem::DrawStringInternal_OfflineCache( FCanvas* InCanvas, const
 		{
 			if( LastTexture != Tex->Resource || BatchedElements == nullptr )
 			{
-				FBatchedElementParameters* BatchedElementParameters = nullptr;
-				BatchedElements = InCanvas->GetBatchedElements(FCanvas::ET_Triangle, BatchedElementParameters, Tex->Resource, BlendMode, FontRenderInfo.GlowInfo);
+				FBatchedElementParameters* BatchedElementParams = nullptr;
+				BatchedElements = InCanvas->GetBatchedElements(FCanvas::ET_Triangle, BatchedElementParams, Tex->Resource, BlendMode, FontRenderInfo.GlowInfo);
 				check(BatchedElements != nullptr);
 				// trade-off between memory and performance by pre-allocating more reserved space 
 				// for the triangles/vertices of the batched elements used to render the text tiles
@@ -1181,8 +1181,8 @@ void FCanvasTextItem::DrawStringInternal_RuntimeCache( FCanvas* InCanvas, const
 				FontTexture = FontCache->GetEngineTextureResource( FontTextureIndex );
 				check(FontTexture);
 
-				FBatchedElementParameters* BatchedElementParameters = nullptr;
-				BatchedElements = InCanvas->GetBatchedElements(FCanvas::ET_Triangle, BatchedElementParameters, FontTexture, BlendMode, FontRenderInfo.GlowInfo);
+				FBatchedElementParameters* BatchedElementParams = nullptr;
+				BatchedElements = InCanvas->GetBatchedElements(FCanvas::ET_Triangle, BatchedElementParams, FontTexture, BlendMode, FontRenderInfo.GlowInfo);
 				check(BatchedElements);
 
 				InvTextureSizeX = 1.0f/FontTexture->GetSizeX();
@@ -1327,8 +1327,8 @@ void FCanvasShapedTextItem::DrawStringInternal( FCanvas* InCanvas, const FVector
 					FontTexture = FontCache->GetEngineTextureResource( FontTextureIndex );
 					check(FontTexture);
 
-					FBatchedElementParameters* BatchedElementParameters = nullptr;
-					BatchedElements = InCanvas->GetBatchedElements(FCanvas::ET_Triangle, BatchedElementParameters, FontTexture, BlendMode, FontRenderInfo.GlowInfo);
+					FBatchedElementParameters* BatchedElementParams = nullptr;
+					BatchedElements = InCanvas->GetBatchedElements(FCanvas::ET_Triangle, BatchedElementParams, FontTexture, BlendMode, FontRenderInfo.GlowInfo);
 					check(BatchedElements);
 
 					InvTextureSizeX = 1.0f/FontTexture->GetSizeX();
diff --git a/Engine/Source/Runtime/Engine/Private/Volume.cpp b/Engine/Source/Runtime/Engine/Private/Volume.cpp
--- a/Engine/Source/Runtime/Engine/Private/Volume.cpp
+++ b/Engine/Source/Runtime/Engine/Private/Volume.cpp
@@ -43,10 +43,10 @@ void AVolume::PostEditChangeProperty(FPropertyChangedEvent& PropertyChangedEvent
 {
 	Super::PostEditChangeProperty(PropertyChangedEvent);
 
-	static FName BrushBuilder(TEXT("BrushBuilder"));
+	static FName NAME_BrushBuilder(TEXT("BrushBuilder"));
 
 	// The brush builder that created this volume has changed. Notify listeners
-	if( PropertyChangedEvent.ChangeType != EPropertyChangeType::Interactive && PropertyChangedEvent.MemberProperty && PropertyChangedEvent.MemberProperty->GetFName() == BrushBuilder )
+	if( PropertyChangedEvent.ChangeType != EPropertyChangeType::Interactive && PropertyChangedEvent.MemberProperty && PropertyChangedEvent.MemberProperty->GetFName() == NAME_BrushBuilder )
 	{
 		OnVolumeShapeChanged.Broadcast(*this);
 	}
diff --git a/Engine/Source/Runtime/Engine/Public/AnimEncoding.h b/Engine/Source/Runtime/Engine/Public/AnimEncoding.h
--- a/Engine/Source/Runtime/Engine/Public/AnimEncoding.h
+++ b/Engine/Source/Runtime/Engine/Public/AnimEncoding.h
@@ -529,7 +529,7 @@ FORCEINLINE float AnimEncoding::TimeToIndex(
 			checkSlow(KeyPos >= 0.0f);
 			const float KeyPosFloor = floorf(KeyPos);
 			PosIndex0Out = FMath::Min( FMath::TruncToInt(KeyPosFloor), NumKeys );
-			Alpha = KeyPos - KeyPosFloor;
+			Alpha = Seq.Interpolation == EAnimInterpolationType::Step ? 0.0f : KeyPos - KeyPosFloor;
 			PosIndex1Out = FMath::Min( PosIndex0Out + 1, NumKeys );
 		}
 	}
@@ -676,7 +676,7 @@ FORCEINLINE float AnimEncoding::TimeToIndex(
 		// compute the blend parameters for the keys we have found
 		int32 Delta= FMath::Max(HighFrame - LowFrame, 1);
 		const float Remainder = (FramePos - (float)LowFrame);
-		Alpha = Remainder / (float)Delta;
+		Alpha = Seq.Interpolation == EAnimInterpolationType::Step ? 0.f : (Remainder / (float)Delta);
 	}
 	
 	return Alpha;
diff --git a/Engine/Source/Runtime/Engine/Public/Animation/AnimPhysicsSolver.h b/Engine/Source/Runtime/Engine/Public/Animation/AnimPhysicsSolver.h
--- a/Engine/Source/Runtime/Engine/Public/Animation/AnimPhysicsSolver.h
+++ b/Engine/Source/Runtime/Engine/Public/Animation/AnimPhysicsSolver.h
@@ -528,6 +528,12 @@ public:
 	*/
 	static void ConstrainPlanar(float DeltaTime, TArray<FAnimPhysLinearLimit>& LimitContainer, FAnimPhysRigidBody* Body, const FTransform& PlaneTransform);
 
+	/** Constrains the position of a body within the requested sphere */
+	static void ConstrainSphericalInner(float DeltaTime, TArray<FAnimPhysLinearLimit>& LimitContainer, FAnimPhysRigidBody* Body, const FTransform& SphereTransform, float SphereRadius);
+
+	/** Constrains the position of a body outside of the requested sphere */
+	static void ConstrainSphericalOuter(float DeltaTime, TArray<FAnimPhysLinearLimit>& LimitContainer, FAnimPhysRigidBody* Body, const FTransform& SphereTransform, float SphereRadius);
+
 	//////////////////////////////////////////////////////////////////////////
 	// Spring creation methods
 	static void CreateSpring(TArray<FAnimPhysSpring>& SpringContainer, FAnimPhysRigidBody* Body0, FVector Position0, FAnimPhysRigidBody* Body1, FVector Position1);
diff --git a/Engine/Source/Runtime/Engine/Public/Animation/AnimTypes.h b/Engine/Source/Runtime/Engine/Public/Animation/AnimTypes.h
--- a/Engine/Source/Runtime/Engine/Public/Animation/AnimTypes.h
+++ b/Engine/Source/Runtime/Engine/Public/Animation/AnimTypes.h
@@ -349,3 +349,17 @@ struct FAnimWeight
 		return (InWeight >= (1.f - ZERO_ANIMWEIGHT_THRESH));
 	}
 };
+
+/** 
+ * Indicates how animation should be evaluated between keys.
+ */
+UENUM(BlueprintType)
+enum class EAnimInterpolationType : uint8
+{
+	/** Linear interpolation when looking up values between keys. */
+	Linear		UMETA(DisplayName="Linear"),
+
+	/** Step interpolation when looking up values between keys. */
+	Step		UMETA(DisplayName="Step"),
+};
+
diff --git a/Engine/Source/Runtime/Engine/Public/AudioDevice.h b/Engine/Source/Runtime/Engine/Public/AudioDevice.h
--- a/Engine/Source/Runtime/Engine/Public/AudioDevice.h
+++ b/Engine/Source/Runtime/Engine/Public/AudioDevice.h
@@ -912,6 +912,9 @@ protected:
 	/** Processes the set of pending sounds that need to be stopped */ 
 	void ProcessingPendingActiveSoundStops();
 
+	/** Query if the editor is in VR Preview for the current play world. Returns false for non-editor builds */
+	static bool CanUseVRAudioDevice();
+
 public:
 
 	/** The maximum number of concurrent audible sounds */
diff --git a/Engine/Source/Runtime/Engine/Public/Raster.h b/Engine/Source/Runtime/Engine/Public/Raster.h
--- a/Engine/Source/Runtime/Engine/Public/Raster.h
+++ b/Engine/Source/Runtime/Engine/Public/Raster.h
@@ -96,37 +96,37 @@ private:
 		float DeltaMinX,
 		float TopMaxX,
 		float DeltaMaxX,
-		float MinY,
-		float MaxY,
+		float InMinY,
+		float InMaxY,
 		bool BackFacing
 		)
 	{
-		int32	IntMinY = FMath::Clamp(FMath::CeilToInt(MinY),RasterPolicyType::GetMinY(),RasterPolicyType::GetMaxY() + 1),
-			IntMaxY = FMath::Clamp(FMath::CeilToInt(MaxY),RasterPolicyType::GetMinY(),RasterPolicyType::GetMaxY() + 1);
+		int32	IntMinY = FMath::Clamp(FMath::CeilToInt(InMinY),RasterPolicyType::GetMinY(),RasterPolicyType::GetMaxY() + 1),
+			IntMaxY = FMath::Clamp(FMath::CeilToInt(InMaxY),RasterPolicyType::GetMinY(),RasterPolicyType::GetMaxY() + 1);
 
 		for(int32 IntY = IntMinY;IntY < IntMaxY;IntY++)
 		{
-			float			Y = IntY - MinY,
-							MinX = TopMinX + DeltaMinX * Y,
-							MaxX = TopMaxX + DeltaMaxX * Y;
+			float			Y = IntY - InMinY;
+			float			LocalMinX = TopMinX + DeltaMinX * Y;
+			float			LocalMaxX = TopMaxX + DeltaMaxX * Y;
 			InterpolantType	MinInterpolant = TopMinInterpolant + DeltaMinInterpolant * Y,
 							MaxInterpolant = TopMaxInterpolant + DeltaMaxInterpolant * Y;
 
-			if(MinX > MaxX)
+			if(LocalMinX > LocalMaxX)
 			{
-				Exchange(MinX,MaxX);
+				Exchange(LocalMinX,LocalMaxX);
 				Exchange(MinInterpolant,MaxInterpolant);
 			}
 
-			if(MaxX > MinX)
+			if(LocalMaxX > LocalMinX)
 			{
-				int32				IntMinX = FMath::Clamp(FMath::CeilToInt(MinX),RasterPolicyType::GetMinX(),RasterPolicyType::GetMaxX() + 1),
-								IntMaxX = FMath::Clamp(FMath::CeilToInt(MaxX),RasterPolicyType::GetMinX(),RasterPolicyType::GetMaxX() + 1);
-				InterpolantType	DeltaInterpolant = (MaxInterpolant - MinInterpolant) / (MaxX - MinX);
+				int32				IntMinX = FMath::Clamp(FMath::CeilToInt(LocalMinX),RasterPolicyType::GetMinX(),RasterPolicyType::GetMaxX() + 1),
+								IntMaxX = FMath::Clamp(FMath::CeilToInt(LocalMaxX),RasterPolicyType::GetMinX(),RasterPolicyType::GetMaxX() + 1);
+				InterpolantType	DeltaInterpolant = (MaxInterpolant - MinInterpolant) / (LocalMaxX - LocalMinX);
 
 				for(int32 X = IntMinX;X < IntMaxX;X++)
 				{
-					RasterPolicyType::ProcessPixel(X,IntY,MinInterpolant + DeltaInterpolant * (X - MinX),BackFacing);
+					RasterPolicyType::ProcessPixel(X,IntY,MinInterpolant + DeltaInterpolant * (X - LocalMinX),BackFacing);
 				}
 			}
 		}
diff --git a/Engine/Source/Runtime/GameplayTasks/Private/GameplayTasksComponent.cpp b/Engine/Source/Runtime/GameplayTasks/Private/GameplayTasksComponent.cpp
--- a/Engine/Source/Runtime/GameplayTasks/Private/GameplayTasksComponent.cpp
+++ b/Engine/Source/Runtime/GameplayTasks/Private/GameplayTasksComponent.cpp
@@ -432,7 +432,8 @@ void UGameplayTasksComponent::UpdateTaskActivations()
 
 		for (int32 Idx = 0; Idx < ActivationList.Num(); Idx++)
 		{
-			if (ActivationList[Idx])
+			// check if task wasn't already finished as a result of activating previous elements of this list
+			if (ActivationList[Idx] && !ActivationList[Idx]->IsFinished())
 			{
 				ActivationList[Idx]->ActivateInTaskQueue();
 			}
diff --git a/Engine/Source/Runtime/Landscape/Private/Landscape.cpp b/Engine/Source/Runtime/Landscape/Private/Landscape.cpp
--- a/Engine/Source/Runtime/Landscape/Private/Landscape.cpp
+++ b/Engine/Source/Runtime/Landscape/Private/Landscape.cpp
@@ -867,16 +867,16 @@ void ULandscapeComponent::DestroyComponent(bool bPromoteChildren/*= false*/)
 
 FBoxSphereBounds ULandscapeComponent::CalcBounds(const FTransform& LocalToWorld) const
 {
-	FBox Bounds = CachedLocalBox.TransformBy(LocalToWorld);
-	Bounds = Bounds.ExpandBy({0, 0, NegativeZBoundsExtension}, {0, 0, PositiveZBoundsExtension});
+	FBox MyBounds = CachedLocalBox.TransformBy(LocalToWorld);
+	MyBounds = MyBounds.ExpandBy({0, 0, NegativeZBoundsExtension}, {0, 0, PositiveZBoundsExtension});
 
 	ALandscapeProxy* Proxy = GetLandscapeProxy();
 	if (Proxy)
 	{
-		Bounds = Bounds.ExpandBy({0, 0, Proxy->NegativeZBoundsExtension}, {0, 0, Proxy->PositiveZBoundsExtension});
+		MyBounds = MyBounds.ExpandBy({0, 0, Proxy->NegativeZBoundsExtension}, {0, 0, Proxy->PositiveZBoundsExtension});
 	}
 
-	return FBoxSphereBounds(Bounds);
+	return FBoxSphereBounds(MyBounds);
 }
 
 void ULandscapeComponent::OnRegister()
diff --git a/Engine/Source/Runtime/Landscape/Private/LandscapeEdit.cpp b/Engine/Source/Runtime/Landscape/Private/LandscapeEdit.cpp
--- a/Engine/Source/Runtime/Landscape/Private/LandscapeEdit.cpp
+++ b/Engine/Source/Runtime/Landscape/Private/LandscapeEdit.cpp
@@ -57,8 +57,7 @@ ULandscapeLayerInfoObject* ALandscapeProxy::VisibilityLayer = NULL;
 void ULandscapeComponent::Init(int32 InBaseX, int32 InBaseY, int32 InComponentSizeQuads, int32 InNumSubsections, int32 InSubsectionSizeQuads)
 {
 	SetSectionBase(FIntPoint(InBaseX, InBaseY));
-	FVector RelativeLocation = FVector(GetSectionBase() - GetLandscapeProxy()->LandscapeSectionOffset);
-	SetRelativeLocation(RelativeLocation);
+	SetRelativeLocation(FVector(GetSectionBase() - GetLandscapeProxy()->LandscapeSectionOffset));
 	ComponentSizeQuads = InComponentSizeQuads;
 	NumSubsections = InNumSubsections;
 	SubsectionSizeQuads = InSubsectionSizeQuads;
diff --git a/Engine/Source/Runtime/Landscape/Private/LandscapeGrass.cpp b/Engine/Source/Runtime/Landscape/Private/LandscapeGrass.cpp
--- a/Engine/Source/Runtime/Landscape/Private/LandscapeGrass.cpp
+++ b/Engine/Source/Runtime/Landscape/Private/LandscapeGrass.cpp
@@ -1369,10 +1369,10 @@ struct FAsyncGrassBuilder : public FGrassBuilderBase
 	{
 		const int32 SubsectionSizeQuads = Component->SubsectionSizeQuads;
 		const int32 NumSubsections = Component->NumSubsections;
-		const int32 ComponentSizeQuads = Component->ComponentSizeQuads;
+		const int32 LandscapeComponentSizeQuads = Component->ComponentSizeQuads;
 	
 		const int32 StaticLightingLOD = Component->GetLandscapeProxy()->StaticLightingLOD;
-		const int32 ComponentSizeVerts = ComponentSizeQuads + 1;
+		const int32 ComponentSizeVerts = LandscapeComponentSizeQuads + 1;
 		const float LightMapRes = Component->StaticLightingResolution > 0.0f ? Component->StaticLightingResolution : Component->GetLandscapeProxy()->StaticLightingResolution;
 		const int32 LightingLOD = Component->GetLandscapeProxy()->StaticLightingLOD;
 
@@ -1381,13 +1381,13 @@ struct FAsyncGrassBuilder : public FGrassBuilderBase
 		int32 PatchExpandCountX = 0;
 		int32 PatchExpandCountY = 0;
 		int32 DesiredSize = 1;
-		const float LightMapRatio = ::GetTerrainExpandPatchCount(LightMapRes, PatchExpandCountX, PatchExpandCountY, ComponentSizeQuads, (NumSubsections * (SubsectionSizeQuads + 1)), DesiredSize, LightingLOD);
+		const float LightMapRatio = ::GetTerrainExpandPatchCount(LightMapRes, PatchExpandCountX, PatchExpandCountY, LandscapeComponentSizeQuads, (NumSubsections * (SubsectionSizeQuads + 1)), DesiredSize, LightingLOD);
 		const float LightmapLODScaleX = LightMapRatio / ((ComponentSizeVerts >> StaticLightingLOD) + 2 * PatchExpandCountX);
 		const float LightmapLODScaleY = LightMapRatio / ((ComponentSizeVerts >> StaticLightingLOD) + 2 * PatchExpandCountY);
 		const float LightmapBiasX = PatchExpandCountX * LightmapLODScaleX;
 		const float LightmapBiasY = PatchExpandCountY * LightmapLODScaleY;
-		const float LightmapScaleX = LightmapLODScaleX * (float)((ComponentSizeVerts >> StaticLightingLOD) - 1) / ComponentSizeQuads;
-		const float LightmapScaleY = LightmapLODScaleY * (float)((ComponentSizeVerts >> StaticLightingLOD) - 1) / ComponentSizeQuads;
+		const float LightmapScaleX = LightmapLODScaleX * (float)((ComponentSizeVerts >> StaticLightingLOD) - 1) / LandscapeComponentSizeQuads;
+		const float LightmapScaleY = LightmapLODScaleY * (float)((ComponentSizeVerts >> StaticLightingLOD) - 1) / LandscapeComponentSizeQuads;
 
 		LightMapComponentScale = FVector2D(LightmapScaleX, LightmapScaleY) / FVector2D(DrawScale);
 		LightMapComponentBias = FVector2D(LightmapBiasX, LightmapBiasY);
diff --git a/Engine/Source/Runtime/Landscape/Private/LandscapeSplines.cpp b/Engine/Source/Runtime/Landscape/Private/LandscapeSplines.cpp
--- a/Engine/Source/Runtime/Landscape/Private/LandscapeSplines.cpp
+++ b/Engine/Source/Runtime/Landscape/Private/LandscapeSplines.cpp
@@ -127,7 +127,7 @@ public:
 		// e.g. someone playing with splines on a newly-created perfectly-flat landscape
 		static const float DepthBias = 0.0001;
 
-		const FMatrix& LocalToWorld = GetLocalToWorld();
+		const FMatrix& MyLocalToWorld = GetLocalToWorld();
 
 		const FLinearColor SelectedSplineColor = GEngine->GetSelectedMaterialColor();
 		const FLinearColor SelectedControlPointSpriteColor = FLinearColor::White + (GEngine->GetSelectedMaterialColor() * GEngine->SelectionHighlightIntensityBillboards * 10); 
@@ -144,19 +144,19 @@ public:
 					const FLinearColor SegmentColor = Segment.bSelected ? SelectedSplineColor : SplineColor;
 
 					FLandscapeSplineInterpPoint OldPoint = Segment.Points[0];
-					OldPoint.Center       = LocalToWorld.TransformPosition(OldPoint.Center);
-					OldPoint.Left         = LocalToWorld.TransformPosition(OldPoint.Left);
-					OldPoint.Right        = LocalToWorld.TransformPosition(OldPoint.Right);
-					OldPoint.FalloffLeft  = LocalToWorld.TransformPosition(OldPoint.FalloffLeft);
-					OldPoint.FalloffRight = LocalToWorld.TransformPosition(OldPoint.FalloffRight);
+					OldPoint.Center       = MyLocalToWorld.TransformPosition(OldPoint.Center);
+					OldPoint.Left         = MyLocalToWorld.TransformPosition(OldPoint.Left);
+					OldPoint.Right        = MyLocalToWorld.TransformPosition(OldPoint.Right);
+					OldPoint.FalloffLeft  = MyLocalToWorld.TransformPosition(OldPoint.FalloffLeft);
+					OldPoint.FalloffRight = MyLocalToWorld.TransformPosition(OldPoint.FalloffRight);
 					for (int32 i = 1; i < Segment.Points.Num(); i++)
 					{
 						FLandscapeSplineInterpPoint NewPoint = Segment.Points[i];
-						NewPoint.Center       = LocalToWorld.TransformPosition(NewPoint.Center);
-						NewPoint.Left         = LocalToWorld.TransformPosition(NewPoint.Left);
-						NewPoint.Right        = LocalToWorld.TransformPosition(NewPoint.Right);
-						NewPoint.FalloffLeft  = LocalToWorld.TransformPosition(NewPoint.FalloffLeft);
-						NewPoint.FalloffRight = LocalToWorld.TransformPosition(NewPoint.FalloffRight);
+						NewPoint.Center       = MyLocalToWorld.TransformPosition(NewPoint.Center);
+						NewPoint.Left         = MyLocalToWorld.TransformPosition(NewPoint.Left);
+						NewPoint.Right        = MyLocalToWorld.TransformPosition(NewPoint.Right);
+						NewPoint.FalloffLeft  = MyLocalToWorld.TransformPosition(NewPoint.FalloffLeft);
+						NewPoint.FalloffRight = MyLocalToWorld.TransformPosition(NewPoint.FalloffRight);
 
 						// Draw lines from the last keypoint.
 						PDI->SetHitProxy(Segment.HitProxy);
@@ -183,12 +183,12 @@ public:
 
 				for (const FControlPointProxy& ControlPoint : ControlPoints)
 				{
-					const FVector ControlPointLocation = LocalToWorld.TransformPosition(ControlPoint.Location);
+					const FVector ControlPointLocation = MyLocalToWorld.TransformPosition(ControlPoint.Location);
 
 					// Draw Sprite
 					if (bDrawControlPointSprite)
 					{
-						const float ControlPointSpriteScale = LocalToWorld.GetScaleVector().X * ControlPoint.SpriteScale;
+						const float ControlPointSpriteScale = MyLocalToWorld.GetScaleVector().X * ControlPoint.SpriteScale;
 						const FVector ControlPointSpriteLocation = ControlPointLocation + FVector(0, 0, ControlPointSpriteScale * 0.75f);
 						const FLinearColor ControlPointSpriteColor = ControlPoint.bSelected ? SelectedControlPointSpriteColor : FLinearColor::White;
 
@@ -213,11 +213,11 @@ public:
 					if (ControlPoint.Points.Num() == 1)
 					{
 						FLandscapeSplineInterpPoint NewPoint = ControlPoint.Points[0];
-						NewPoint.Center = LocalToWorld.TransformPosition(NewPoint.Center);
-						NewPoint.Left   = LocalToWorld.TransformPosition(NewPoint.Left);
-						NewPoint.Right  = LocalToWorld.TransformPosition(NewPoint.Right);
-						NewPoint.FalloffLeft  = LocalToWorld.TransformPosition(NewPoint.FalloffLeft);
-						NewPoint.FalloffRight = LocalToWorld.TransformPosition(NewPoint.FalloffRight);
+						NewPoint.Center = MyLocalToWorld.TransformPosition(NewPoint.Center);
+						NewPoint.Left   = MyLocalToWorld.TransformPosition(NewPoint.Left);
+						NewPoint.Right  = MyLocalToWorld.TransformPosition(NewPoint.Right);
+						NewPoint.FalloffLeft  = MyLocalToWorld.TransformPosition(NewPoint.FalloffLeft);
+						NewPoint.FalloffRight = MyLocalToWorld.TransformPosition(NewPoint.FalloffRight);
 
 						// draw end for spline connection
 						PDI->DrawPoint(NewPoint.Center, ControlPointColor, 6.0f, GetDepthPriorityGroup(View));
@@ -232,19 +232,19 @@ public:
 					else if (ControlPoint.Points.Num() >= 2)
 					{
 						FLandscapeSplineInterpPoint OldPoint = ControlPoint.Points.Last();
-						//OldPoint.Left   = LocalToWorld.TransformPosition(OldPoint.Left);
-						OldPoint.Right  = LocalToWorld.TransformPosition(OldPoint.Right);
-						//OldPoint.FalloffLeft  = LocalToWorld.TransformPosition(OldPoint.FalloffLeft);
-						OldPoint.FalloffRight = LocalToWorld.TransformPosition(OldPoint.FalloffRight);
+						//OldPoint.Left   = MyLocalToWorld.TransformPosition(OldPoint.Left);
+						OldPoint.Right  = MyLocalToWorld.TransformPosition(OldPoint.Right);
+						//OldPoint.FalloffLeft  = MyLocalToWorld.TransformPosition(OldPoint.FalloffLeft);
+						OldPoint.FalloffRight = MyLocalToWorld.TransformPosition(OldPoint.FalloffRight);
 
 						for (const FLandscapeSplineInterpPoint& Point : ControlPoint.Points)
 						{
 							FLandscapeSplineInterpPoint NewPoint = Point;
-							NewPoint.Center = LocalToWorld.TransformPosition(NewPoint.Center);
-							NewPoint.Left   = LocalToWorld.TransformPosition(NewPoint.Left);
-							NewPoint.Right  = LocalToWorld.TransformPosition(NewPoint.Right);
-							NewPoint.FalloffLeft  = LocalToWorld.TransformPosition(NewPoint.FalloffLeft);
-							NewPoint.FalloffRight = LocalToWorld.TransformPosition(NewPoint.FalloffRight);
+							NewPoint.Center = MyLocalToWorld.TransformPosition(NewPoint.Center);
+							NewPoint.Left   = MyLocalToWorld.TransformPosition(NewPoint.Left);
+							NewPoint.Right  = MyLocalToWorld.TransformPosition(NewPoint.Right);
+							NewPoint.FalloffLeft  = MyLocalToWorld.TransformPosition(NewPoint.FalloffLeft);
+							NewPoint.FalloffRight = MyLocalToWorld.TransformPosition(NewPoint.FalloffRight);
 
 							PDI->SetHitProxy(ControlPoint.HitProxy);
 
diff --git a/Engine/Source/Runtime/MovieSceneTracks/Private/Sections/MovieSceneBoolSection.cpp b/Engine/Source/Runtime/MovieSceneTracks/Private/Sections/MovieSceneBoolSection.cpp
--- a/Engine/Source/Runtime/MovieSceneTracks/Private/Sections/MovieSceneBoolSection.cpp
+++ b/Engine/Source/Runtime/MovieSceneTracks/Private/Sections/MovieSceneBoolSection.cpp
@@ -28,9 +28,6 @@ void UMovieSceneBoolSection::MoveSection( float DeltaPosition, TSet<FKeyHandle>&
 
 void UMovieSceneBoolSection::DilateSection( float DilationFactor, float Origin, TSet<FKeyHandle>& KeyHandles )
 {
-	float StartTime = GetStartTime();
-	float EndTime = GetEndTime();
-
 	Super::DilateSection(DilationFactor, Origin, KeyHandles);
 	
 	BoolCurve.ScaleCurve(Origin, DilationFactor, KeyHandles);
diff --git a/Engine/Source/Runtime/OpenGLDrv/Public/OpenGLResources.h b/Engine/Source/Runtime/OpenGLDrv/Public/OpenGLResources.h
--- a/Engine/Source/Runtime/OpenGLDrv/Public/OpenGLResources.h
+++ b/Engine/Source/Runtime/OpenGLDrv/Public/OpenGLResources.h
@@ -65,13 +65,13 @@ class TOpenGLBuffer : public BaseType
 		{
 			while ( InSize > 0)
 			{
-				const uint32 Size = FMath::Min<uint32>( BlockSize, InSize);
+				const uint32 BufferSize = FMath::Min<uint32>( BlockSize, InSize);
 			
-				FOpenGL::BufferSubData( Type, InOffset, Size, Data);
+				FOpenGL::BufferSubData( Type, InOffset, BufferSize, Data);
 
-				InOffset += Size;
-				InSize -= Size;
-				Data += Size;
+				InOffset += BufferSize;
+				InSize -= BufferSize;
+				Data += BufferSize;
 			}
 		}
 		else
diff --git a/Engine/Source/Runtime/Renderer/Private/SceneRendering.cpp b/Engine/Source/Runtime/Renderer/Private/SceneRendering.cpp
--- a/Engine/Source/Runtime/Renderer/Private/SceneRendering.cpp
+++ b/Engine/Source/Runtime/Renderer/Private/SceneRendering.cpp
@@ -1011,7 +1011,7 @@ void FViewInfo::CreateUniformBuffer(
 	FrameUniformShaderParameters.CircleDOFParams = CircleDofHalfCoc(*this);
 
 	FScene* Scene = (FScene*)Family->Scene;
-	ERHIFeatureLevel::Type FeatureLevel = Scene == nullptr ? GMaxRHIFeatureLevel : Scene->GetFeatureLevel();
+	ERHIFeatureLevel::Type RHIFeatureLevel = Scene == nullptr ? GMaxRHIFeatureLevel : Scene->GetFeatureLevel();
 
 	if (Scene && Scene->SkyLight)
 	{
@@ -1038,7 +1038,7 @@ void FViewInfo::CreateUniformBuffer(
 
 	FrameUniformShaderParameters.MobilePreviewMode =
 		(GIsEditor &&
-		(FeatureLevel == ERHIFeatureLevel::ES2 || FeatureLevel == ERHIFeatureLevel::ES3_1) &&
+		(RHIFeatureLevel == ERHIFeatureLevel::ES2 || RHIFeatureLevel == ERHIFeatureLevel::ES3_1) &&
 		GMaxRHIFeatureLevel > ERHIFeatureLevel::ES3_1) ? 1.0f : 0.0f;
 
 	// Padding between the left and right eye may be introduced by an HMD, which instanced stereo needs to account for.
diff --git a/Engine/Source/Runtime/Renderer/Private/StaticMeshDrawList.inl b/Engine/Source/Runtime/Renderer/Private/StaticMeshDrawList.inl
--- a/Engine/Source/Runtime/Renderer/Private/StaticMeshDrawList.inl
+++ b/Engine/Source/Runtime/Renderer/Private/StaticMeshDrawList.inl
@@ -890,7 +890,7 @@ void TStaticMeshDrawList<DrawingPolicyType>::SortFrontToBack(FVector ViewPositio
 }
 
 template<typename DrawingPolicyType>
-void TStaticMeshDrawList<DrawingPolicyType>::GetUsedPrimitivesBasedOnMaterials(ERHIFeatureLevel::Type FeatureLevel, const TArray<const FMaterial*>& Materials, TArray<FPrimitiveSceneInfo*>& PrimitivesToUpdate)
+void TStaticMeshDrawList<DrawingPolicyType>::GetUsedPrimitivesBasedOnMaterials(ERHIFeatureLevel::Type InFeatureLevel, const TArray<const FMaterial*>& Materials, TArray<FPrimitiveSceneInfo*>& PrimitivesToUpdate)
 {
 	for (typename TDrawingPolicySet::TIterator DrawingPolicyIt(DrawingPolicySet); DrawingPolicyIt; ++DrawingPolicyIt)
 	{
@@ -907,7 +907,7 @@ void TStaticMeshDrawList<DrawingPolicyType>::GetUsedPrimitivesBasedOnMaterials(E
 			if (Proxy)
 			{
 				check(!Proxy->IsDeleted());
-				FMaterial* MaterialResource = Proxy->GetMaterialNoFallback(FeatureLevel);
+				FMaterial* MaterialResource = Proxy->GetMaterialNoFallback(InFeatureLevel);
 
 				if (Materials.Contains(MaterialResource))
 				{
diff --git a/Engine/Source/Runtime/UMG/Private/Components/Border.cpp b/Engine/Source/Runtime/UMG/Private/Components/Border.cpp
--- a/Engine/Source/Runtime/UMG/Private/Components/Border.cpp
+++ b/Engine/Source/Runtime/UMG/Private/Components/Border.cpp
@@ -71,11 +71,11 @@ UClass* UBorder::GetSlotClass() const
 	return UBorderSlot::StaticClass();
 }
 
-void UBorder::OnSlotAdded(UPanelSlot* Slot)
+void UBorder::OnSlotAdded(UPanelSlot* InSlot)
 {
 	// Copy the content properties into the new slot so that it matches what has been setup
 	// so far by the user.
-	UBorderSlot* BorderSlot = CastChecked<UBorderSlot>(Slot);
+	UBorderSlot* BorderSlot = CastChecked<UBorderSlot>(InSlot);
 	BorderSlot->Padding = Padding;
 	BorderSlot->HorizontalAlignment = HorizontalAlignment;
 	BorderSlot->VerticalAlignment = VerticalAlignment;
@@ -88,7 +88,7 @@ void UBorder::OnSlotAdded(UPanelSlot* Slot)
 	}
 }
 
-void UBorder::OnSlotRemoved(UPanelSlot* Slot)
+void UBorder::OnSlotRemoved(UPanelSlot* InSlot)
 {
 	// Remove the widget from the live slot if it exists.
 	if ( MyBorder.IsValid() )
@@ -305,19 +305,19 @@ void UBorder::PostEditChangeProperty(struct FPropertyChangedEvent& PropertyChang
 		{
 			FName PropertyName = PropertyChangedEvent.Property->GetFName();
 
-			if ( UBorderSlot* Slot = Cast<UBorderSlot>(GetContentSlot()) )
+			if ( UBorderSlot* BorderSlot = Cast<UBorderSlot>(GetContentSlot()) )
 			{
 				if ( PropertyName == "Padding" )
 				{
-					FObjectEditorUtils::MigratePropertyValue(this, "Padding", Slot, "Padding");
+					FObjectEditorUtils::MigratePropertyValue(this, "Padding", BorderSlot, "Padding");
 				}
 				else if ( PropertyName == "HorizontalAlignment" )
 				{
-					FObjectEditorUtils::MigratePropertyValue(this, "HorizontalAlignment", Slot, "HorizontalAlignment");
+					FObjectEditorUtils::MigratePropertyValue(this, "HorizontalAlignment", BorderSlot, "HorizontalAlignment");
 				}
 				else if ( PropertyName == "VerticalAlignment" )
 				{
-					FObjectEditorUtils::MigratePropertyValue(this, "VerticalAlignment", Slot, "VerticalAlignment");
+					FObjectEditorUtils::MigratePropertyValue(this, "VerticalAlignment", BorderSlot, "VerticalAlignment");
 				}
 			}
 		}
diff --git a/Engine/Source/Runtime/UMG/Private/Components/Button.cpp b/Engine/Source/Runtime/UMG/Private/Components/Button.cpp
--- a/Engine/Source/Runtime/UMG/Private/Components/Button.cpp
+++ b/Engine/Source/Runtime/UMG/Private/Components/Button.cpp
@@ -64,16 +64,16 @@ UClass* UButton::GetSlotClass() const
 	return UButtonSlot::StaticClass();
 }
 
-void UButton::OnSlotAdded(UPanelSlot* Slot)
+void UButton::OnSlotAdded(UPanelSlot* InSlot)
 {
 	// Add the child to the live slot if it already exists
 	if ( MyButton.IsValid() )
 	{
-		Cast<UButtonSlot>(Slot)->BuildSlot(MyButton.ToSharedRef());
+		CastChecked<UButtonSlot>(InSlot)->BuildSlot(MyButton.ToSharedRef());
 	}
 }
 
-void UButton::OnSlotRemoved(UPanelSlot* Slot)
+void UButton::OnSlotRemoved(UPanelSlot* InSlot)
 {
 	// Remove the widget from the live slot if it exists.
 	if ( MyButton.IsValid() )
diff --git a/Engine/Source/Runtime/UMG/Private/Components/CanvasPanel.cpp b/Engine/Source/Runtime/UMG/Private/Components/CanvasPanel.cpp
--- a/Engine/Source/Runtime/UMG/Private/Components/CanvasPanel.cpp
+++ b/Engine/Source/Runtime/UMG/Private/Components/CanvasPanel.cpp
@@ -28,21 +28,21 @@ UClass* UCanvasPanel::GetSlotClass() const
 	return UCanvasPanelSlot::StaticClass();
 }
 
-void UCanvasPanel::OnSlotAdded(UPanelSlot* Slot)
+void UCanvasPanel::OnSlotAdded(UPanelSlot* InSlot)
 {
 	// Add the child to the live canvas if it already exists
 	if ( MyCanvas.IsValid() )
 	{
-		Cast<UCanvasPanelSlot>(Slot)->BuildSlot(MyCanvas.ToSharedRef());
+		CastChecked<UCanvasPanelSlot>(InSlot)->BuildSlot(MyCanvas.ToSharedRef());
 	}
 }
 
-void UCanvasPanel::OnSlotRemoved(UPanelSlot* Slot)
+void UCanvasPanel::OnSlotRemoved(UPanelSlot* InSlot)
 {
 	// Remove the widget from the live slot if it exists.
 	if ( MyCanvas.IsValid() )
 	{
-		TSharedPtr<SWidget> Widget = Slot->Content->GetCachedWidget();
+		TSharedPtr<SWidget> Widget = InSlot->Content->GetCachedWidget();
 		if ( Widget.IsValid() )
 		{
 			MyCanvas->RemoveSlot(Widget.ToSharedRef());
@@ -54,9 +54,9 @@ TSharedRef<SWidget> UCanvasPanel::RebuildWidget()
 {
 	MyCanvas = SNew(SConstraintCanvas);
 
-	for ( UPanelSlot* Slot : Slots )
+	for ( UPanelSlot* PanelSlot : Slots )
 	{
-		if ( UCanvasPanelSlot* TypedSlot = Cast<UCanvasPanelSlot>(Slot) )
+		if ( UCanvasPanelSlot* TypedSlot = Cast<UCanvasPanelSlot>(PanelSlot) )
 		{
 			TypedSlot->Parent = this;
 			TypedSlot->BuildSlot(MyCanvas.ToSharedRef());
@@ -78,13 +78,13 @@ TSharedPtr<SConstraintCanvas> UCanvasPanel::GetCanvasWidget() const
 
 bool UCanvasPanel::GetGeometryForSlot(int32 SlotIndex, FGeometry& ArrangedGeometry) const
 {
-	UCanvasPanelSlot* Slot = Cast<UCanvasPanelSlot>(Slots[SlotIndex]);
-	return GetGeometryForSlot(Slot, ArrangedGeometry);
+	UCanvasPanelSlot* PanelSlot = CastChecked<UCanvasPanelSlot>(Slots[SlotIndex]);
+	return GetGeometryForSlot(PanelSlot, ArrangedGeometry);
 }
 
-bool UCanvasPanel::GetGeometryForSlot(UCanvasPanelSlot* Slot, FGeometry& ArrangedGeometry) const
+bool UCanvasPanel::GetGeometryForSlot(UCanvasPanelSlot* InSlot, FGeometry& ArrangedGeometry) const
 {
-	if ( Slot->Content == nullptr )
+	if ( InSlot->Content == nullptr )
 	{
 		return false;
 	}
@@ -97,7 +97,7 @@ bool UCanvasPanel::GetGeometryForSlot(UCanvasPanelSlot* Slot, FGeometry& Arrange
 
 		for ( int32 ChildIndex = 0; ChildIndex < ArrangedChildren.Num(); ChildIndex++ )
 		{
-			if ( ArrangedChildren[ChildIndex].Widget == Slot->Content->GetCachedWidget() )
+			if ( ArrangedChildren[ChildIndex].Widget == InSlot->Content->GetCachedWidget() )
 			{
 				ArrangedGeometry = ArrangedChildren[ChildIndex].Geometry;
 				return true;
diff --git a/Engine/Source/Runtime/UMG/Private/Components/CheckBox.cpp b/Engine/Source/Runtime/UMG/Private/Components/CheckBox.cpp
--- a/Engine/Source/Runtime/UMG/Private/Components/CheckBox.cpp
+++ b/Engine/Source/Runtime/UMG/Private/Components/CheckBox.cpp
@@ -56,16 +56,16 @@ void UCheckBox::SynchronizeProperties()
 	MyCheckbox->SetIsChecked( OPTIONAL_BINDING(ECheckBoxState, CheckedState) );
 }
 
-void UCheckBox::OnSlotAdded(UPanelSlot* Slot)
+void UCheckBox::OnSlotAdded(UPanelSlot* InSlot)
 {
 	// Add the child to the live slot if it already exists
 	if ( MyCheckbox.IsValid() )
 	{
-		MyCheckbox->SetContent(Slot->Content ? Slot->Content->TakeWidget() : SNullWidget::NullWidget);
+		MyCheckbox->SetContent(InSlot->Content ? InSlot->Content->TakeWidget() : SNullWidget::NullWidget);
 	}
 }
 
-void UCheckBox::OnSlotRemoved(UPanelSlot* Slot)
+void UCheckBox::OnSlotRemoved(UPanelSlot* InSlot)
 {
 	// Remove the widget from the live slot if it exists.
 	if ( MyCheckbox.IsValid() )
diff --git a/Engine/Source/Runtime/UMG/Private/Components/GridPanel.cpp b/Engine/Source/Runtime/UMG/Private/Components/GridPanel.cpp
--- a/Engine/Source/Runtime/UMG/Private/Components/GridPanel.cpp
+++ b/Engine/Source/Runtime/UMG/Private/Components/GridPanel.cpp
@@ -28,21 +28,21 @@ UClass* UGridPanel::GetSlotClass() const
 	return UGridSlot::StaticClass();
 }
 
-void UGridPanel::OnSlotAdded(UPanelSlot* Slot)
+void UGridPanel::OnSlotAdded(UPanelSlot* InSlot)
 {
 	// Add the child to the live canvas if it already exists
 	if ( MyGridPanel.IsValid() )
 	{
-		Cast<UGridSlot>(Slot)->BuildSlot(MyGridPanel.ToSharedRef());
+		CastChecked<UGridSlot>(InSlot)->BuildSlot(MyGridPanel.ToSharedRef());
 	}
 }
 
-void UGridPanel::OnSlotRemoved(UPanelSlot* Slot)
+void UGridPanel::OnSlotRemoved(UPanelSlot* InSlot)
 {
 	// Remove the widget from the live slot if it exists.
 	if ( MyGridPanel.IsValid() )
 	{
-		TSharedPtr<SWidget> Widget = Slot->Content->GetCachedWidget();
+		TSharedPtr<SWidget> Widget = InSlot->Content->GetCachedWidget();
 		if ( Widget.IsValid() )
 		{
 			MyGridPanel->RemoveSlot(Widget.ToSharedRef());
@@ -54,9 +54,9 @@ TSharedRef<SWidget> UGridPanel::RebuildWidget()
 {
 	MyGridPanel = SNew(SGridPanel);
 
-	for ( UPanelSlot* Slot : Slots )
+	for ( UPanelSlot* PanelSlot : Slots )
 	{
-		if ( UGridSlot* TypedSlot = Cast<UGridSlot>(Slot) )
+		if ( UGridSlot* TypedSlot = Cast<UGridSlot>(PanelSlot) )
 		{
 			TypedSlot->Parent = this;
 			TypedSlot->BuildSlot( MyGridPanel.ToSharedRef() );
diff --git a/Engine/Source/Runtime/UMG/Private/Components/HorizontalBox.cpp b/Engine/Source/Runtime/UMG/Private/Components/HorizontalBox.cpp
--- a/Engine/Source/Runtime/UMG/Private/Components/HorizontalBox.cpp
+++ b/Engine/Source/Runtime/UMG/Private/Components/HorizontalBox.cpp
@@ -28,21 +28,21 @@ UClass* UHorizontalBox::GetSlotClass() const
 	return UHorizontalBoxSlot::StaticClass();
 }
 
-void UHorizontalBox::OnSlotAdded(UPanelSlot* Slot)
+void UHorizontalBox::OnSlotAdded(UPanelSlot* InSlot)
 {
 	// Add the child to the live canvas if it already exists
 	if ( MyHorizontalBox.IsValid() )
 	{
-		Cast<UHorizontalBoxSlot>(Slot)->BuildSlot(MyHorizontalBox.ToSharedRef());
+		CastChecked<UHorizontalBoxSlot>(InSlot)->BuildSlot(MyHorizontalBox.ToSharedRef());
 	}
 }
 
-void UHorizontalBox::OnSlotRemoved(UPanelSlot* Slot)
+void UHorizontalBox::OnSlotRemoved(UPanelSlot* InSlot)
 {
 	// Remove the widget from the live slot if it exists.
 	if ( MyHorizontalBox.IsValid() )
 	{
-		TSharedPtr<SWidget> Widget = Slot->Content->GetCachedWidget();
+		TSharedPtr<SWidget> Widget = InSlot->Content->GetCachedWidget();
 		if ( Widget.IsValid() )
 		{
 			MyHorizontalBox->RemoveSlot(Widget.ToSharedRef());
@@ -59,9 +59,9 @@ TSharedRef<SWidget> UHorizontalBox::RebuildWidget()
 {
 	MyHorizontalBox = SNew(SHorizontalBox);
 
-	for ( UPanelSlot* Slot : Slots )
+	for ( UPanelSlot* PanelSlot : Slots )
 	{
-		if ( UHorizontalBoxSlot* TypedSlot = Cast<UHorizontalBoxSlot>(Slot) )
+		if ( UHorizontalBoxSlot* TypedSlot = Cast<UHorizontalBoxSlot>(PanelSlot) )
 		{
 			TypedSlot->Parent = this;
 			TypedSlot->BuildSlot(MyHorizontalBox.ToSharedRef());
diff --git a/Engine/Source/Runtime/UMG/Private/Components/InvalidationBox.cpp b/Engine/Source/Runtime/UMG/Private/Components/InvalidationBox.cpp
--- a/Engine/Source/Runtime/UMG/Private/Components/InvalidationBox.cpp
+++ b/Engine/Source/Runtime/UMG/Private/Components/InvalidationBox.cpp
@@ -39,16 +39,16 @@ TSharedRef<SWidget> UInvalidationBox::RebuildWidget()
 	return BuildDesignTimeWidget(MyInvalidationPanel.ToSharedRef());
 }
 
-void UInvalidationBox::OnSlotAdded(UPanelSlot* Slot)
+void UInvalidationBox::OnSlotAdded(UPanelSlot* InSlot)
 {
 	// Add the child to the live slot if it already exists
 	if ( MyInvalidationPanel.IsValid() )
 	{
-		MyInvalidationPanel->SetContent(Slot->Content ? Slot->Content->TakeWidget() : SNullWidget::NullWidget);
+		MyInvalidationPanel->SetContent(InSlot->Content ? InSlot->Content->TakeWidget() : SNullWidget::NullWidget);
 	}
 }
 
-void UInvalidationBox::OnSlotRemoved(UPanelSlot* Slot)
+void UInvalidationBox::OnSlotRemoved(UPanelSlot* InSlot)
 {
 	// Remove the widget from the live slot if it exists.
 	if ( MyInvalidationPanel.IsValid() )
diff --git a/Engine/Source/Runtime/UMG/Private/Components/MenuAnchor.cpp b/Engine/Source/Runtime/UMG/Private/Components/MenuAnchor.cpp
--- a/Engine/Source/Runtime/UMG/Private/Components/MenuAnchor.cpp
+++ b/Engine/Source/Runtime/UMG/Private/Components/MenuAnchor.cpp
@@ -39,16 +39,16 @@ TSharedRef<SWidget> UMenuAnchor::RebuildWidget()
 	return BuildDesignTimeWidget( MyMenuAnchor.ToSharedRef() );
 }
 
-void UMenuAnchor::OnSlotAdded(UPanelSlot* Slot)
+void UMenuAnchor::OnSlotAdded(UPanelSlot* InSlot)
 {
 	// Add the child to the live slot if it already exists
 	if ( MyMenuAnchor.IsValid() )
 	{
-		MyMenuAnchor->SetContent(Slot->Content ? Slot->Content->TakeWidget() : SNullWidget::NullWidget);
+		MyMenuAnchor->SetContent(InSlot->Content ? InSlot->Content->TakeWidget() : SNullWidget::NullWidget);
 	}
 }
 
-void UMenuAnchor::OnSlotRemoved(UPanelSlot* Slot)
+void UMenuAnchor::OnSlotRemoved(UPanelSlot* InSlot)
 {
 	// Remove the widget from the live slot if it exists.
 	if ( MyMenuAnchor.IsValid() )
diff --git a/Engine/Source/Runtime/UMG/Private/Components/NamedSlot.cpp b/Engine/Source/Runtime/UMG/Private/Components/NamedSlot.cpp
--- a/Engine/Source/Runtime/UMG/Private/Components/NamedSlot.cpp
+++ b/Engine/Source/Runtime/UMG/Private/Components/NamedSlot.cpp
@@ -56,16 +56,16 @@ void UNamedSlot::SynchronizeProperties()
 	Super::SynchronizeProperties();
 }
 
-void UNamedSlot::OnSlotAdded(UPanelSlot* Slot)
+void UNamedSlot::OnSlotAdded(UPanelSlot* InSlot)
 {
 	// Add the child to the live slot if it already exists
-	if ( MyBox.IsValid() && Slot->Content )
+	if ( MyBox.IsValid() && InSlot->Content )
 	{
-		MyBox->SetContent(Slot->Content->TakeWidget());
+		MyBox->SetContent(InSlot->Content->TakeWidget());
 	}
 }
 
-void UNamedSlot::OnSlotRemoved(UPanelSlot* Slot)
+void UNamedSlot::OnSlotRemoved(UPanelSlot* InSlot)
 {
 	// Remove the widget from the live slot if it exists.
 	if ( MyBox.IsValid() )
diff --git a/Engine/Source/Runtime/UMG/Private/Components/Overlay.cpp b/Engine/Source/Runtime/UMG/Private/Components/Overlay.cpp
--- a/Engine/Source/Runtime/UMG/Private/Components/Overlay.cpp
+++ b/Engine/Source/Runtime/UMG/Private/Components/Overlay.cpp
@@ -33,21 +33,21 @@ UClass* UOverlay::GetSlotClass() const
 	return UOverlaySlot::StaticClass();
 }
 
-void UOverlay::OnSlotAdded(UPanelSlot* Slot)
+void UOverlay::OnSlotAdded(UPanelSlot* InSlot)
 {
 	// Add the child to the live canvas if it already exists
 	if ( MyOverlay.IsValid() )
 	{
-		Cast<UOverlaySlot>(Slot)->BuildSlot(MyOverlay.ToSharedRef());
+		CastChecked<UOverlaySlot>(InSlot)->BuildSlot(MyOverlay.ToSharedRef());
 	}
 }
 
-void UOverlay::OnSlotRemoved(UPanelSlot* Slot)
+void UOverlay::OnSlotRemoved(UPanelSlot* InSlot)
 {
 	// Remove the widget from the live slot if it exists.
 	if ( MyOverlay.IsValid() )
 	{
-		TSharedPtr<SWidget> Widget = Slot->Content->GetCachedWidget();
+		TSharedPtr<SWidget> Widget = InSlot->Content->GetCachedWidget();
 		if ( Widget.IsValid() )
 		{
 			MyOverlay->RemoveSlot(Widget.ToSharedRef());
@@ -59,9 +59,9 @@ TSharedRef<SWidget> UOverlay::RebuildWidget()
 {
 	MyOverlay = SNew(SOverlay);
 
-	for ( UPanelSlot* Slot : Slots )
+	for ( UPanelSlot* PanelSlot : Slots )
 	{
-		if ( UOverlaySlot* TypedSlot = Cast<UOverlaySlot>(Slot) )
+		if ( UOverlaySlot* TypedSlot = Cast<UOverlaySlot>(PanelSlot) )
 		{
 			TypedSlot->Parent = this;
 			TypedSlot->BuildSlot(MyOverlay.ToSharedRef());
diff --git a/Engine/Source/Runtime/UMG/Private/Components/PanelWidget.cpp b/Engine/Source/Runtime/UMG/Private/Components/PanelWidget.cpp
--- a/Engine/Source/Runtime/UMG/Private/Components/PanelWidget.cpp
+++ b/Engine/Source/Runtime/UMG/Private/Components/PanelWidget.cpp
@@ -79,21 +79,21 @@ bool UPanelWidget::RemoveChildAt(int32 Index)
 		return false;
 	}
 
-	UPanelSlot* Slot = Slots[Index];
-	if ( Slot->Content )
+	UPanelSlot* PanelSlot = Slots[Index];
+	if ( PanelSlot->Content )
 	{
-		Slot->Content->Slot = nullptr;
+		PanelSlot->Content->Slot = nullptr;
 	}
 
 	Slots.RemoveAt(Index);
 
-	OnSlotRemoved(Slot);
+	OnSlotRemoved(PanelSlot);
 
 	const bool bReleaseChildren = true;
-	Slot->ReleaseSlateResources(bReleaseChildren);
+	PanelSlot->ReleaseSlateResources(bReleaseChildren);
 
-	Slot->Parent = nullptr;
-	Slot->Content = nullptr;
+	PanelSlot->Parent = nullptr;
+	PanelSlot->Content = nullptr;
 
 	return true;
 }
@@ -112,21 +112,21 @@ UPanelSlot* UPanelWidget::AddChild(UWidget* Content)
 
 	Content->RemoveFromParent();
 
-	UPanelSlot* Slot = NewObject<UPanelSlot>(this, GetSlotClass());
-	Slot->SetFlags(RF_Transactional);
-	Slot->Content = Content;
-	Slot->Parent = this;
+	UPanelSlot* PanelSlot = NewObject<UPanelSlot>(this, GetSlotClass());
+	PanelSlot->SetFlags(RF_Transactional);
+	PanelSlot->Content = Content;
+	PanelSlot->Parent = this;
 
 	if ( Content )
 	{
-		Content->Slot = Slot;
+		Content->Slot = PanelSlot;
 	}
 
-	Slots.Add(Slot);
+	Slots.Add(PanelSlot);
 
-	OnSlotAdded(Slot);
+	OnSlotAdded(PanelSlot);
 
-	return Slot;
+	return PanelSlot;
 }
 
 bool UPanelWidget::ReplaceChildAt(int32 Index, UWidget* Content)
@@ -136,15 +136,15 @@ bool UPanelWidget::ReplaceChildAt(int32 Index, UWidget* Content)
 		return false;
 	}
 
-	UPanelSlot* Slot = Slots[Index];
-	Slot->Content = Content;
+	UPanelSlot* PanelSlot = Slots[Index];
+	PanelSlot->Content = Content;
 
 	if ( Content )
 	{
-		Content->Slot = Slot;
+		Content->Slot = PanelSlot;
 	}
 
-	Slot->SynchronizeProperties();
+	PanelSlot->SynchronizeProperties();
 
 	return true;
 }
diff --git a/Engine/Source/Runtime/UMG/Private/Components/RetainerBox.cpp b/Engine/Source/Runtime/UMG/Private/Components/RetainerBox.cpp
--- a/Engine/Source/Runtime/UMG/Private/Components/RetainerBox.cpp
+++ b/Engine/Source/Runtime/UMG/Private/Components/RetainerBox.cpp
@@ -85,16 +85,16 @@ void URetainerBox::SynchronizeProperties()
 	MyRetainerWidget->SetTextureParameter(TextureParameter);
 }
 
-void URetainerBox::OnSlotAdded(UPanelSlot* Slot)
+void URetainerBox::OnSlotAdded(UPanelSlot* InSlot)
 {
 	// Add the child to the live slot if it already exists
 	if ( MyRetainerWidget.IsValid() )
 	{
-		MyRetainerWidget->SetContent(Slot->Content ? Slot->Content->TakeWidget() : SNullWidget::NullWidget);
+		MyRetainerWidget->SetContent(InSlot->Content ? InSlot->Content->TakeWidget() : SNullWidget::NullWidget);
 	}
 }
 
-void URetainerBox::OnSlotRemoved(UPanelSlot* Slot)
+void URetainerBox::OnSlotRemoved(UPanelSlot* InSlot)
 {
 	// Remove the widget from the live slot if it exists.
 	if ( MyRetainerWidget.IsValid() )
diff --git a/Engine/Source/Runtime/UMG/Private/Components/SafeZone.cpp b/Engine/Source/Runtime/UMG/Private/Components/SafeZone.cpp
--- a/Engine/Source/Runtime/UMG/Private/Components/SafeZone.cpp
+++ b/Engine/Source/Runtime/UMG/Private/Components/SafeZone.cpp
@@ -22,16 +22,16 @@ const FText USafeZone::GetPaletteCategory()
 }
 #endif
 
-void USafeZone::OnSlotAdded( UPanelSlot* Slot )
+void USafeZone::OnSlotAdded( UPanelSlot* InSlot )
 {
-	Super::OnSlotAdded( Slot );
+	Super::OnSlotAdded( InSlot );
 
 	UpdateWidgetProperties();
 }
 
-void USafeZone::OnSlotRemoved( UPanelSlot* Slot )
+void USafeZone::OnSlotRemoved( UPanelSlot* InSlot )
 {
-	Super::OnSlotRemoved( Slot );
+	Super::OnSlotRemoved( InSlot );
 
 	if ( MySafeZone.IsValid() )
 	{
diff --git a/Engine/Source/Runtime/UMG/Private/Components/ScaleBox.cpp b/Engine/Source/Runtime/UMG/Private/Components/ScaleBox.cpp
--- a/Engine/Source/Runtime/UMG/Private/Components/ScaleBox.cpp
+++ b/Engine/Source/Runtime/UMG/Private/Components/ScaleBox.cpp
@@ -89,16 +89,16 @@ UClass* UScaleBox::GetSlotClass() const
 	return UScaleBoxSlot::StaticClass();
 }
 
-void UScaleBox::OnSlotAdded(UPanelSlot* Slot)
+void UScaleBox::OnSlotAdded(UPanelSlot* InSlot)
 {
 	// Add the child to the live slot if it already exists
 	if ( MyScaleBox.IsValid() )
 	{
-		Cast<UScaleBoxSlot>(Slot)->BuildSlot(MyScaleBox.ToSharedRef());
+		CastChecked<UScaleBoxSlot>(InSlot)->BuildSlot(MyScaleBox.ToSharedRef());
 	}
 }
 
-void UScaleBox::OnSlotRemoved(UPanelSlot* Slot)
+void UScaleBox::OnSlotRemoved(UPanelSlot* InSlot)
 {
 	// Remove the widget from the live slot if it exists.
 	if ( MyScaleBox.IsValid() )
diff --git a/Engine/Source/Runtime/UMG/Private/Components/ScrollBox.cpp b/Engine/Source/Runtime/UMG/Private/Components/ScrollBox.cpp
--- a/Engine/Source/Runtime/UMG/Private/Components/ScrollBox.cpp
+++ b/Engine/Source/Runtime/UMG/Private/Components/ScrollBox.cpp
@@ -39,21 +39,21 @@ UClass* UScrollBox::GetSlotClass() const
 	return UScrollBoxSlot::StaticClass();
 }
 
-void UScrollBox::OnSlotAdded(UPanelSlot* Slot)
+void UScrollBox::OnSlotAdded(UPanelSlot* InSlot)
 {
 	// Add the child to the live canvas if it already exists
 	if ( MyScrollBox.IsValid() )
 	{
-		Cast<UScrollBoxSlot>(Slot)->BuildSlot(MyScrollBox.ToSharedRef());
+		CastChecked<UScrollBoxSlot>(InSlot)->BuildSlot(MyScrollBox.ToSharedRef());
 	}
 }
 
-void UScrollBox::OnSlotRemoved(UPanelSlot* Slot)
+void UScrollBox::OnSlotRemoved(UPanelSlot* InSlot)
 {
 	// Remove the widget from the live slot if it exists.
 	if ( MyScrollBox.IsValid() )
 	{
-		TSharedPtr<SWidget> Widget = Slot->Content->GetCachedWidget();
+		TSharedPtr<SWidget> Widget = InSlot->Content->GetCachedWidget();
 		if ( Widget.IsValid() )
 		{
 			MyScrollBox->RemoveSlot(Widget.ToSharedRef());
@@ -69,9 +69,9 @@ TSharedRef<SWidget> UScrollBox::RebuildWidget()
 		.Orientation(Orientation)
 		.ConsumeMouseWheel(ConsumeMouseWheel);
 
-	for ( UPanelSlot* Slot : Slots )
+	for ( UPanelSlot* PanelSlot : Slots )
 	{
-		if ( UScrollBoxSlot* TypedSlot = Cast<UScrollBoxSlot>(Slot) )
+		if ( UScrollBoxSlot* TypedSlot = Cast<UScrollBoxSlot>(PanelSlot) )
 		{
 			TypedSlot->Parent = this;
 			TypedSlot->BuildSlot(MyScrollBox.ToSharedRef());
diff --git a/Engine/Source/Runtime/UMG/Private/Components/SizeBox.cpp b/Engine/Source/Runtime/UMG/Private/Components/SizeBox.cpp
--- a/Engine/Source/Runtime/UMG/Private/Components/SizeBox.cpp
+++ b/Engine/Source/Runtime/UMG/Private/Components/SizeBox.cpp
@@ -212,16 +212,16 @@ UClass* USizeBox::GetSlotClass() const
 	return USizeBoxSlot::StaticClass();
 }
 
-void USizeBox::OnSlotAdded(UPanelSlot* Slot)
+void USizeBox::OnSlotAdded(UPanelSlot* InSlot)
 {
 	// Add the child to the live slot if it already exists
 	if ( MySizeBox.IsValid() )
 	{
-		Cast<USizeBoxSlot>(Slot)->BuildSlot(MySizeBox.ToSharedRef());
+		CastChecked<USizeBoxSlot>(InSlot)->BuildSlot(MySizeBox.ToSharedRef());
 	}
 }
 
-void USizeBox::OnSlotRemoved(UPanelSlot* Slot)
+void USizeBox::OnSlotRemoved(UPanelSlot* InSlot)
 {
 	// Remove the widget from the live slot if it exists.
 	if ( MySizeBox.IsValid() )
diff --git a/Engine/Source/Runtime/UMG/Private/Components/UniformGridPanel.cpp b/Engine/Source/Runtime/UMG/Private/Components/UniformGridPanel.cpp
--- a/Engine/Source/Runtime/UMG/Private/Components/UniformGridPanel.cpp
+++ b/Engine/Source/Runtime/UMG/Private/Components/UniformGridPanel.cpp
@@ -28,21 +28,21 @@ UClass* UUniformGridPanel::GetSlotClass() const
 	return UUniformGridSlot::StaticClass();
 }
 
-void UUniformGridPanel::OnSlotAdded(UPanelSlot* Slot)
+void UUniformGridPanel::OnSlotAdded(UPanelSlot* InSlot)
 {
 	// Add the child to the live canvas if it already exists
 	if ( MyUniformGridPanel.IsValid() )
 	{
-		Cast<UUniformGridSlot>(Slot)->BuildSlot(MyUniformGridPanel.ToSharedRef());
+		CastChecked<UUniformGridSlot>(InSlot)->BuildSlot(MyUniformGridPanel.ToSharedRef());
 	}
 }
 
-void UUniformGridPanel::OnSlotRemoved(UPanelSlot* Slot)
+void UUniformGridPanel::OnSlotRemoved(UPanelSlot* InSlot)
 {
 	// Remove the widget from the live slot if it exists.
 	if ( MyUniformGridPanel.IsValid() )
 	{
-		TSharedPtr<SWidget> Widget = Slot->Content->GetCachedWidget();
+		TSharedPtr<SWidget> Widget = InSlot->Content->GetCachedWidget();
 		if ( Widget.IsValid() )
 		{
 			MyUniformGridPanel->RemoveSlot(Widget.ToSharedRef());
@@ -54,9 +54,9 @@ TSharedRef<SWidget> UUniformGridPanel::RebuildWidget()
 {
 	MyUniformGridPanel = SNew(SUniformGridPanel);
 
-	for ( UPanelSlot* Slot : Slots )
+	for ( UPanelSlot* PanelSlot : Slots )
 	{
-		if ( UUniformGridSlot* TypedSlot = Cast<UUniformGridSlot>(Slot) )
+		if ( UUniformGridSlot* TypedSlot = Cast<UUniformGridSlot>(PanelSlot) )
 		{
 			TypedSlot->Parent = this;
 			TypedSlot->BuildSlot(MyUniformGridPanel.ToSharedRef());
diff --git a/Engine/Source/Runtime/UMG/Private/Components/VerticalBox.cpp b/Engine/Source/Runtime/UMG/Private/Components/VerticalBox.cpp
--- a/Engine/Source/Runtime/UMG/Private/Components/VerticalBox.cpp
+++ b/Engine/Source/Runtime/UMG/Private/Components/VerticalBox.cpp
@@ -28,21 +28,21 @@ UClass* UVerticalBox::GetSlotClass() const
 	return UVerticalBoxSlot::StaticClass();
 }
 
-void UVerticalBox::OnSlotAdded(UPanelSlot* Slot)
+void UVerticalBox::OnSlotAdded(UPanelSlot* InSlot)
 {
 	// Add the child to the live canvas if it already exists
 	if ( MyVerticalBox.IsValid() )
 	{
-		Cast<UVerticalBoxSlot>(Slot)->BuildSlot(MyVerticalBox.ToSharedRef());
+		CastChecked<UVerticalBoxSlot>(InSlot)->BuildSlot(MyVerticalBox.ToSharedRef());
 	}
 }
 
-void UVerticalBox::OnSlotRemoved(UPanelSlot* Slot)
+void UVerticalBox::OnSlotRemoved(UPanelSlot* InSlot)
 {
 	// Remove the widget from the live slot if it exists.
 	if ( MyVerticalBox.IsValid() )
 	{
-		TSharedPtr<SWidget> Widget = Slot->Content->GetCachedWidget();
+		TSharedPtr<SWidget> Widget = InSlot->Content->GetCachedWidget();
 		if ( Widget.IsValid() )
 		{
 			MyVerticalBox->RemoveSlot(Widget.ToSharedRef());
@@ -59,9 +59,9 @@ TSharedRef<SWidget> UVerticalBox::RebuildWidget()
 {
 	MyVerticalBox = SNew(SVerticalBox);
 
-	for ( UPanelSlot* Slot : Slots )
+	for ( UPanelSlot* PanelSlot : Slots )
 	{
-		if ( UVerticalBoxSlot* TypedSlot = Cast<UVerticalBoxSlot>(Slot) )
+		if ( UVerticalBoxSlot* TypedSlot = Cast<UVerticalBoxSlot>(PanelSlot) )
 		{
 			TypedSlot->Parent = this;
 			TypedSlot->BuildSlot(MyVerticalBox.ToSharedRef());
diff --git a/Engine/Source/Runtime/UMG/Private/Components/Viewport.cpp b/Engine/Source/Runtime/UMG/Private/Components/Viewport.cpp
--- a/Engine/Source/Runtime/UMG/Private/Components/Viewport.cpp
+++ b/Engine/Source/Runtime/UMG/Private/Components/Viewport.cpp
@@ -401,16 +401,16 @@ void UViewport::SynchronizeProperties()
 	}
 }
 
-void UViewport::OnSlotAdded(UPanelSlot* Slot)
+void UViewport::OnSlotAdded(UPanelSlot* InSlot)
 {
 	// Add the child to the live canvas if it already exists
 	if ( ViewportWidget.IsValid() )
 	{
-		ViewportWidget->SetContent(Slot->Content ? Slot->Content->TakeWidget() : SNullWidget::NullWidget);
+		ViewportWidget->SetContent(InSlot->Content ? InSlot->Content->TakeWidget() : SNullWidget::NullWidget);
 	}
 }
 
-void UViewport::OnSlotRemoved(UPanelSlot* Slot)
+void UViewport::OnSlotRemoved(UPanelSlot* InSlot)
 {
 	// Remove the widget from the live slot if it exists.
 	if ( ViewportWidget.IsValid() )
diff --git a/Engine/Source/Runtime/UMG/Private/Components/WidgetComponent.cpp b/Engine/Source/Runtime/UMG/Private/Components/WidgetComponent.cpp
--- a/Engine/Source/Runtime/UMG/Private/Components/WidgetComponent.cpp
+++ b/Engine/Source/Runtime/UMG/Private/Components/WidgetComponent.cpp
@@ -138,7 +138,7 @@ public:
 		FMaterialRenderProxy* MaterialProxy = MaterialInstance->GetRenderProxy(IsSelected());
 #endif
 
-		const FMatrix& LocalToWorld = GetLocalToWorld();
+		const FMatrix& ViewportLocalToWorld = GetLocalToWorld();
 
 		if( RenderTarget )
 		{
@@ -176,7 +176,7 @@ public:
 						MeshBuilder.AddTriangle(VertexIndices[0], VertexIndices[1], VertexIndices[2]);
 						MeshBuilder.AddTriangle(VertexIndices[0], VertexIndices[2], VertexIndices[3]);
 
-						MeshBuilder.GetMesh(LocalToWorld, MaterialProxy, SDPG_World, false, true, ViewIndex, Collector);
+						MeshBuilder.GetMesh(ViewportLocalToWorld, MaterialProxy, SDPG_World, false, true, ViewIndex, Collector);
 					}
 				}
 			}
@@ -194,7 +194,7 @@ public:
 #endif
 	}
 
-	void RenderCollision(UBodySetup* InBodySetup, FMeshElementCollector& Collector, int32 ViewIndex, const FEngineShowFlags& EngineShowFlags, const FBoxSphereBounds& Bounds, bool bRenderInEditor) const
+	void RenderCollision(UBodySetup* InBodySetup, FMeshElementCollector& Collector, int32 ViewIndex, const FEngineShowFlags& EngineShowFlags, const FBoxSphereBounds& InBounds, bool bRenderInEditor) const
 	{
 		if ( InBodySetup )
 		{
@@ -1074,7 +1074,7 @@ void UWidgetComponent::UpdateWidget()
 
 void UWidgetComponent::UpdateRenderTarget()
 {
-	bool bRenderStateDirty = false;
+	bool bWidgetRenderStateDirty = false;
 	bool bClearColorChanged = false;
 
 	FLinearColor ActualBackgroundColor = BackgroundColor;
@@ -1093,7 +1093,7 @@ void UWidgetComponent::UpdateRenderTarget()
 			RenderTarget = NewObject<UTextureRenderTarget2D>(this);
 			RenderTarget->ClearColor = ActualBackgroundColor;
 
-			bClearColorChanged = bRenderStateDirty = true;
+			bClearColorChanged = bWidgetRenderStateDirty = true;
 
 			RenderTarget->InitCustomFormat(DrawSize.X, DrawSize.Y, PF_B8G8R8A8, false);
 
@@ -1106,17 +1106,17 @@ void UWidgetComponent::UpdateRenderTarget()
 			{
 				RenderTarget->InitCustomFormat(DrawSize.X, DrawSize.Y, PF_B8G8R8A8, false);
 				RenderTarget->UpdateResourceImmediate(false);
-				bRenderStateDirty = true;
+				bWidgetRenderStateDirty = true;
 			}
 
 			// Update the clear color
 			if ( RenderTarget->ClearColor != ActualBackgroundColor )
 			{
 				RenderTarget->ClearColor = ActualBackgroundColor;
-				bClearColorChanged = bRenderStateDirty = true;
+				bClearColorChanged = bWidgetRenderStateDirty = true;
 			}
 
-			if ( bRenderStateDirty )
+			if ( bWidgetRenderStateDirty )
 			{
 				RenderTarget->UpdateResource();
 			}
@@ -1137,10 +1137,10 @@ void UWidgetComponent::UpdateRenderTarget()
 		if ( MaterialInstance->GetScalarParameterValue(ParabolaDistortionName, CurrentParabolaValue) && CurrentParabolaValue != ParabolaDistortion )
 		{
 			MaterialInstance->SetScalarParameterValue(ParabolaDistortionName, ParabolaDistortion);
-			bRenderStateDirty = true;
+			bWidgetRenderStateDirty = true;
 		}
 
-		if ( bRenderStateDirty )
+		if ( bWidgetRenderStateDirty )
 		{
 			MarkRenderStateDirty();
 		}
diff --git a/Engine/Source/Runtime/UMG/Private/Components/WidgetSwitcher.cpp b/Engine/Source/Runtime/UMG/Private/Components/WidgetSwitcher.cpp
--- a/Engine/Source/Runtime/UMG/Private/Components/WidgetSwitcher.cpp
+++ b/Engine/Source/Runtime/UMG/Private/Components/WidgetSwitcher.cpp
@@ -81,21 +81,21 @@ UClass* UWidgetSwitcher::GetSlotClass() const
 	return UWidgetSwitcherSlot::StaticClass();
 }
 
-void UWidgetSwitcher::OnSlotAdded(UPanelSlot* Slot)
+void UWidgetSwitcher::OnSlotAdded(UPanelSlot* InSlot)
 {
 	// Add the child to the live canvas if it already exists
 	if ( MyWidgetSwitcher.IsValid() )
 	{
-		Cast<UWidgetSwitcherSlot>(Slot)->BuildSlot(MyWidgetSwitcher.ToSharedRef());
+		CastChecked<UWidgetSwitcherSlot>(InSlot)->BuildSlot(MyWidgetSwitcher.ToSharedRef());
 	}
 }
 
-void UWidgetSwitcher::OnSlotRemoved(UPanelSlot* Slot)
+void UWidgetSwitcher::OnSlotRemoved(UPanelSlot* InSlot)
 {
 	// Remove the widget from the live slot if it exists.
 	if ( MyWidgetSwitcher.IsValid() )
 	{
-		TSharedPtr<SWidget> Widget = Slot->Content->GetCachedWidget();
+		TSharedPtr<SWidget> Widget = InSlot->Content->GetCachedWidget();
 		if ( Widget.IsValid() )
 		{
 			MyWidgetSwitcher->RemoveSlot(Widget.ToSharedRef());
@@ -107,9 +107,9 @@ TSharedRef<SWidget> UWidgetSwitcher::RebuildWidget()
 {
 	MyWidgetSwitcher = SNew(SWidgetSwitcher);
 
-	for ( UPanelSlot* Slot : Slots )
+	for ( UPanelSlot* PanelSlot : Slots )
 	{
-		if ( UWidgetSwitcherSlot* TypedSlot = Cast<UWidgetSwitcherSlot>(Slot) )
+		if ( UWidgetSwitcherSlot* TypedSlot = Cast<UWidgetSwitcherSlot>(PanelSlot) )
 		{
 			TypedSlot->Parent = this;
 			TypedSlot->BuildSlot(MyWidgetSwitcher.ToSharedRef());
diff --git a/Engine/Source/Runtime/UMG/Private/Components/WrapBox.cpp b/Engine/Source/Runtime/UMG/Private/Components/WrapBox.cpp
--- a/Engine/Source/Runtime/UMG/Private/Components/WrapBox.cpp
+++ b/Engine/Source/Runtime/UMG/Private/Components/WrapBox.cpp
@@ -28,21 +28,21 @@ UClass* UWrapBox::GetSlotClass() const
 	return UWrapBoxSlot::StaticClass();
 }
 
-void UWrapBox::OnSlotAdded(UPanelSlot* Slot)
+void UWrapBox::OnSlotAdded(UPanelSlot* InSlot)
 {
 	// Add the child to the live canvas if it already exists
 	if ( MyWrapBox.IsValid() )
 	{
-		Cast<UWrapBoxSlot>(Slot)->BuildSlot(MyWrapBox.ToSharedRef());
+		CastChecked<UWrapBoxSlot>(InSlot)->BuildSlot(MyWrapBox.ToSharedRef());
 	}
 }
 
-void UWrapBox::OnSlotRemoved(UPanelSlot* Slot)
+void UWrapBox::OnSlotRemoved(UPanelSlot* InSlot)
 {
 	// Remove the widget from the live slot if it exists.
 	if ( MyWrapBox.IsValid() )
 	{
-		TSharedPtr<SWidget> Widget = Slot->Content->GetCachedWidget();
+		TSharedPtr<SWidget> Widget = InSlot->Content->GetCachedWidget();
 		if ( Widget.IsValid() )
 		{
 			MyWrapBox->RemoveSlot(Widget.ToSharedRef());
@@ -60,9 +60,9 @@ TSharedRef<SWidget> UWrapBox::RebuildWidget()
 	MyWrapBox = SNew(SWrapBox)
 		.UseAllottedWidth(true);
 
-	for ( UPanelSlot* Slot : Slots )
+	for ( UPanelSlot* PanelSlot : Slots )
 	{
-		if ( UWrapBoxSlot* TypedSlot = Cast<UWrapBoxSlot>(Slot) )
+		if ( UWrapBoxSlot* TypedSlot = Cast<UWrapBoxSlot>(PanelSlot) )
 		{
 			TypedSlot->Parent = this;
 			TypedSlot->BuildSlot(MyWrapBox.ToSharedRef());
diff --git a/Engine/Source/Runtime/UMG/Public/Components/PanelWidget.h b/Engine/Source/Runtime/UMG/Public/Components/PanelWidget.h
--- a/Engine/Source/Runtime/UMG/Public/Components/PanelWidget.h
+++ b/Engine/Source/Runtime/UMG/Public/Components/PanelWidget.h
@@ -124,12 +124,12 @@ public:
 
 	virtual void ConnectEditorData() override
 	{
-		for ( UPanelSlot* Slot : Slots )
+		for ( UPanelSlot* PanelSlot : Slots )
 		{
-			Slot->Parent = this;
-			if ( Slot->Content )
+			PanelSlot->Parent = this;
+			if ( PanelSlot->Content )
 			{
-				Slot->Content->Slot = Slot;
+				PanelSlot->Content->Slot = PanelSlot;
 			}
 		}
 	}
@@ -146,12 +146,12 @@ protected:
 		return UPanelSlot::StaticClass();
 	}
 
-	virtual void OnSlotAdded(UPanelSlot* Slot)
+	virtual void OnSlotAdded(UPanelSlot* InSlot)
 	{
 
 	}
 
-	virtual void OnSlotRemoved(UPanelSlot* Slot)
+	virtual void OnSlotRemoved(UPanelSlot* InSlot)
 	{
 
 	}
diff --git a/Engine/Source/Runtime/Windows/D3D11RHI/Private/D3D11Commands.cpp b/Engine/Source/Runtime/Windows/D3D11RHI/Private/D3D11Commands.cpp
--- a/Engine/Source/Runtime/Windows/D3D11RHI/Private/D3D11Commands.cpp
+++ b/Engine/Source/Runtime/Windows/D3D11RHI/Private/D3D11Commands.cpp
@@ -2269,9 +2269,9 @@ void FD3D11DynamicRHI::RHITransitionResources(EResourceTransitionAccess Transiti
 	}
 }
 
-void FD3D11DynamicRHI::RHITransitionResources(EResourceTransitionAccess TransitionType, EResourceTransitionPipeline TransitionPipeline, FUnorderedAccessViewRHIParamRef* InUAVs, int32 NumUAVs, FComputeFenceRHIParamRef WriteFence)
+void FD3D11DynamicRHI::RHITransitionResources(EResourceTransitionAccess TransitionType, EResourceTransitionPipeline TransitionPipeline, FUnorderedAccessViewRHIParamRef* InUAVs, int32 InNumUAVs, FComputeFenceRHIParamRef WriteFence)
 {
-	for (int32 i = 0; i < NumUAVs; ++i)
+	for (int32 i = 0; i < InNumUAVs; ++i)
 	{
 		if (InUAVs[i])
 		{
diff --git a/Engine/Source/Runtime/Windows/D3D12RHI/Private/D3D12Allocation.cpp b/Engine/Source/Runtime/Windows/D3D12RHI/Private/D3D12Allocation.cpp
--- a/Engine/Source/Runtime/Windows/D3D12RHI/Private/D3D12Allocation.cpp
+++ b/Engine/Source/Runtime/Windows/D3D12RHI/Private/D3D12Allocation.cpp
@@ -322,8 +322,6 @@ void FD3D12BuddyAllocator::CleanUpAllocations()
 
 void FD3D12BuddyAllocator::ReleaseAllResources()
 {
-	FD3D12ResourceBlockInfo* Block = nullptr;
-
 	for (FD3D12ResourceBlockInfo*& Block : DeferredDeletionQueue)
 	{
 		DeallocateInternal(Block);
diff --git a/Engine/Source/Runtime/Windows/D3D12RHI/Private/D3D12Commands.cpp b/Engine/Source/Runtime/Windows/D3D12RHI/Private/D3D12Commands.cpp
--- a/Engine/Source/Runtime/Windows/D3D12RHI/Private/D3D12Commands.cpp
+++ b/Engine/Source/Runtime/Windows/D3D12RHI/Private/D3D12Commands.cpp
@@ -332,12 +332,12 @@ void FD3D12CommandContext::RHITransitionResources(EResourceTransitionAccess Tran
 }
 
 
-void FD3D12CommandContext::RHITransitionResources(EResourceTransitionAccess TransitionType, EResourceTransitionPipeline TransitionPipeline, FUnorderedAccessViewRHIParamRef* InUAVs, int32 NumUAVs, FComputeFenceRHIParamRef WriteComputeFenceRHI)
+void FD3D12CommandContext::RHITransitionResources(EResourceTransitionAccess TransitionType, EResourceTransitionPipeline TransitionPipeline, FUnorderedAccessViewRHIParamRef* InUAVs, int32 InNumUAVs, FComputeFenceRHIParamRef WriteComputeFenceRHI)
 {
 	static IConsoleVariable* CVarShowTransitions = IConsoleManager::Get().FindConsoleVariable(TEXT("r.ProfileGPU.ShowTransitions"));
 	const bool bShowTransitionEvents = CVarShowTransitions->GetInt() != 0;
 
-	SCOPED_RHI_CONDITIONAL_DRAW_EVENTF(*this, RHITransitionResources, bShowTransitionEvents, TEXT("TransitionTo: %s: %i UAVs"), *FResourceTransitionUtility::ResourceTransitionAccessStrings[(int32)TransitionType], NumUAVs);
+	SCOPED_RHI_CONDITIONAL_DRAW_EVENTF(*this, RHITransitionResources, bShowTransitionEvents, TEXT("TransitionTo: %s: %i UAVs"), *FResourceTransitionUtility::ResourceTransitionAccessStrings[(int32)TransitionType], InNumUAVs);
 	const bool bUAVTransition = (TransitionType == EResourceTransitionAccess::EReadable) || (TransitionType == EResourceTransitionAccess::EWritable || TransitionType == EResourceTransitionAccess::ERWBarrier);
 	const bool bUAVBarrier = (TransitionType == EResourceTransitionAccess::ERWBarrier && TransitionPipeline == EResourceTransitionPipeline::EComputeToCompute);
 
@@ -386,7 +386,7 @@ void FD3D12CommandContext::RHITransitionResources(EResourceTransitionAccess Tran
 		uint32 BarrierCount = 0;
 		const uint32 MaxBarrierCount = 16;
 		D3D12_RESOURCE_BARRIER Barriers[MaxBarrierCount];
-		for (int32 i = 0; i < NumUAVs; ++i)
+		for (int32 i = 0; i < InNumUAVs; ++i)
 		{
 			if (InUAVs[i])
 			{
@@ -1416,15 +1416,15 @@ void FD3D12CommandContext::CommitNonComputeShaderConstants()
 
 	// MSFT: Seb: Do we need to support this none parallel case?
 	//if (GRHISupportsParallelRHIExecute)
-	FD3D12BoundShaderState* RESTRICT CurrentBoundShaderState = this->CurrentBoundShaderState.GetReference();
+	FD3D12BoundShaderState* RESTRICT CurrentBoundShaderStatePtr = CurrentBoundShaderState.GetReference();
 	//else
-	//	FD3D12BoundShaderState* RESTRICT CurrentBoundShaderState = (FD3D12BoundShaderState*)OwningRHI.BoundShaderStateHistory.GetLast();
+	//	FD3D12BoundShaderState* RESTRICT CurrentBoundShaderStatePtr = (FD3D12BoundShaderState*)OwningRHI.BoundShaderStateHistory.GetLast();
 
-	check(CurrentBoundShaderState);
+	check(CurrentBoundShaderStatePtr);
 
 	// Only set the constant buffer if this shader needs the global constant buffer bound
 	// Otherwise we will overwrite a different constant buffer
-	if (CurrentBoundShaderState->bShaderNeedsGlobalConstantBuffer[SF_Vertex])
+	if (CurrentBoundShaderStatePtr->bShaderNeedsGlobalConstantBuffer[SF_Vertex])
 	{
 		// Commit and bind vertex shader constants
 		for (uint32 i=0;i<MAX_CONSTANT_BUFFER_SLOTS; i++)
@@ -1440,7 +1440,7 @@ void FD3D12CommandContext::CommitNonComputeShaderConstants()
 	// is always reset whenever bUsingTessellation changes in SetBoundShaderState()
 	if (bUsingTessellation)
 	{
-		if (CurrentBoundShaderState->bShaderNeedsGlobalConstantBuffer[SF_Hull])
+		if (CurrentBoundShaderStatePtr->bShaderNeedsGlobalConstantBuffer[SF_Hull])
 		{
 			// Commit and bind hull shader constants
 			for (uint32 i=0;i<MAX_CONSTANT_BUFFER_SLOTS; i++)
@@ -1450,7 +1450,7 @@ void FD3D12CommandContext::CommitNonComputeShaderConstants()
 			}
 		}
 
-		if (CurrentBoundShaderState->bShaderNeedsGlobalConstantBuffer[SF_Domain])
+		if (CurrentBoundShaderStatePtr->bShaderNeedsGlobalConstantBuffer[SF_Domain])
 		{
 			// Commit and bind domain shader constants
 			for (uint32 i=0;i<MAX_CONSTANT_BUFFER_SLOTS; i++)
@@ -1461,7 +1461,7 @@ void FD3D12CommandContext::CommitNonComputeShaderConstants()
 		}
 	}
 
-	if (CurrentBoundShaderState->bShaderNeedsGlobalConstantBuffer[SF_Geometry])
+	if (CurrentBoundShaderStatePtr->bShaderNeedsGlobalConstantBuffer[SF_Geometry])
 	{
 		// Commit and bind geometry shader constants
 		for (uint32 i=0;i<MAX_CONSTANT_BUFFER_SLOTS; i++)
@@ -1471,7 +1471,7 @@ void FD3D12CommandContext::CommitNonComputeShaderConstants()
 		}
 	}
 
-	if (CurrentBoundShaderState->bShaderNeedsGlobalConstantBuffer[SF_Pixel])
+	if (CurrentBoundShaderStatePtr->bShaderNeedsGlobalConstantBuffer[SF_Pixel])
 	{
 		// Commit and bind pixel shader constants
 		for (uint32 i=0;i<MAX_CONSTANT_BUFFER_SLOTS; i++)
@@ -1571,29 +1571,29 @@ void FD3D12CommandContext::CommitGraphicsResourceTables()
 
 	// MSFT: Seb: Do we need this non parallel case now that context objects are always used?
 //#if PLATFORM_SUPPORTS_PARALLEL_RHI_EXECUTE
-	FD3D12BoundShaderState* RESTRICT CurrentBoundShaderState = this->CurrentBoundShaderState.GetReference();
+	FD3D12BoundShaderState* RESTRICT CurrentBoundShaderStatePtr = CurrentBoundShaderState.GetReference();
 //#else
-//	FD3D12BoundShaderState* RESTRICT CurrentBoundShaderState = (FD3D12BoundShaderState*)OwningRHI.BoundShaderStateHistory.GetLast();
+//	FD3D12BoundShaderState* RESTRICT CurrentBoundShaderStatePtr = (FD3D12BoundShaderState*)OwningRHI.BoundShaderStateHistory.GetLast();
 //#endif
-	check(CurrentBoundShaderState);
+	check(CurrentBoundShaderStatePtr);
 
-	if (auto* Shader = CurrentBoundShaderState->GetVertexShader())
+	if (auto* Shader = CurrentBoundShaderStatePtr->GetVertexShader())
 	{
 		SetResourcesFromTables(Shader);
 	}
-	if (auto* Shader = CurrentBoundShaderState->GetPixelShader())
+	if (auto* Shader = CurrentBoundShaderStatePtr->GetPixelShader())
 	{
 		SetResourcesFromTables(Shader);
 	}
-	if (auto* Shader = CurrentBoundShaderState->GetHullShader())
+	if (auto* Shader = CurrentBoundShaderStatePtr->GetHullShader())
 	{
 		SetResourcesFromTables(Shader);
 	}
-	if (auto* Shader = CurrentBoundShaderState->GetDomainShader())
+	if (auto* Shader = CurrentBoundShaderStatePtr->GetDomainShader())
 	{
 		SetResourcesFromTables(Shader);
 	}
-	if (auto* Shader = CurrentBoundShaderState->GetGeometryShader())
+	if (auto* Shader = CurrentBoundShaderStatePtr->GetGeometryShader())
 	{
 		SetResourcesFromTables(Shader);
 	}
diff --git a/Engine/Source/Runtime/Windows/D3D12RHI/Private/D3D12StateCache.cpp b/Engine/Source/Runtime/Windows/D3D12RHI/Private/D3D12StateCache.cpp
--- a/Engine/Source/Runtime/Windows/D3D12RHI/Private/D3D12StateCache.cpp
+++ b/Engine/Source/Runtime/Windows/D3D12RHI/Private/D3D12StateCache.cpp
@@ -1293,7 +1293,7 @@ void FD3D12PipelineStateCache::RebuildFromDiskCache()
 
 	uint32 NumGraphicsPSOs = DiskCaches[PSO_CACHE_GRAPHICS].GetNumPSOs();
 
-	for (uint32 i = 0; i < NumGraphicsPSOs; i++)
+	for (uint32 GraphicPSOIndex = 0; GraphicPSOIndex < NumGraphicsPSOs; GraphicPSOIndex++)
 	{
 		FD3D12LowLevelGraphicsPipelineStateDesc* GraphicsPSODesc = nullptr;
 		DiskCaches[PSO_CACHE_GRAPHICS].SetPointerAndAdvanceFilePosition((void**)&GraphicsPSODesc, sizeof(*GraphicsPSODesc));
@@ -1318,23 +1318,23 @@ void FD3D12PipelineStateCache::RebuildFromDiskCache()
 		if (PSODesc->InputLayout.NumElements)
 		{
 			DiskCaches[PSO_CACHE_GRAPHICS].SetPointerAndAdvanceFilePosition((void**)&PSODesc->InputLayout.pInputElementDescs, PSODesc->InputLayout.NumElements * sizeof(D3D12_INPUT_ELEMENT_DESC), true);
-			for (uint32 i = 0; i < PSODesc->InputLayout.NumElements; i++)
+			for (uint32 InputElementIndex = 0; InputElementIndex < PSODesc->InputLayout.NumElements; InputElementIndex++)
 			{
 				// Get the Sematic name string
 				uint32* stringLength = nullptr;
 				DiskCaches[PSO_CACHE_GRAPHICS].SetPointerAndAdvanceFilePosition((void**)&stringLength, sizeof(uint32));
-				DiskCaches[PSO_CACHE_GRAPHICS].SetPointerAndAdvanceFilePosition((void**)&PSODesc->InputLayout.pInputElementDescs[i].SemanticName, *stringLength, true);
+				DiskCaches[PSO_CACHE_GRAPHICS].SetPointerAndAdvanceFilePosition((void**)&PSODesc->InputLayout.pInputElementDescs[InputElementIndex].SemanticName, *stringLength, true);
 			}
 		}
 		if (PSODesc->StreamOutput.NumEntries)
 		{
 			DiskCaches[PSO_CACHE_GRAPHICS].SetPointerAndAdvanceFilePosition((void**)&PSODesc->StreamOutput.pSODeclaration, PSODesc->StreamOutput.NumEntries * sizeof(D3D12_SO_DECLARATION_ENTRY), true);
-			for (uint32 i = 0; i < PSODesc->StreamOutput.NumEntries; i++)
+			for (uint32 InputElementIndex = 0; InputElementIndex < PSODesc->StreamOutput.NumEntries; InputElementIndex++)
 			{
 				//Get the Sematic name string
 				uint32* stringLength = nullptr;
 				DiskCaches[PSO_CACHE_GRAPHICS].SetPointerAndAdvanceFilePosition((void**)&stringLength, sizeof(uint32));
-				DiskCaches[PSO_CACHE_GRAPHICS].SetPointerAndAdvanceFilePosition((void**)&PSODesc->StreamOutput.pSODeclaration[i].SemanticName, *stringLength, true);
+				DiskCaches[PSO_CACHE_GRAPHICS].SetPointerAndAdvanceFilePosition((void**)&PSODesc->StreamOutput.pSODeclaration[InputElementIndex].SemanticName, *stringLength, true);
 			}
 		}
 		if (PSODesc->StreamOutput.NumStrides)
@@ -1380,7 +1380,7 @@ void FD3D12PipelineStateCache::RebuildFromDiskCache()
 
 	uint32 NumComputePSOs = DiskCaches[PSO_CACHE_COMPUTE].GetNumPSOs();
 
-	for (uint32 i = 0; i < NumComputePSOs; i++)
+	for (uint32 ComputeIndex = 0; ComputeIndex < NumComputePSOs; ComputeIndex++)
 	{
 		FD3D12ComputePipelineStateDesc* ComputePSODesc = nullptr;
 		DiskCaches[PSO_CACHE_COMPUTE].SetPointerAndAdvanceFilePosition((void**)&ComputePSODesc, sizeof(*ComputePSODesc));
diff --git a/Engine/Source/Runtime/Windows/D3D12RHI/Private/D3D12UniformBuffer.cpp b/Engine/Source/Runtime/Windows/D3D12RHI/Private/D3D12UniformBuffer.cpp
--- a/Engine/Source/Runtime/Windows/D3D12RHI/Private/D3D12UniformBuffer.cpp
+++ b/Engine/Source/Runtime/Windows/D3D12RHI/Private/D3D12UniformBuffer.cpp
@@ -136,9 +136,9 @@ void FD3D12UniformBuffer::CacheResourcesInternal()
 {
 	uint32 Start = FPlatformTime::Cycles();
 
-	const FRHIUniformBufferLayout& Layout = GetLayout();
-	int32 NumResources = Layout.Resources.Num();
-	const uint8* RESTRICT ResourceTypes = Layout.Resources.GetData();
+	const FRHIUniformBufferLayout& MyLayout = GetLayout();
+	int32 NumResources = MyLayout.Resources.Num();
+	const uint8* RESTRICT ResourceTypes = MyLayout.Resources.GetData();
 	const TRefCountPtr<FRHIResource>* RESTRICT Resources = ResourceTable.GetData();
 	FResourcePair* RESTRICT RawResources = RawResourceTable.GetData();
 	float CurrentTime = FApp::GetCurrentTime();
diff --git a/Engine/Source/Runtime/Windows/D3D12RHI/Private/Windows/WindowsD3D12Device.cpp b/Engine/Source/Runtime/Windows/D3D12RHI/Private/Windows/WindowsD3D12Device.cpp
--- a/Engine/Source/Runtime/Windows/D3D12RHI/Private/Windows/WindowsD3D12Device.cpp
+++ b/Engine/Source/Runtime/Windows/D3D12RHI/Private/Windows/WindowsD3D12Device.cpp
@@ -400,11 +400,11 @@ void FD3D12DynamicRHI::PostInit()
 	}
 }
 
-void FD3D12DynamicRHI::PerRHISetup(FD3D12Device* MainDevice)
+void FD3D12DynamicRHI::PerRHISetup(FD3D12Device* InMainDevice)
 {
 	check(!GIsRHIInitialized);
 
-	DXGI_ADAPTER_DESC* AdapterDesc = MainDevice->GetD3DAdapterDesc();
+	DXGI_ADAPTER_DESC* AdapterDesc = InMainDevice->GetD3DAdapterDesc();
 
     GTexturePoolSize = 0;
 
@@ -434,7 +434,7 @@ void FD3D12DynamicRHI::PerRHISetup(FD3D12Device* MainDevice)
 	// Consider 50% of the shared memory but max 25% of total system memory.
 	int64 ConsideredSharedSystemMemory = FMath::Min(FD3D12GlobalStats::GSharedSystemMemory / 2ll, TotalPhysicalMemory / 4ll);
 
-	IDXGIAdapter3* DxgiAdapter3 = MainDevice->GetAdapter3();
+	IDXGIAdapter3* DxgiAdapter3 = InMainDevice->GetAdapter3();
 	DXGI_QUERY_VIDEO_MEMORY_INFO LocalVideoMemoryInfo;
 	VERIFYD3D11RESULT(DxgiAdapter3->QueryVideoMemoryInfo(0, DXGI_MEMORY_SEGMENT_GROUP_LOCAL, &LocalVideoMemoryInfo));
 	const int64 TargetBudget = LocalVideoMemoryInfo.Budget * 0.90f;	// Target using 90% of our budget to account for some fragmentation.
diff --git a/Engine/Source/Runtime/Windows/XAudio2/Private/XAudio2Device.cpp b/Engine/Source/Runtime/Windows/XAudio2/Private/XAudio2Device.cpp
--- a/Engine/Source/Runtime/Windows/XAudio2/Private/XAudio2Device.cpp
+++ b/Engine/Source/Runtime/Windows/XAudio2/Private/XAudio2Device.cpp
@@ -136,7 +136,7 @@ bool FXAudio2Device::InitializeHardware()
 	GConfig->GetString(TEXT("/Script/WindowsTargetPlatform.WindowsTargetSettings"), TEXT("AudioDevice"), WindowsAudioDeviceName, GEngineIni);
 
 	// Allow HMD to specify audio device, if one was not specified in settings
-	if (WindowsAudioDeviceName.IsEmpty() && IHeadMountedDisplayModule::IsAvailable())
+	if (WindowsAudioDeviceName.IsEmpty() && FAudioDevice::CanUseVRAudioDevice() && IHeadMountedDisplayModule::IsAvailable())
 	{
 		WindowsAudioDeviceName = IHeadMountedDisplayModule::Get().GetAudioOutputDevice();
 	}
