Copying //UE4/Dev-Framework to //UE4/Dev-Main (Source: //UE4/Dev-Framework @ 3021292)

diff --git a/Engine/Source/Runtime/Engine/Private/SkeletalRenderCPUSkin.cpp b/Engine/Source/Runtime/Engine/Private/SkeletalRenderCPUSkin.cpp
--- a/Engine/Source/Runtime/Engine/Private/SkeletalRenderCPUSkin.cpp
+++ b/Engine/Source/Runtime/Engine/Private/SkeletalRenderCPUSkin.cpp
@@ -216,7 +216,9 @@ void FSkeletalMeshObjectCPUSkin::Update(int32 LODIndex,USkinnedMeshComponent* In
 	if( GIsEditor )
 	{
 		// this does not need thread-safe update
+#if WITH_EDITORONLY_DATA
 		ProgressiveDrawingFraction = InMeshComponent->ProgressiveDrawingFraction;
+#endif // #if WITH_EDITORONLY_DATA
 		CustomSortAlternateIndexMode = (ECustomSortAlternateIndexMode)InMeshComponent->CustomSortAlternateIndexMode;
 	}
 }
diff --git a/Engine/Source/Runtime/Engine/Private/SkeletalRenderGPUSkin.cpp b/Engine/Source/Runtime/Engine/Private/SkeletalRenderGPUSkin.cpp
--- a/Engine/Source/Runtime/Engine/Private/SkeletalRenderGPUSkin.cpp
+++ b/Engine/Source/Runtime/Engine/Private/SkeletalRenderGPUSkin.cpp
@@ -273,7 +273,9 @@ void FSkeletalMeshObjectGPUSkin::Update(int32 LODIndex,USkinnedMeshComponent* In
 	if( GIsEditor )
 	{
 		// this does not need thread-safe update
+#if WITH_EDITORONLY_DATA
 		ProgressiveDrawingFraction = InMeshComponent->ProgressiveDrawingFraction;
+#endif // #if WITH_EDITORONLY_DATA
 		CustomSortAlternateIndexMode = (ECustomSortAlternateIndexMode)InMeshComponent->CustomSortAlternateIndexMode;
 	}
 }
diff --git a/Engine/Source/Runtime/Engine/Private/SoundNodeConcatenator.cpp b/Engine/Source/Runtime/Engine/Private/SoundNodeConcatenator.cpp
--- a/Engine/Source/Runtime/Engine/Private/SoundNodeConcatenator.cpp
+++ b/Engine/Source/Runtime/Engine/Private/SoundNodeConcatenator.cpp
@@ -42,9 +42,6 @@ bool USoundNodeConcatenator::NotifyWaveInstanceFinished( FWaveInstance* WaveInst
 	ConcatenatorPayload.CurrentChildNodeSoundIndex++;
 	if (ConcatenatorPayload.NodeIndex < ChildNodes.Num())
 	{
-		USoundNode* CurrentChildNode = ChildNodes[ConcatenatorPayload.NodeIndex];
-		check(CurrentChildNode);
-
 		if (ConcatenatorPayload.CurrentChildNodeSoundIndex == ConcatenatorPayload.CurrentChildNodeNumSound)
 		{
 			// Find the next non-null index
diff --git a/Engine/Source/Runtime/Engine/Private/UnrealEngine.cpp b/Engine/Source/Runtime/Engine/Private/UnrealEngine.cpp
--- a/Engine/Source/Runtime/Engine/Private/UnrealEngine.cpp
+++ b/Engine/Source/Runtime/Engine/Private/UnrealEngine.cpp
@@ -1500,6 +1500,31 @@ void UEngine::InitializeObjectReferences()
 	LoadSpecialMaterial(InvalidLightmapSettingsMaterialName.ToString(), InvalidLightmapSettingsMaterial, false);
 	LoadSpecialMaterial(ArrowMaterialName.ToString(), ArrowMaterial, false);
 
+#if !UE_BUILD_SHIPPING
+	LoadSpecialMaterial(TEXT("/Engine/EngineMaterials/PhAT_JointLimitMaterial.PhAT_JointLimitMaterial"), ConstraintLimitMaterial, false);
+
+	ConstraintLimitMaterialX = UMaterialInstanceDynamic::Create(ConstraintLimitMaterial, NULL);
+	ConstraintLimitMaterialX->SetVectorParameterValue(FName("Color"), FLinearColor::Red);
+	ConstraintLimitMaterialX->SetScalarParameterValue(FName("Desaturation"), 0.6f);
+	ConstraintLimitMaterialXAxis = UMaterialInstanceDynamic::Create(ConstraintLimitMaterial, NULL);
+	ConstraintLimitMaterialXAxis->SetVectorParameterValue(FName("Color"), FLinearColor::Red);
+
+	ConstraintLimitMaterialY = UMaterialInstanceDynamic::Create(ConstraintLimitMaterial, NULL);
+	ConstraintLimitMaterialY->SetVectorParameterValue(FName("Color"), FLinearColor::Green);
+	ConstraintLimitMaterialY->SetScalarParameterValue(FName("Desaturation"), 0.6f);
+	ConstraintLimitMaterialYAxis = UMaterialInstanceDynamic::Create(ConstraintLimitMaterial, NULL);
+	ConstraintLimitMaterialYAxis->SetVectorParameterValue(FName("Color"), FLinearColor::Green);
+
+	ConstraintLimitMaterialZ = UMaterialInstanceDynamic::Create(ConstraintLimitMaterial, NULL);
+	ConstraintLimitMaterialZ->SetVectorParameterValue(FName("Color"), FLinearColor::Blue);
+	ConstraintLimitMaterialZ->SetScalarParameterValue(FName("Desaturation"), 0.6f);
+	ConstraintLimitMaterialZAxis = UMaterialInstanceDynamic::Create(ConstraintLimitMaterial, NULL);
+	ConstraintLimitMaterialZAxis->SetVectorParameterValue(FName("Color"), FLinearColor::Blue);
+
+	ConstraintLimitMaterialPrismatic = UMaterialInstanceDynamic::Create(ConstraintLimitMaterial, NULL);
+	ConstraintLimitMaterialPrismatic->SetVectorParameterValue(FName("Color"), FLinearColor(FColor::Orange));
+#endif
+
 
 	if (GIsEditor && !IsRunningCommandlet())
 	{
@@ -1511,17 +1536,7 @@ void UEngine::InitializeObjectReferences()
 #endif
 
 		LoadSpecialMaterial(PreviewShadowsIndicatorMaterialName.ToString(), PreviewShadowsIndicatorMaterial, false);
-		LoadSpecialMaterial(ConstraintLimitMaterialName.ToString(), ConstraintLimitMaterial, false);
-
-		ConstraintLimitMaterialX = UMaterialInstanceDynamic::Create(ConstraintLimitMaterial, NULL);
-		ConstraintLimitMaterialX->SetVectorParameterValue(FName("Color"), FLinearColor::Red);
-		ConstraintLimitMaterialY = UMaterialInstanceDynamic::Create(ConstraintLimitMaterial, NULL);
-		ConstraintLimitMaterialY->SetVectorParameterValue(FName("Color"), FLinearColor::Green);
-		ConstraintLimitMaterialZ = UMaterialInstanceDynamic::Create(ConstraintLimitMaterial, NULL);
-		ConstraintLimitMaterialZ->SetVectorParameterValue(FName("Color"), FLinearColor::Blue);
-		ConstraintLimitMaterialPrismatic = UMaterialInstanceDynamic::Create(ConstraintLimitMaterial, NULL);
-		ConstraintLimitMaterialPrismatic->SetVectorParameterValue(FName("Color"), FLinearColor(FColor::Orange));
-
+		
 		//@TODO: This should move into the editor (used in editor modes exclusively)
 		if (DefaultBSPVertexTexture == NULL)
 		{
diff --git a/Engine/Source/Runtime/Engine/Public/AdvancedPreviewScene.h b/Engine/Source/Runtime/Engine/Public/AdvancedPreviewScene.h
new file mode 100644
--- /dev/null
+++ b/Engine/Source/Runtime/Engine/Public/AdvancedPreviewScene.h
@@ -0,0 +1,63 @@
+// Copyright 1998-2016 Epic Games, Inc. All Rights Reserved.
+
+/*=============================================================================
+	PrettyPreviewScene.h: Pretty preview scene definitions.
+=============================================================================*/
+
+#pragma once
+
+#if WITH_EDITOR
+
+#include "PreviewScene.h"
+
+#include "TickableEditorObject.h"
+#include "EditorSupportDelegates.h"
+
+struct FPreviewSceneProfile;
+
+class USphereReflectionCaptureComponent;
+class USkyLightComponent;
+class UStaticMeshComponent;
+class UMaterialInstanceConstant;
+class UPostProcessComponent;
+class UAssetViewerSettings;
+class UEditorPerProjectUserSettings;
+
+class ENGINE_API FAdvancedPreviewScene : public FPreviewScene, public FTickableEditorObject
+{
+public:
+	FAdvancedPreviewScene(ConstructionValues CVS, float InFloorOffset = 0.0f);
+
+	void UpdateScene(FPreviewSceneProfile& Profile, bool bUpdateSkyLight = true, bool bUpdateEnvironment = true, bool bUpdatePostProcessing = true, bool bUpdateDirectionalLight = true);
+
+	void SetFloorOffset(const float InFloorOffset);
+	void SetProfileIndex(const int32 InProfileIndex);
+
+	/* Begin FTickableEditorObject */
+	virtual void Tick(float DeltaTime) override;
+	virtual bool IsTickable() const override;
+	virtual TStatId GetStatId() const override;
+	/* End FTickableEditorObject */
+
+	const bool HandleViewportInput(FViewport* InViewport, int32 ControllerId, FKey Key, float Delta, float DeltaTime, int32 NumSamples, bool bGamepad);
+
+	const float GetSkyRotation() const;
+	void SetSkyRotation(const float SkyRotation);	
+
+	const UStaticMeshComponent* GetFloorMeshComponent() const;
+protected:
+	USkyLightComponent* SkyLightComponent;
+	UStaticMeshComponent* SkyComponent;
+	UMaterialInstanceConstant* InstancedSkyMaterial;
+	UPostProcessComponent* PostProcessComponent;
+	UStaticMeshComponent* FloorMeshComponent;
+	UAssetViewerSettings* DefaultSettings;
+	bool bRotateLighting;
+
+	float CurrentRotationSpeed;
+	float PreviousRotation;
+
+	int32 CurrentProfileIndex;
+};
+
+#endif // WITH_EDITOR
diff --git a/Engine/Source/Runtime/Engine/Public/Engine.h b/Engine/Source/Runtime/Engine/Public/Engine.h
--- a/Engine/Source/Runtime/Engine/Public/Engine.h
+++ b/Engine/Source/Runtime/Engine/Public/Engine.h
@@ -52,7 +52,6 @@
 #include "Engine/SubDSurface.h"
 #include "Engine/DebugCameraController.h"
 #include "Engine/DecalActor.h"
-#include "PhysicsEngine/DestructibleActor.h"
 #include "Engine/Canvas.h"
 #include "GameFramework/HUD.h"
 #include "Atmosphere/AtmosphericFog.h"
diff --git a/Engine/Source/Runtime/Engine/Public/FinalPostProcessSettings.h b/Engine/Source/Runtime/Engine/Public/FinalPostProcessSettings.h
--- a/Engine/Source/Runtime/Engine/Public/FinalPostProcessSettings.h
+++ b/Engine/Source/Runtime/Engine/Public/FinalPostProcessSettings.h
@@ -48,8 +48,8 @@ public:
 	void UpdateEntry(const FCubemapEntry &Entry, float Weight)
 	{
 		bool Existing = false;
-		// Note: make sure to read the array size each loop here, since we remove elements.
-		for(int32 i = 0; i < ContributingCubemaps.Num(); ++i)
+		// Always verify the index is valid since elements can be removed!
+		for(int32 i = 0; ContributingCubemaps.IsValidIndex(i); ++i)
 		{
 			FCubemapEntry& Local = ContributingCubemaps[i];
 			
@@ -65,8 +65,9 @@ public:
 
 			if(Local.AmbientCubemapTintMulScaleValue.IsAlmostBlack())
 			{
-				ContributingCubemaps.RemoveAt(i, /*bAllowShrinking=*/ false);
+				ContributingCubemaps.RemoveAt(i, 1, /*bAllowShrinking=*/ false);
 				i--; // Maintain same index in the loop after loop increment since we removed an element.
+				continue; // Not strictly necessary but protect against future code using i incorrectly.
 			}
 		}
 
diff --git a/Engine/Source/Runtime/Engine/Public/SceneManagement.h b/Engine/Source/Runtime/Engine/Public/SceneManagement.h
--- a/Engine/Source/Runtime/Engine/Public/SceneManagement.h
+++ b/Engine/Source/Runtime/Engine/Public/SceneManagement.h
@@ -1846,6 +1846,9 @@ extern ENGINE_API void DrawCylinder(class FPrimitiveDrawInterface* PDI,const FVe
 
 extern ENGINE_API void DrawCylinder(class FPrimitiveDrawInterface* PDI, const FMatrix& CylToWorld, const FVector& Base, const FVector& XAxis, const FVector& YAxis, const FVector& ZAxis,
 	float Radius, float HalfHeight, int32 Sides, const FMaterialRenderProxy* MaterialInstance, uint8 DepthPriority);
+//Draws a cylinder along the axis from Start to End
+extern ENGINE_API void DrawCylinder(class FPrimitiveDrawInterface* PDI, const FVector& Start, const FVector& End, float Radius, int32 Sides, const FMaterialRenderProxy* MaterialInstance, uint8 DepthPriority);
+
 
 extern ENGINE_API void GetBoxMesh(const FMatrix& BoxToWorld,const FVector& Radii,const FMaterialRenderProxy* MaterialRenderProxy,uint8 DepthPriority,int32 ViewIndex,FMeshElementCollector& Collector);
 extern ENGINE_API void GetHalfSphereMesh(const FVector& Center, const FVector& Radii, int32 NumSides, int32 NumRings, float StartAngle, float EndAngle, const FMaterialRenderProxy* MaterialRenderProxy, uint8 DepthPriority, bool bDisableBackfaceCulling, 
@@ -1858,6 +1861,9 @@ extern ENGINE_API void GetCylinderMesh(const FVector& Base, const FVector& XAxis
 									float Radius, float HalfHeight, int32 Sides, const FMaterialRenderProxy* MaterialInstance, uint8 DepthPriority, int32 ViewIndex, FMeshElementCollector& Collector);
 extern ENGINE_API void GetCylinderMesh(const FMatrix& CylToWorld, const FVector& Base, const FVector& XAxis, const FVector& YAxis, const FVector& ZAxis,
 									float Radius, float HalfHeight, int32 Sides, const FMaterialRenderProxy* MaterialInstance, uint8 DepthPriority, int32 ViewIndex, FMeshElementCollector& Collector);
+//Draws a cylinder along the axis from Start to End
+extern ENGINE_API void GetCylinderMesh(const FVector& Start, const FVector& End, float Radius, int32 Sides, const FMaterialRenderProxy* MaterialInstance, uint8 DepthPriority, int32 ViewIndex, FMeshElementCollector& Collector);
+
 extern ENGINE_API void GetConeMesh(const FMatrix& LocalToWorld, float AngleWidth, float AngleHeight, int32 NumSides,
 									const FMaterialRenderProxy* MaterialRenderProxy, uint8 DepthPriority, int32 ViewIndex, FMeshElementCollector& Collector);
 extern ENGINE_API void GetCapsuleMesh(const FVector& Origin, const FVector& XAxis, const FVector& YAxis, const FVector& ZAxis, const FLinearColor& Color, float Radius, float HalfHeight, int32 NumSides,
diff --git a/Engine/Source/Runtime/Windows/XAudio2/Private/XAudio2Source.cpp b/Engine/Source/Runtime/Windows/XAudio2/Private/XAudio2Source.cpp
--- a/Engine/Source/Runtime/Windows/XAudio2/Private/XAudio2Source.cpp
+++ b/Engine/Source/Runtime/Windows/XAudio2/Private/XAudio2Source.cpp
@@ -392,7 +392,6 @@ bool FXAudio2SoundSource::CreateSource( void )
 
 	uint32 NumSends = 0;
 
-#if XAUDIO2_SUPPORTS_SENDLIST
 	// Create a source that goes to the spatialisation code and reverb effect
 	Destinations[NumSends].pOutputVoice = Effects->DryPremasterVoice;
 	if( IsEQFilterApplied() )
@@ -417,7 +416,6 @@ bool FXAudio2SoundSource::CreateSource( void )
 	{
 		NumSends, Destinations
 	};
-#endif	//XAUDIO2_SUPPORTS_SENDLIST
 
 	// Mark the source as music if it is a member of the music group and allow low, band and high pass filters
 
@@ -448,17 +446,13 @@ bool FXAudio2SoundSource::CreateSource( void )
 			EffectDescriptor[0].InitialState = true;
 
 			const XAUDIO2_EFFECT_CHAIN EffectChain = { 1, EffectDescriptor };
-			AudioDevice->DeviceProperties->GetFreeSourceVoice(&Source, XAudio2Buffer->PCM, &EffectChain, MaxEffectChainChannels);
+			AudioDevice->DeviceProperties->GetFreeSourceVoice(&Source, XAudio2Buffer->PCM, &EffectChain, &SourceSendList, MaxEffectChainChannels);
 
 			if (!Source)
 			{
 				return false;
 			}
 
-#if XAUDIO2_SUPPORTS_SENDLIST
-			Source->SetOutputVoices(&SourceSendList);
-#endif
-
 			// We succeeded, then return
 			bCreatedWithSpatializationEffect = true;
 		}
@@ -469,16 +463,12 @@ bool FXAudio2SoundSource::CreateSource( void )
 		check(AudioDevice->DeviceProperties != nullptr);
 		check(AudioDevice->DeviceProperties->XAudio2 != nullptr);
 
-		AudioDevice->DeviceProperties->GetFreeSourceVoice(&Source, XAudio2Buffer->PCM, nullptr);
+		AudioDevice->DeviceProperties->GetFreeSourceVoice(&Source, XAudio2Buffer->PCM, nullptr, &SourceSendList);
 
 		if (!Source)
 		{
 			return false;
 		}
-
-#if XAUDIO2_SUPPORTS_SENDLIST
-		Source->SetOutputVoices(&SourceSendList);
-#endif
 	}
 
 	return true;
diff --git a/Engine/Source/Runtime/Windows/XAudio2/Private/XAudio2Support.h b/Engine/Source/Runtime/Windows/XAudio2/Private/XAudio2Support.h
--- a/Engine/Source/Runtime/Windows/XAudio2/Private/XAudio2Support.h
+++ b/Engine/Source/Runtime/Windows/XAudio2/Private/XAudio2Support.h
@@ -18,13 +18,6 @@
 #ifndef X3DAUDIO_VECTOR_IS_A_D3DVECTOR
 	#define X3DAUDIO_VECTOR_IS_A_D3DVECTOR		1
 #endif	//X3DAUDIO_VECTOR_IS_A_D3DVECTOR
-#ifndef XAUDIO2_SUPPORTS_SENDLIST
-	#define XAUDIO2_SUPPORTS_SENDLIST			1
-#endif	//XAUDIO2_SUPPORTS_SENDLIST
-#ifndef XAUDIO2_SUPPORTS_VOICE_POOL
-	#define XAUDIO2_SUPPORTS_VOICE_POOL			0
-#endif	//XAUDIO2_SUPPORTS_VOICE_POOL
-
 
 
 /*------------------------------------------------------------------------------------
@@ -772,47 +765,13 @@ struct FXAudioDeviceProperties
 	}
 
 	/** Returns either a new IXAudio2SourceVoice or a recycled IXAudio2SourceVoice according to the sound format and max channel count in the voice's effect chain*/
-	void GetFreeSourceVoice(IXAudio2SourceVoice** Voice, const FPCMBufferInfo& BufferInfo, const XAUDIO2_EFFECT_CHAIN* EffectChain = nullptr, int32 MaxEffectChainChannels = 0)
+	void GetFreeSourceVoice(IXAudio2SourceVoice** Voice, const FPCMBufferInfo& BufferInfo, const XAUDIO2_EFFECT_CHAIN* EffectChain = nullptr, const XAUDIO2_VOICE_SENDS* SendList = nullptr, int32 MaxEffectChainChannels = 0)
 	{
 		bool bSuccess = false;
 
-#if XAUDIO2_SUPPORTS_VOICE_POOL
-		// First find the pool for the given format
-		FSourceVoicePoolEntry* VoicePoolEntry = nullptr;
-		for (int32 i = 0; i < VoicePool.Num(); ++i)
-		{
-			if (VoicePool[i]->Format == BufferInfo.PCMFormat && VoicePool[i]->MaxEffectChainChannels == MaxEffectChainChannels)
-			{
-				VoicePoolEntry = VoicePool[i];
-				check(VoicePoolEntry);
-				bSuccess = true;
-				break;
-			}
-		}
-
-		// If we found a voice pool entry for this format and we have free voices
-		// then use the voice
-		if (VoicePoolEntry && VoicePoolEntry->FreeVoices.Num() > 0)
-		{
-			*Voice = VoicePoolEntry->FreeVoices.Pop(false);
-			check(*Voice);
-
-			bSuccess = Validate(TEXT("GetFreeSourceVoice, XAudio2->CreateSourceVoice"),
-								(*Voice)->SetEffectChain(EffectChain));
-		}
-		else
-		{
-			// Create a brand new source voice with this format.
-			check(XAudio2 != nullptr);
-			bSuccess = Validate(TEXT("GetFreeSourceVoice, XAudio2->CreateSourceVoice"),
-								XAudio2->CreateSourceVoice(Voice, &BufferInfo.PCMFormat, XAUDIO2_VOICE_USEFILTER, MAX_PITCH, &SourceCallback, nullptr, EffectChain));
-		}
-#else // XAUDIO2_SUPPORTS_VOICE_POOL
 		check(XAudio2 != nullptr);
 		bSuccess = Validate(TEXT("GetFreeSourceVoice, XAudio2->CreateSourceVoice"),
-							XAudio2->CreateSourceVoice(Voice, &BufferInfo.PCMFormat, XAUDIO2_VOICE_USEFILTER, MAX_PITCH, &SourceCallback, nullptr, EffectChain));
-#endif // XAUDIO2_SUPPORTS_VOICE_POOL
-
+							XAudio2->CreateSourceVoice(Voice, &BufferInfo.PCMFormat, XAUDIO2_VOICE_USEFILTER, MAX_PITCH, &SourceCallback, SendList, EffectChain));
 		if (bSuccess)
 		{
 			// Track the number of source voices out in the world
@@ -828,61 +787,7 @@ struct FXAudioDeviceProperties
 	/** Releases the voice into a pool of free voices according to the voice format and the max effect chain channels */
 	void ReleaseSourceVoice(IXAudio2SourceVoice* Voice, const FPCMBufferInfo& BufferInfo, const int32 MaxEffectChainChannels)
 	{
-#if XAUDIO2_SUPPORTS_VOICE_POOL
-
-		check(Voice != nullptr);
-
-		// Make sure the voice is stopped
-		Validate(TEXT("ReleaseSourceVoice, Voice->Stop()"), Voice->Stop());
-
-		// And make sure there's no audio remaining the voice so when it's re-used it's fresh.
-		Validate(TEXT("ReleaseSourceVoice, Voice->FlushSourceBuffers()"), Voice->FlushSourceBuffers());
-
-#if XAUDIO2_SUPPORTS_SENDLIST
-		// Clear out the send effects (OutputVoices). When the voice gets reused, the old internal state might be invalid 
-		// when the new send effects are applied to the voice.
-		Validate(TEXT("ReleaseSourceVoice, Voice->SetOutputVoices(nullptr)"), Voice->SetOutputVoices(nullptr));
-#endif
-
-		// Release the effect chain
-		Validate(TEXT("ReleaseSourceVoice, Voice->SetEffectChain(nullptr);"), Voice->SetEffectChain(nullptr));
-
-		// See if there is an existing pool for this source voice
-		FSourceVoicePoolEntry* VoicePoolEntry = nullptr;
-		for (int32 i = 0; i < VoicePool.Num(); ++i)
-		{
-			if (VoicePool[i]->Format == BufferInfo.PCMFormat && VoicePool[i]->MaxEffectChainChannels == MaxEffectChainChannels)
-			{
-				VoicePoolEntry = VoicePool[i];
-				break;
-			}
-		}
-
-		// If we found a voice pool entry for this voice, add the voice to the list of free voices
-		if (VoicePoolEntry)
-		{
-			VoicePoolEntry->FreeVoices.Add(Voice);
-		}
-		else
-		{
-			// Otherwise We need to make a new voice pool entry with this format and max effect chain channels
-			VoicePoolEntry = new FSourceVoicePoolEntry();
-			if (VoicePoolEntry)
-			{
-				VoicePoolEntry->Format = BufferInfo.PCMFormat;
-				VoicePoolEntry->FreeVoices.Add(Voice);
-				VoicePoolEntry->MaxEffectChainChannels = MaxEffectChainChannels;
-				VoicePool.Add(VoicePoolEntry);
-			}
-			else
-			{
-				// If we failed to create a new voice pool entry, then destroy the voice
-				Voice->DestroyVoice();
-			}
-		}
-#else // XAUDIO2_SUPPORTS_VOICE_POOL
 		Voice->DestroyVoice();
-#endif // XAUDIO2_SUPPORTS_VOICE_POOL
 
 		--NumActiveVoices;
 	}
