Copying //UE4/Dev-Rendering to //UE4/Dev-Main (Source: //UE4/Dev-Rendering @ 3072736)

diff --git a/Engine/Source/Runtime/ShaderCore/Private/Shader.cpp b/Engine/Source/Runtime/ShaderCore/Private/Shader.cpp
--- a/Engine/Source/Runtime/ShaderCore/Private/Shader.cpp
+++ b/Engine/Source/Runtime/ShaderCore/Private/Shader.cpp
@@ -775,6 +775,7 @@ FArchive& operator<<(FArchive& Ar,class FSelfContainedShaderId& Ref)
  * This still needs to initialize members to safe values since FShaderType::GenerateSerializationHistory uses this constructor.
  */
 FShader::FShader() : 
+	SerializedResource(nullptr),
 	ShaderPipeline(nullptr),
 	VFType(nullptr),
 	Type(nullptr), 
@@ -1770,4 +1771,12 @@ void ShaderMapAppendKeyString(EShaderPlatform Platform, FString& KeyString)
 			KeyString += TEXT("_FS");
 		}
 	}
+
+	{
+		static const auto CVar = IConsoleManager::Get().FindTConsoleVariableDataInt(TEXT("r.VertexFoggingForOpaque"));
+		if (CVar && CVar->GetValueOnAnyThread() > 0)
+		{
+			KeyString += TEXT("_VFO");
+		}
+	}
 }
diff --git a/Engine/Source/Runtime/Windows/D3D11RHI/Private/Windows/WindowsD3D11Device.cpp b/Engine/Source/Runtime/Windows/D3D11RHI/Private/Windows/WindowsD3D11Device.cpp
--- a/Engine/Source/Runtime/Windows/D3D11RHI/Private/Windows/WindowsD3D11Device.cpp
+++ b/Engine/Source/Runtime/Windows/D3D11RHI/Private/Windows/WindowsD3D11Device.cpp
@@ -47,14 +47,14 @@ static TAutoConsoleVariable<int32> CVarForceNvidiaToSM4(
 
 static TAutoConsoleVariable<int32> CVarAMDUseMultiThreadedDevice(
 	TEXT("r.AMDD3D11MultiThreadedDevice"),
-	1,
+	0,
 	TEXT("If true, creates a multithreaded D3D11 device on AMD hardware (workaround for driver bug)\n")
 	TEXT("Changes will only take effect in new game/editor instances - can't be changed at runtime.\n"),
 	ECVF_Default);
 
 static TAutoConsoleVariable<int32> CVarAMDDisableAsyncTextureCreation(
 	TEXT("r.AMDDisableAsyncTextureCreation"),
-	1,
+	0,
 	TEXT("If true, uses synchronous texture creation on AMD hardware (workaround for driver bug)\n")
 	TEXT("Changes will only take effect in new game/editor instances - can't be changed at runtime.\n"),
 	ECVF_Default);
@@ -292,8 +292,7 @@ void FD3D11DynamicRHIModule::FindAdapter()
 	for(uint32 AdapterIndex = 0; DXGIFactory1->EnumAdapters(AdapterIndex,TempAdapter.GetInitReference()) != DXGI_ERROR_NOT_FOUND; ++AdapterIndex)
 	{
 		// to make sure the array elements can be indexed with AdapterIndex
-		DXGI_ADAPTER_DESC AdapterDesc;
-		ZeroMemory(&AdapterDesc, sizeof(DXGI_ADAPTER_DESC));
+		DXGI_ADAPTER_DESC& AdapterDesc = AdapterDescription[AdapterDescription.AddZeroed()];
 
 		// Check that if adapter supports D3D11.
 		if(TempAdapter)
@@ -364,8 +363,6 @@ void FD3D11DynamicRHIModule::FindAdapter()
 				}
 			}
 		}
-		
-		AdapterDescription.Add(AdapterDesc);
 	}
 
 	if(bFavorNonIntegrated && (bIsAnyAMD || bIsAnyNVIDIA))
