Copying //UE4/Dev-Framework to //UE4/Dev-Main (Source: //UE4/Dev-Framework @ 3136612)

diff --git a/Engine/Source/Runtime/Engine/Public/Slate/SceneViewport.h b/Engine/Source/Runtime/Engine/Public/Slate/SceneViewport.h
--- a/Engine/Source/Runtime/Engine/Public/Slate/SceneViewport.h
+++ b/Engine/Source/Runtime/Engine/Public/Slate/SceneViewport.h
@@ -220,6 +220,7 @@ public:
 	virtual FReply OnFocusReceived( const FFocusEvent& InFocusEvent ) override;
 	virtual void OnFocusLost( const FFocusEvent& InFocusEvent ) override;
 	virtual void OnViewportClosed() override;
+	virtual FReply OnRequestWindowClose() override;
 	virtual TWeakPtr<SWidget> GetWidget() override;
 	virtual FReply OnViewportActivated(const FWindowActivateEvent& InActivateEvent) override;
 	virtual void OnViewportDeactivated(const FWindowActivateEvent& InActivateEvent) override;
diff --git a/Engine/Source/Runtime/Engine/Public/UnrealClient.h b/Engine/Source/Runtime/Engine/Public/UnrealClient.h
--- a/Engine/Source/Runtime/Engine/Public/UnrealClient.h
+++ b/Engine/Source/Runtime/Engine/Public/UnrealClient.h
@@ -824,6 +824,17 @@ public:
 	virtual void Activated(FViewport* Viewport, const FWindowActivateEvent& InActivateEvent) {}
 	virtual void Deactivated(FViewport* Viewport, const FWindowActivateEvent& InActivateEvent) {}
 
+	/**
+	 * Called when the top level window associated with the viewport has been requested to close.
+	 * At this point, the viewport has not been closed and the operation may be canceled.
+	 * This may not called from PIE, Editor Windows, on consoles, or before the game ends
+ 	 * from other methods.
+	 * This is only when the platform specific window is closed.
+	 *
+	 * @return True if the viewport may be closed, false otherwise.
+	 */
+	virtual bool WindowCloseRequested() { return true; }
+
 	virtual void CloseRequested(FViewport* Viewport) {}
 
 	virtual bool RequiresHitProxyStorage() { return true; }
diff --git a/Engine/Source/Runtime/EngineSettings/Classes/GameMapsSettings.h b/Engine/Source/Runtime/EngineSettings/Classes/GameMapsSettings.h
--- a/Engine/Source/Runtime/EngineSettings/Classes/GameMapsSettings.h
+++ b/Engine/Source/Runtime/EngineSettings/Classes/GameMapsSettings.h
@@ -28,6 +28,20 @@ namespace EThreePlayerSplitScreenType
 	};
 }
 
+/** Helper structure, used to associate GameModes with shortcut names. */
+USTRUCT()
+struct FGameModeName
+{
+	GENERATED_USTRUCT_BODY()
+
+	/** Abbreviation/prefix that can be used as an alias for the class name */
+	UPROPERTY(EditAnywhere, Category = DefaultModes, meta = (MetaClass = "GameModeBase"))
+	FString Name;
+
+	/** GameMode class to load */
+	UPROPERTY(EditAnywhere, Category = DefaultModes, meta = (MetaClass = "GameModeBase"))
+	FStringClassReference GameMode;
+};
 
 UCLASS(config=Engine, defaultconfig)
 class ENGINESETTINGS_API UGameMapsSettings
@@ -51,6 +65,20 @@ class ENGINESETTINGS_API UGameMapsSettings
 	 */
 	static const FString& GetGlobalDefaultGameMode( );
 
+	/**
+	 * Searches the GameModeClassAliases list for a named game mode, if not found will return passed in string
+	 * 
+	 * @return the proper game type class path to load
+	 */
+	static FString GetGameModeForName( const FString& GameModeName );
+
+	/**
+	 * Searches the GameModeMapPrefixes list for a named game mode, if not found will return passed in string
+	 * 
+	 * @return the proper game type class path to load, or empty if not found
+	 */
+	static FString GetGameModeForMapName( const FString& MapName );
+
 	/**
 	 * Set the default map to use (see GameDefaultMap below)
 	 *
@@ -116,13 +144,21 @@ private:
 	FStringAssetReference ServerDefaultMap;
 
 	/** GameMode to use if not specified in any other way. (e.g. per-map DefaultGameMode or on the URL). */
-	UPROPERTY(config, noclear, EditAnywhere, Category=DefaultModes, meta=(MetaClass="GameMode", DisplayName="Default GameMode"))
+	UPROPERTY(config, noclear, EditAnywhere, Category=DefaultModes, meta=(MetaClass="GameModeBase", DisplayName="Default GameMode"))
 	FStringClassReference GlobalDefaultGameMode;
 
 	/**
 	 * GameMode to use if not specified in any other way. (e.g. per-map DefaultGameMode or on the URL) (DEDICATED SERVERS)
 	 * If not set, the GlobalDefaultGameMode value will be used.
 	 */
-	UPROPERTY(config, EditAnywhere, Category=DefaultModes, meta=(MetaClass="GameMode"), AdvancedDisplay)
+	UPROPERTY(config, EditAnywhere, Category=DefaultModes, meta=(MetaClass="GameModeBase"), AdvancedDisplay)
 	FStringClassReference GlobalDefaultServerGameMode;
+
+	/** Overrides the GameMode to use when loading a map that starts with a specific prefix */
+	UPROPERTY(config, EditAnywhere, Category = DefaultModes, AdvancedDisplay)
+	TArray<FGameModeName> GameModeMapPrefixes;
+
+	/** List of GameModes to load when game= is specified in the URL (e.g. "DM" could be an alias for "MyProject.MyGameModeMP_DM") */
+	UPROPERTY(config, EditAnywhere, Category = DefaultModes, AdvancedDisplay)
+	TArray<FGameModeName> GameModeClassAliases;
 };
diff --git a/Engine/Source/Runtime/EngineSettings/Private/EngineSettingsModule.cpp b/Engine/Source/Runtime/EngineSettings/Private/EngineSettingsModule.cpp
--- a/Engine/Source/Runtime/EngineSettings/Private/EngineSettingsModule.cpp
+++ b/Engine/Source/Runtime/EngineSettings/Private/EngineSettingsModule.cpp
@@ -92,6 +92,92 @@ const FString& UGameMapsSettings::GetGlobalDefaultGameMode( )
 		: GameMapsSettings->GlobalDefaultGameMode.ToString();
 }
 
+FString UGameMapsSettings::GetGameModeForName(const FString& GameModeName)
+{
+	UGameMapsSettings* GameMapsSettings = Cast<UGameMapsSettings>(UGameMapsSettings::StaticClass()->GetDefaultObject());
+
+	// Look to see if this should be remapped from a shortname to full class name
+	for (const FGameModeName& Alias : GameMapsSettings->GameModeClassAliases)
+	{
+		if (GameModeName == Alias.Name)
+		{
+			// switch GameClassName to the full name
+			return Alias.GameMode.ToString();
+		}
+	}
+
+	// Check deprecated config
+	FConfigSection* GameModeSection = GConfig->GetSectionPrivate(TEXT("/Script/Engine.GameMode"), false, true, GGameIni);
+	
+	if (GameModeSection)
+	{
+		TArray<FString> ConfigLines;
+		GameModeSection->MultiFind(TEXT("GameModeClassAliases"), ConfigLines);
+
+		if (ConfigLines.Num())
+		{
+			UE_LOG(LogLoad, Warning, TEXT("GameMode::GameModeClassAliases are deprecated, move to GameMapsSettings"));
+
+			for (FString& ConfigString : ConfigLines)
+			{
+				FString ModeName, ModePath;
+				if (FParse::Value(*ConfigString, TEXT("ShortName="), ModeName) && FParse::Value(*ConfigString, TEXT("GameClassName="), ModePath))
+				{
+					if (ModeName == GameModeName)
+					{
+						return ModePath;
+					}
+				}
+			}
+
+		}
+	}
+
+	return GameModeName;
+}
+
+FString UGameMapsSettings::GetGameModeForMapName(const FString& MapName)
+{
+	UGameMapsSettings* GameMapsSettings = Cast<UGameMapsSettings>(UGameMapsSettings::StaticClass()->GetDefaultObject());
+
+	// See if we have a per-prefix default specified
+	for (const FGameModeName& Prefix : GameMapsSettings->GameModeMapPrefixes)
+	{
+		if ((Prefix.Name.Len() > 0) && MapName.StartsWith(Prefix.Name))
+		{
+			return Prefix.GameMode.ToString();
+		}
+	}
+
+	// Check deprecated config
+	FConfigSection* GameModeSection = GConfig->GetSectionPrivate(TEXT("/Script/Engine.WorldSettings"), false, true, GGameIni);
+
+	if (GameModeSection)
+	{
+		TArray<FString> ConfigLines;
+		GameModeSection->MultiFind(TEXT("DefaultMapPrefixes"), ConfigLines);
+
+		if (ConfigLines.Num())
+		{
+			UE_LOG(LogLoad, Warning, TEXT("GameMode::DefaultMapPrefixes are deprecated, move to GameMapsSettings::GameModeMapPrefixes"));
+
+			for (FString& ConfigString : ConfigLines)
+			{
+				FString Prefix, ModePath;
+				if (FParse::Value(*ConfigString, TEXT("Prefix="), Prefix) && FParse::Value(*ConfigString, TEXT("GameMode="), ModePath))
+				{
+					if (MapName.StartsWith(Prefix))
+					{
+						return ModePath;
+					}
+				}
+			}
+
+		}
+	}
+
+	return FString();
+}
 
 void UGameMapsSettings::SetGameDefaultMap( const FString& NewMap )
 {
diff --git a/Engine/Source/Runtime/GameplayAbilities/Private/GameplayEffect.cpp b/Engine/Source/Runtime/GameplayAbilities/Private/GameplayEffect.cpp
--- a/Engine/Source/Runtime/GameplayAbilities/Private/GameplayEffect.cpp
+++ b/Engine/Source/Runtime/GameplayAbilities/Private/GameplayEffect.cpp
@@ -3413,7 +3413,7 @@ void FActiveGameplayEffectsContainer::Uninitialize()
 float FActiveGameplayEffectsContainer::GetServerWorldTime() const
 {
 	UWorld* World = Owner->GetWorld();
-	AGameState* GameState = World->GetGameState<AGameState>();
+	AGameStateBase* GameState = World->GetGameState();
 	if (GameState)
 	{
 		return GameState->GetServerWorldTimeSeconds();
diff --git a/Engine/Source/Runtime/GeometryCache/Classes/GeometryCacheComponent.h b/Engine/Source/Runtime/GeometryCache/Classes/GeometryCacheComponent.h
--- a/Engine/Source/Runtime/GeometryCache/Classes/GeometryCacheComponent.h
+++ b/Engine/Source/Runtime/GeometryCache/Classes/GeometryCacheComponent.h
@@ -207,10 +207,10 @@ protected:
 	UPROPERTY(EditAnywhere, Category = GeometryCache)
 	bool bLooping;
 
-	UPROPERTY(EditAnywhere, Category = GeometryCache)
+	UPROPERTY(EditAnywhere, Category = GeometryCache, meta = (UIMin = "-14400.0", UIMax = "14400.0", ClampMin = "-14400.0", ClampMax = "14400.0"))
 	float StartTimeOffset;
 
-	UPROPERTY(EditAnywhere, Category = GeometryCache)
+	UPROPERTY(EditAnywhere, Category = GeometryCache, meta = (UIMin = "-512.0", UIMax = "512.0", ClampMin = "-512.0", ClampMax = "512.0"))
 	float PlaybackSpeed;
 
 	UPROPERTY(VisibleAnywhere, Category = GeometryCache)
diff --git a/Engine/Source/Runtime/GeometryCache/Private/GeometryCacheComponent.cpp b/Engine/Source/Runtime/GeometryCache/Private/GeometryCacheComponent.cpp
--- a/Engine/Source/Runtime/GeometryCache/Private/GeometryCacheComponent.cpp
+++ b/Engine/Source/Runtime/GeometryCache/Private/GeometryCacheComponent.cpp
@@ -68,8 +68,9 @@ void UGeometryCacheComponent::SetupTrackData()
 			FGeometryCacheMeshData* MeshData = nullptr;
 
 			// Retrieve the matrix/mesh data and the appropriate sample indices
-			Track->UpdateMatrixData(ElapsedTime + StartTimeOffset, bLooping, MatrixSampleIndex, WorldMatrix);
-			Track->UpdateMeshData(ElapsedTime + StartTimeOffset, bLooping, MeshSampleIndex, MeshData);
+			const float ClampedStartTimeOffset = FMath::Clamp(StartTimeOffset, -14400.0f, 14400.0f);
+			Track->UpdateMatrixData(ElapsedTime + ClampedStartTimeOffset, bLooping, MatrixSampleIndex, WorldMatrix);
+			Track->UpdateMeshData(ElapsedTime + ClampedStartTimeOffset, bLooping, MeshSampleIndex, MeshData);
 
 			// First time so create rather than update the mesh sections
 			CreateTrackSection(TrackIndex, WorldMatrix, MeshData);
@@ -100,9 +101,9 @@ void UGeometryCacheComponent::TickComponent(float DeltaTime, enum ELevelTick Tic
 	if (GeometryCache && bRunning)
 	{
 		// Increase total elapsed time since BeginPlay according to PlayDirection and speed
-		ElapsedTime += (DeltaTime * PlayDirection * PlaybackSpeed);
+		ElapsedTime += (DeltaTime * PlayDirection * GetPlaybackSpeed());
 
-		if (ElapsedTime < 0.0f)
+		if (ElapsedTime < 0.0f && bLooping)
 		{
 			ElapsedTime += Duration;
 		}
@@ -113,9 +114,9 @@ void UGeometryCacheComponent::TickComponent(float DeltaTime, enum ELevelTick Tic
 			UGeometryCacheTrack* Track = GeometryCache->Tracks[TrackIndex];
 			FMatrix WorldMatrix;
 			FGeometryCacheMeshData* MeshData = nullptr;
-
-			const bool bUpdateMatrix = Track->UpdateMatrixData(ElapsedTime + StartTimeOffset, bLooping, TrackMatrixSampleIndices[TrackIndex], WorldMatrix);
-			const bool bUpdateMesh = Track->UpdateMeshData(ElapsedTime + StartTimeOffset, bLooping, TrackMeshSampleIndices[TrackIndex], MeshData);
+			const float ClampedStartTimeOffset = FMath::Clamp(StartTimeOffset, -14400.0f, 14400.0f);
+			const bool bUpdateMatrix = Track->UpdateMatrixData(ElapsedTime + ClampedStartTimeOffset, bLooping, TrackMatrixSampleIndices[TrackIndex], WorldMatrix);
+			const bool bUpdateMesh = Track->UpdateMeshData(ElapsedTime + ClampedStartTimeOffset, bLooping, TrackMeshSampleIndices[TrackIndex], MeshData);
 
 			// Update sections according what is required
 			if (bUpdateMatrix)
@@ -311,8 +312,9 @@ void UGeometryCacheComponent::OnObjectReimported(UGeometryCache* ImportedGeometr
 				FGeometryCacheMeshData* MeshData = nullptr;
 
 				// Retrieve the matrix/mesh data and the appropriate sample indices
-				Track->UpdateMatrixData(ElapsedTime + StartTimeOffset, bLooping, MatrixSampleIndex, WorldMatrix);
-				Track->UpdateMeshData(ElapsedTime + StartTimeOffset, bLooping, MeshSampleIndex, MeshData);
+				const float ClampedStartTimeOffset = FMath::Clamp(StartTimeOffset, -14400.0f, 14400.0f);
+				Track->UpdateMatrixData(ElapsedTime + ClampedStartTimeOffset, bLooping, MatrixSampleIndex, WorldMatrix);
+				Track->UpdateMeshData(ElapsedTime + ClampedStartTimeOffset, bLooping, MeshSampleIndex, MeshData);
 
 				// First time so create rather than update the mesh sections
 				CreateTrackSection(TrackIndex, WorldMatrix, MeshData);
@@ -375,13 +377,13 @@ bool UGeometryCacheComponent::IsPlayingReversed() const
 
 float UGeometryCacheComponent::GetPlaybackSpeed() const
 {
-	return PlaybackSpeed;
+	return FMath::Clamp(PlaybackSpeed, -512.0f, 512.0f);
 }
 
 void UGeometryCacheComponent::SetPlaybackSpeed(const float NewPlaybackSpeed)
 {
 	// Currently only positive play back speeds are supported
-	PlaybackSpeed = fabs( NewPlaybackSpeed );
+	PlaybackSpeed = FMath::Clamp( NewPlaybackSpeed, -512.0f, 512.0f );
 }
 
 bool UGeometryCacheComponent::SetGeometryCache(UGeometryCache* NewGeomCache)
diff --git a/Engine/Source/Runtime/Launch/Launch.Build.cs b/Engine/Source/Runtime/Launch/Launch.Build.cs
--- a/Engine/Source/Runtime/Launch/Launch.Build.cs
+++ b/Engine/Source/Runtime/Launch/Launch.Build.cs
@@ -63,6 +63,7 @@ public class Launch : ModuleRules
 				DynamicallyLoadedModuleNames.Add("D3D12RHI");
 				DynamicallyLoadedModuleNames.Add("D3D11RHI");
 				DynamicallyLoadedModuleNames.Add("XAudio2");
+				DynamicallyLoadedModuleNames.Add("AudioMixerXAudio2");
 			}
 			else if (Target.Platform == UnrealTargetPlatform.Mac)
 			{
diff --git a/Engine/Source/Runtime/Launch/Private/LaunchEngineLoop.cpp b/Engine/Source/Runtime/Launch/Private/LaunchEngineLoop.cpp
--- a/Engine/Source/Runtime/Launch/Private/LaunchEngineLoop.cpp
+++ b/Engine/Source/Runtime/Launch/Private/LaunchEngineLoop.cpp
@@ -1344,8 +1344,9 @@ int32 FEngineLoop::PreInit( const TCHAR* CmdLine )
 #endif	//WITH_EDITOR
 		PRIVATE_GIsRunningCommandlet = true;
 
-		// Allow commandlet rendering based on command line switch (too early to let the commandlet itself override this).
+		// Allow commandlet rendering and/or audio based on command line switch (too early to let the commandlet itself override this).
 		PRIVATE_GAllowCommandletRendering = FParse::Param(FCommandLine::Get(), TEXT("AllowCommandletRendering"));
+		PRIVATE_GAllowCommandletAudio = FParse::Param(FCommandLine::Get(), TEXT("AllowCommandletAudio"));
 
 		// We need to disregard the empty token as we try finding Token + "Commandlet" which would result in finding the
 		// UCommandlet class if Token is empty.
diff --git a/Engine/Source/Runtime/MovieSceneCapture/Public/MovieSceneCaptureSettings.h b/Engine/Source/Runtime/MovieSceneCapture/Public/MovieSceneCaptureSettings.h
--- a/Engine/Source/Runtime/MovieSceneCapture/Public/MovieSceneCaptureSettings.h
+++ b/Engine/Source/Runtime/MovieSceneCapture/Public/MovieSceneCaptureSettings.h
@@ -39,7 +39,7 @@ struct MOVIESCENECAPTURE_API FMovieSceneCaptureSettings
 
 	/** Optional game mode to override the map's default game mode with.  This can be useful if the game's normal mode displays UI elements or loading screens that you don't want captured. */
 	UPROPERTY(config, EditAnywhere, Category=General, AdvancedDisplay)
-	TSubclassOf<class AGameMode> GameModeOverride;
+	TSubclassOf<class AGameModeBase> GameModeOverride;
 
 	/** The format to use for the resulting filename. Extension will be added automatically. Any tokens of the form {token} will be replaced with the corresponding value:
 	 * {fps}		- The captured framerate
diff --git a/Engine/Source/Runtime/MovieSceneTracks/Private/TrackInstances/MovieSceneAudioTrackInstance.cpp b/Engine/Source/Runtime/MovieSceneTracks/Private/TrackInstances/MovieSceneAudioTrackInstance.cpp
--- a/Engine/Source/Runtime/MovieSceneTracks/Private/TrackInstances/MovieSceneAudioTrackInstance.cpp
+++ b/Engine/Source/Runtime/MovieSceneTracks/Private/TrackInstances/MovieSceneAudioTrackInstance.cpp
@@ -233,13 +233,16 @@ TWeakObjectPtr<UAudioComponent> FMovieSceneAudioTrackInstance::GetAudioComponent
 		USoundCue* TempPlaybackAudioCue = NewObject<USoundCue>();
 		UWorld* World = Actor ? Actor->GetWorld() : (Player.GetPlaybackContext() != nullptr) ? Player.GetPlaybackContext()->GetWorld() : nullptr;
 
-		AActor* OuterActor = nullptr; // The audio component should be transient but attached to the actor
+		FAudioDevice::FCreateComponentParams Params(World);
 
-		UAudioComponent* AudioComponent = FAudioDevice::CreateComponent(TempPlaybackAudioCue, World, OuterActor, false, false);								
-		if (AudioComponent && Actor)
+		UAudioComponent* AudioComponent = FAudioDevice::CreateComponent(TempPlaybackAudioCue, Params);
+		if (AudioComponent)
 		{
 			AudioComponent->SetFlags(RF_Transient);
-			AudioComponent->AttachToComponent(Actor->GetRootComponent(), FAttachmentTransformRules::KeepRelativeTransform);
+			if (Actor)
+			{
+				AudioComponent->AttachToComponent(Actor->GetRootComponent(), FAttachmentTransformRules::KeepRelativeTransform);
+			}
 		}
 
 		PlaybackAudioComponents[RowIndex].Add(Actor);
diff --git a/Engine/Source/Runtime/Slate/Private/Framework/Application/SlateApplication.cpp b/Engine/Source/Runtime/Slate/Private/Framework/Application/SlateApplication.cpp
--- a/Engine/Source/Runtime/Slate/Private/Framework/Application/SlateApplication.cpp
+++ b/Engine/Source/Runtime/Slate/Private/Framework/Application/SlateApplication.cpp
@@ -6257,7 +6257,21 @@ void FSlateApplication::OnWindowClose( const TSharedRef< FGenericWindow >& Platf
 
 	if ( Window.IsValid() )
 	{
-		Window->RequestDestroyWindow();
+		bool bCanCloseWindow = true;
+		TSharedPtr< SViewport > CurrentGameViewportWidget = GameViewportWidget.Pin();
+		if (CurrentGameViewportWidget.IsValid())
+		{
+			TSharedPtr< ISlateViewport > SlateViewport = CurrentGameViewportWidget->GetViewportInterface().Pin();
+			if (SlateViewport.IsValid())
+			{
+				bCanCloseWindow = !SlateViewport->OnRequestWindowClose().bIsHandled;
+			}
+		}
+		
+		if (bCanCloseWindow)
+		{
+		    Window->RequestDestroyWindow();
+		}	 
 	}
 }
 
diff --git a/Engine/Source/Runtime/Slate/Private/Widgets/Views/SHeaderRow.cpp b/Engine/Source/Runtime/Slate/Private/Widgets/Views/SHeaderRow.cpp
--- a/Engine/Source/Runtime/Slate/Private/Widgets/Views/SHeaderRow.cpp
+++ b/Engine/Source/Runtime/Slate/Private/Widgets/Views/SHeaderRow.cpp
@@ -491,6 +491,11 @@ void SHeaderRow::RemoveColumn( const FName& InColumnId )
 	RegenerateWidgets();
 }
 
+void SHeaderRow::RefreshColumns()
+{
+	RegenerateWidgets();
+}
+
 void SHeaderRow::ClearColumns()
 {
 	Columns.Empty();
@@ -556,56 +561,59 @@ void SHeaderRow::RegenerateWidgets()
 		{
 			FColumn& SomeColumn = Columns[SlotIndex];
 
-			// Keep track of the last header we created.
-			TSharedPtr<STableColumnHeader> PrecedingHeader = NewlyMadeHeader;
-			NewlyMadeHeader.Reset();
+			if ( SomeColumn.ShouldGenerateWidget.Get(true))
+			{
+				// Keep track of the last header we created.
+				TSharedPtr<STableColumnHeader> PrecedingHeader = NewlyMadeHeader;
+				NewlyMadeHeader.Reset();
 
-			FMargin DefaultPadding = FMargin( HalfSplitterDetectionSize, 0, HalfSplitterDetectionSize, 0 );
-			TSharedRef<STableColumnHeader> NewHeader =
-				SAssignNew(NewlyMadeHeader, STableColumnHeader, SomeColumn, DefaultPadding)
-				.Style( (SlotIndex + 1 == Columns.Num()) ? &Style->LastColumnStyle : &Style->ColumnStyle );
+				FMargin DefaultPadding = FMargin(HalfSplitterDetectionSize, 0, HalfSplitterDetectionSize, 0);
 
-			switch(SomeColumn.SizeRule)
-			{
-			case EColumnSizeMode::Fill:
+				TSharedRef<STableColumnHeader> NewHeader =
+					SAssignNew(NewlyMadeHeader, STableColumnHeader, SomeColumn, DefaultPadding)
+					.Style((SlotIndex + 1 == Columns.Num()) ? &Style->LastColumnStyle : &Style->ColumnStyle);
+
+				switch (SomeColumn.SizeRule)
+				{
+				case EColumnSizeMode::Fill:
 				{
 					TAttribute<float> WidthBinding;
-					WidthBinding.BindRaw( &SomeColumn, &FColumn::GetWidth );
+					WidthBinding.BindRaw(&SomeColumn, &FColumn::GetWidth);
 
 					// Add resizable cell
 					Splitter->AddSlot()
-					.Value( WidthBinding )
-					.SizeRule( SSplitter::FractionOfParent )
-					.OnSlotResized( SSplitter::FOnSlotResized::CreateRaw( &SomeColumn, &FColumn::SetWidth ) )
-					[
-						NewHeader
-					];
+						.Value(WidthBinding)
+						.SizeRule(SSplitter::FractionOfParent)
+						.OnSlotResized(SSplitter::FOnSlotResized::CreateRaw(&SomeColumn, &FColumn::SetWidth))
+						[
+							NewHeader
+						];
 				}
 				break;
 
-			case EColumnSizeMode::Fixed:
+				case EColumnSizeMode::Fixed:
 				{
 					// Add fixed size cell
 					Splitter->AddSlot()
-					.SizeRule( SSplitter::SizeToContent )
-					[
-						SNew( SBox )
-						.WidthOverride( SomeColumn.GetWidth() )
+						.SizeRule(SSplitter::SizeToContent)
+						[
+							SNew(SBox)
+							.WidthOverride(SomeColumn.GetWidth())
 						[
 							NewHeader
 						]
-					];
+						];
 				}
 				break;
 
-			case EColumnSizeMode::Manual:
+				case EColumnSizeMode::Manual:
 				{
 					// Sizing grip to put at the end of the column - we can't use a SSplitter here as it doesn't have the resizing behavior we need
 					const float GripSize = 5.0f;
 					TSharedRef<SBorder> SizingGrip = SNew(SBorder)
 						.Padding(0.0f)
-						.BorderImage( FCoreStyle::Get().GetBrush("NoBorder") )
-						.Cursor( EMouseCursor::ResizeLeftRight )
+						.BorderImage(FCoreStyle::Get().GetBrush("NoBorder"))
+						.Cursor(EMouseCursor::ResizeLeftRight)
 						.Content()
 						[
 							SNew(SSpacer)
@@ -616,7 +624,7 @@ void SHeaderRow::RegenerateWidgets()
 					auto SizingGrip_OnMouseButtonDown = [&SomeColumn, WeakSizingGrip](const FGeometry&, const FPointerEvent&) -> FReply
 					{
 						TSharedPtr<SBorder> SizingGripPtr = WeakSizingGrip.Pin();
-						if ( SizingGripPtr.IsValid() )
+						if (SizingGripPtr.IsValid())
 						{
 							return FReply::Handled().CaptureMouse(SizingGripPtr.ToSharedRef());
 						}
@@ -626,7 +634,7 @@ void SHeaderRow::RegenerateWidgets()
 					auto SizingGrip_OnMouseButtonUp = [&SomeColumn, WeakSizingGrip](const FGeometry&, const FPointerEvent&) -> FReply
 					{
 						TSharedPtr<SBorder> SizingGripPtr = WeakSizingGrip.Pin();
-						if ( SizingGripPtr.IsValid() && SizingGripPtr->HasMouseCapture() )
+						if (SizingGripPtr.IsValid() && SizingGripPtr->HasMouseCapture())
 						{
 							return FReply::Handled().ReleaseMouseCapture();
 						}
@@ -636,7 +644,7 @@ void SHeaderRow::RegenerateWidgets()
 					auto SizingGrip_OnMouseMove = [&SomeColumn, WeakSizingGrip](const FGeometry&, const FPointerEvent& InPointerEvent) -> FReply
 					{
 						TSharedPtr<SBorder> SizingGripPtr = WeakSizingGrip.Pin();
-						if ( SizingGripPtr.IsValid() && SizingGripPtr->HasMouseCapture() )
+						if (SizingGripPtr.IsValid() && SizingGripPtr->HasMouseCapture())
 						{
 							// The sizing grip has been moved, so update our columns size from the movement delta
 							const float NewWidth = SomeColumn.GetWidth() + InPointerEvent.GetCursorDelta().X;
@@ -647,9 +655,9 @@ void SHeaderRow::RegenerateWidgets()
 					};
 
 					// Bind the events to handle the drag sizing
-					SizingGrip->SetOnMouseButtonDown( FPointerEventHandler::CreateLambda(SizingGrip_OnMouseButtonDown) );
-					SizingGrip->SetOnMouseButtonUp( FPointerEventHandler::CreateLambda(SizingGrip_OnMouseButtonUp) );
-					SizingGrip->SetOnMouseMove( FPointerEventHandler::CreateLambda(SizingGrip_OnMouseMove) );
+					SizingGrip->SetOnMouseButtonDown(FPointerEventHandler::CreateLambda(SizingGrip_OnMouseButtonDown));
+					SizingGrip->SetOnMouseButtonUp(FPointerEventHandler::CreateLambda(SizingGrip_OnMouseButtonUp));
+					SizingGrip->SetOnMouseMove(FPointerEventHandler::CreateLambda(SizingGrip_OnMouseMove));
 
 					auto GetColumnWidthAsOptionalSize = [&SomeColumn]() -> FOptionalSize
 					{
@@ -662,29 +670,30 @@ void SHeaderRow::RegenerateWidgets()
 
 					// Add resizable cell
 					Splitter->AddSlot()
-					.SizeRule( SSplitter::SizeToContent )
-					[
-						SNew(SBox)
-						.WidthOverride(WidthBinding)
+						.SizeRule(SSplitter::SizeToContent)
+						[
+							SNew(SBox)
+							.WidthOverride(WidthBinding)
 						[
 							SNew(SOverlay)
-							+SOverlay::Slot()
-							[
-								NewHeader
-							]
-							+SOverlay::Slot()
-							.HAlign(HAlign_Right)
-							[
-								SizingGrip
-							]
+							+ SOverlay::Slot()
+						[
+							NewHeader
 						]
-					];
+					+ SOverlay::Slot()
+						.HAlign(HAlign_Right)
+						[
+							SizingGrip
+						]
+						]
+						];
 				}
 				break;
 
-			default:
-				break;
-			}
+				default:
+					break;
+				}
+			}			
 		}
 	}
 
diff --git a/Engine/Source/Runtime/Slate/Public/Widgets/Views/SHeaderRow.h b/Engine/Source/Runtime/Slate/Public/Widgets/Views/SHeaderRow.h
--- a/Engine/Source/Runtime/Slate/Public/Widgets/Views/SHeaderRow.h
+++ b/Engine/Source/Runtime/Slate/Public/Widgets/Views/SHeaderRow.h
@@ -122,6 +122,8 @@ public:
 			SLATE_ATTRIBUTE( EColumnSortMode::Type, SortMode )
 			SLATE_ATTRIBUTE( EColumnSortPriority::Type, SortPriority )
 			SLATE_EVENT( FOnSortModeChanged, OnSort )
+
+			SLATE_ATTRIBUTE(bool, ShouldGenerateWidget)
 		SLATE_END_ARGS()
 
 		FColumn( const FArguments& InArgs )
@@ -143,6 +145,7 @@ public:
 			, SortMode( InArgs._SortMode )
 			, SortPriority( InArgs._SortPriority )
 			, OnSortModeChanged( InArgs._OnSort )
+			, ShouldGenerateWidget(InArgs._ShouldGenerateWidget)
 		{
 			if ( InArgs._FixedWidth.IsSet() )
 			{
@@ -215,6 +218,8 @@ public:
 		TAttribute< EColumnSortMode::Type > SortMode;
 		TAttribute< EColumnSortPriority::Type > SortPriority;
 		FOnSortModeChanged OnSortModeChanged;
+
+		TAttribute<bool> ShouldGenerateWidget;
 	};
 
 	/** Create a column with the specified ColumnId */
@@ -255,6 +260,9 @@ public:
 	/** Removes a column from the header */
 	void RemoveColumn( const FName& InColumnId );
 
+	/** Force refreshing of the column widgets*/
+	void RefreshColumns();
+
 	/** Removes all columns from the header */
 	void ClearColumns();
 
diff --git a/Engine/Source/Runtime/Slate/Public/Widgets/Views/STableRow.h b/Engine/Source/Runtime/Slate/Public/Widgets/Views/STableRow.h
--- a/Engine/Source/Runtime/Slate/Public/Widgets/Views/STableRow.h
+++ b/Engine/Source/Runtime/Slate/Public/Widgets/Views/STableRow.h
@@ -1056,53 +1056,54 @@ protected:
 		for( int32 ColumnIndex = 0; ColumnIndex < NumColumns; ++ColumnIndex )
 		{
 			const SHeaderRow::FColumn& Column = Columns[ColumnIndex];
-			const TSharedRef< SWidget >* const ExistingWidget = ColumnIdToSlotContents.Find( Column.ColumnId );
-
-			TSharedRef< SWidget > CellContents = SNullWidget::NullWidget;
-			if ( ExistingWidget != nullptr )
-			{
-				CellContents = *ExistingWidget;
-			}
-			else
+			if (Column.ShouldGenerateWidget.Get(true))
 			{
-				CellContents = GenerateWidgetForColumn( Column.ColumnId );
-			}
+				TSharedRef< SWidget >* ExistingWidget = ColumnIdToSlotContents.Find(Column.ColumnId);
+				TSharedRef< SWidget > CellContents = SNullWidget::NullWidget;
+				if (ExistingWidget != nullptr)
+				{
+					CellContents = *ExistingWidget;
+				}
+				else
+				{
+					CellContents = GenerateWidgetForColumn(Column.ColumnId);
+				}
 
-			switch(Column.SizeRule)
-			{
-			case EColumnSizeMode::Fill:
+				switch (Column.SizeRule)
+				{
+				case EColumnSizeMode::Fill:
 				{
 					TAttribute<float> WidthBinding;
-					WidthBinding.BindRaw( &Column, &SHeaderRow::FColumn::GetWidth );
+					WidthBinding.BindRaw(&Column, &SHeaderRow::FColumn::GetWidth);
 
 					SHorizontalBox::FSlot& NewSlot = Box->AddSlot()
-					.HAlign(Column.CellHAlignment)
-					.VAlign(Column.CellVAlignment)
-					.FillWidth( WidthBinding )
-					[
-						CellContents
-					];
+						.HAlign(Column.CellHAlignment)
+						.VAlign(Column.CellVAlignment)
+						.FillWidth(WidthBinding)
+						[
+							CellContents
+						];
 				}
 				break;
 
-			case EColumnSizeMode::Fixed:
+				case EColumnSizeMode::Fixed:
 				{
 					Box->AddSlot()
-					.AutoWidth()
-					[
-						SNew( SBox )
-						.WidthOverride( Column.Width.Get() )
+						.AutoWidth()
+						[
+							SNew(SBox)
+							.WidthOverride(Column.Width.Get())
 						.HAlign(Column.CellHAlignment)
 						.VAlign(Column.CellVAlignment)
 						.Content()
 						[
 							CellContents
 						]
-					];
+						];
 				}
 				break;
 
-			case EColumnSizeMode::Manual:
+				case EColumnSizeMode::Manual:
 				{
 					auto GetColumnWidthAsOptionalSize = [&Column]() -> FOptionalSize
 					{
@@ -1114,25 +1115,26 @@ protected:
 					WidthBinding.Bind(TAttribute<FOptionalSize>::FGetter::CreateLambda(GetColumnWidthAsOptionalSize));
 
 					Box->AddSlot()
-					.AutoWidth()
-					[
-						SNew( SBox )
-						.WidthOverride(WidthBinding)
+						.AutoWidth()
+						[
+							SNew(SBox)
+							.WidthOverride(WidthBinding)
 						.HAlign(Column.CellHAlignment)
 						.VAlign(Column.CellVAlignment)
 						.Content()
 						[
 							CellContents
 						]
-					];
+						];
 				}
 				break;
 
-			default:
-				break;
-			}
+				default:
+					break;
+				}
 
-			NewColumnIdToSlotContents.Add( Column.ColumnId, CellContents );
+				NewColumnIdToSlotContents.Add(Column.ColumnId, CellContents);
+			}
 		}
 
 		ColumnIdToSlotContents = NewColumnIdToSlotContents;
diff --git a/Engine/Source/Runtime/SlateCore/Public/Rendering/RenderingCommon.h b/Engine/Source/Runtime/SlateCore/Public/Rendering/RenderingCommon.h
--- a/Engine/Source/Runtime/SlateCore/Public/Rendering/RenderingCommon.h
+++ b/Engine/Source/Runtime/SlateCore/Public/Rendering/RenderingCommon.h
@@ -526,14 +526,28 @@ public:
 	}
 
 	/**
-	 * Called when the viewports top level window is being closed
+	 * Called when the top level window associated with the viewport has been requested to close.
+	 * At this point, the viewport has not been closed and the operation may be canceled.
+	 * This may not called from PIE, Editor Windows, on consoles, or before the game ends
+ 	 * from other methods.
+	 * This is only when the platform specific window is closed.
+	 *
+	 * @return FReply::Handled if the close event was consumed (and the window should remain open).
+	 */
+	virtual FReply OnRequestWindowClose()
+	{
+		return FReply::Unhandled();
+	}
+
+	/**
+	 * Called when the viewport has been requested to close.
 	 */
 	virtual void OnViewportClosed()
 	{
 	}
 
 	/**
-	 * Called when the viewports top level window is being closed
+	 * Gets the SWidget associated with this viewport
 	 */
 	virtual TWeakPtr<SWidget> GetWidget()
 	{
diff --git a/Engine/Source/Runtime/UMG/Private/Components/WidgetComponent.cpp b/Engine/Source/Runtime/UMG/Private/Components/WidgetComponent.cpp
--- a/Engine/Source/Runtime/UMG/Private/Components/WidgetComponent.cpp
+++ b/Engine/Source/Runtime/UMG/Private/Components/WidgetComponent.cpp
@@ -329,9 +329,6 @@ UWidgetComponent::UWidgetComponent( const FObjectInitializer& PCIP )
 	Pivot = FVector2D(0.5, 0.5);
 
 	bAddedToScreen = false;
-
-	// We want this because we want EndPlay to be called!
-	bWantsBeginPlay = true;
 }
 
 void UWidgetComponent::EndPlay(const EEndPlayReason::Type EndPlayReason)
diff --git a/Engine/Source/Runtime/UnrealAudio/Private/Tests/UnrealAudioTestGenerators.cpp b/Engine/Source/Runtime/UnrealAudio/Private/Tests/UnrealAudioTestGenerators.cpp
--- a/Engine/Source/Runtime/UnrealAudio/Private/Tests/UnrealAudioTestGenerators.cpp
+++ b/Engine/Source/Runtime/UnrealAudio/Private/Tests/UnrealAudioTestGenerators.cpp
@@ -6,14 +6,6 @@
 
 #if ENABLE_UNREAL_AUDIO & WITH_DEV_AUTOMATION_TESTS
 
-#ifndef TWO_PI
-#define TWO_PI	6.28318530718
-#endif
-
-#ifndef PI_OVER_TWO
-#define PI_OVER_TWO 1.57079632679
-#endif
-
 namespace UAudio
 {
 namespace Test
@@ -44,14 +36,14 @@ namespace Test
 
 	static float WrapTwoPi(float Value)
 	{
-		while (Value > TWO_PI)
+		while (Value > 2.0 * PI)
 		{
-			Value -= TWO_PI;
+			Value -= 2.0 * PI;
 		}
 
 		while (Value < 0)
 		{
-			Value += TWO_PI;
+			Value += 2.0 * PI;
 		}
 		return Value;
 	}
@@ -188,7 +180,7 @@ namespace Test
 
 		check(CallbackData.FrameRate > 0.0f);
 
-		float Theta = TWO_PI * Frequency / CallbackData.FrameRate;
+		float Theta = 2.0 * PI * Frequency / CallbackData.FrameRate;
 		float InverseQ = 0.5f / Quality;
 
 		float Temp = InverseQ * FMath::Sin(Theta);
@@ -318,7 +310,7 @@ namespace Test
 		if (bIsInit)
 		{
 			check(CallbackData.FrameRate > 0.0f);
-			PhaseDelta = TWO_PI * Frequency / CallbackData.FrameRate;
+			PhaseDelta = 2.0 * PI * Frequency / CallbackData.FrameRate;
 			TargetPhaseDelta = PhaseDelta;
 			PhaseDeltaDelta = 0.0f;
 			Phase = 0.0f;
@@ -326,7 +318,7 @@ namespace Test
 		}
 		else
 		{
-			TargetPhaseDelta = TWO_PI * Frequency / CallbackData.FrameRate;
+			TargetPhaseDelta = 2.0 * PI * Frequency / CallbackData.FrameRate;
 			PhaseDeltaDelta = (TargetPhaseDelta - PhaseDelta) / 100.0f;
 			bNewValue = true;
 		}
@@ -760,11 +752,11 @@ namespace Test
 	}
 
 	FPhaseModulator::OscData::OscData(float InMaxAmp)
-		: Phase(FMath::FRandRange(0.0f, TWO_PI))
+		: Phase(FMath::FRandRange(0.0f, 2.0 * PI))
 		, Delta(0.01f)
 		, TargetDelta(Delta)
 		, DeltaEase(0.001f)
-		, SweepPhase(FMath::FRandRange(0.0f, TWO_PI))
+		, SweepPhase(FMath::FRandRange(0.0f, 2.0 * PI))
 		, SweepDelta(0.0f)
 		, Amp(0.0f)
 		, MaxAmp(InMaxAmp)
@@ -883,7 +875,7 @@ namespace Test
 		float ModIndexValue = ModIndex.GetValue();
 
 		// Plug the mod value into the carrier to generate the value of this synth
-		float Value = Carrier.GetAmp() * ModIndexValue * FMath::Sin(Carrier.GetPhase() + TWO_PI * ModValue * ModIndexValue);
+		float Value = Carrier.GetAmp() * ModIndexValue * FMath::Sin(Carrier.GetPhase() + 2.0 * PI * ModValue * ModIndexValue);
 
 		// Scale by the total amplitude
 		Value *= Parent->Amplitude;
diff --git a/Engine/Source/Runtime/Windows/AudioMixerXAudio2/AudioMixerXAudio2.Build.cs b/Engine/Source/Runtime/Windows/AudioMixerXAudio2/AudioMixerXAudio2.Build.cs
new file mode 100644
--- /dev/null
+++ b/Engine/Source/Runtime/Windows/AudioMixerXAudio2/AudioMixerXAudio2.Build.cs
@@ -0,0 +1,32 @@
+// Copyright 1998-2016 Epic Games, Inc. All Rights Reserved.
+
+using UnrealBuildTool;
+
+public class AudioMixerXAudio2 : ModuleRules
+{
+	public AudioMixerXAudio2(TargetInfo Target)
+	{
+		PrivateIncludePathModuleNames.Add("TargetPlatform");
+		PublicIncludePaths.Add("Runtime/AudioMixer/Public");
+		PrivateIncludePaths.Add("Runtime/AudioMixer/Private");
+
+		PrivateDependencyModuleNames.AddRange(
+			new string[] {
+					"Core",
+					"CoreUObject",
+					"Engine"
+				}
+		);
+
+		PrecompileForTargets = PrecompileTargetsType.None;
+
+		PrivateDependencyModuleNames.Add("AudioMixer");
+
+		AddEngineThirdPartyPrivateStaticDependencies(Target,
+			"DX11Audio",
+			"UEOgg",
+			"Vorbis",
+			"VorbisFile"
+		);
+	}
+}
diff --git a/Engine/Source/Runtime/Windows/AudioMixerXAudio2/Private/AudioMixerPlatformXAudio2.cpp b/Engine/Source/Runtime/Windows/AudioMixerXAudio2/Private/AudioMixerPlatformXAudio2.cpp
new file mode 100644
--- /dev/null
+++ b/Engine/Source/Runtime/Windows/AudioMixerXAudio2/Private/AudioMixerPlatformXAudio2.cpp
@@ -0,0 +1,431 @@
+// Copyright 1998-2016 Epic Games, Inc. All Rights Reserved.
+
+/**
+	Concrete implementation of FAudioDevice for XAudio2
+
+	See https://msdn.microsoft.com/en-us/library/windows/desktop/hh405049%28v=vs.85%29.aspx
+*/
+
+#include "AudioMixer.h"
+#include "AudioMixerDevice.h"
+#include "AudioMixerPlatformXAudio2.h"
+#include "OpusAudioInfo.h"
+#include "VorbisAudioInfo.h"
+
+// Macro to check result code for XAudio2 failure, get the string version, log, and goto a cleanup
+#define XAUDIO2_CLEANUP_ON_FAIL(Result)						\
+	if (FAILED(Result))										\
+	{														\
+		const TCHAR* ErrorString = GetErrorString(Result);	\
+		AUDIO_PLATFORM_ERROR(ErrorString);					\
+		goto Cleanup;										\
+	}
+
+// Macro to check result for XAudio2 failure, get string version, log, and return false
+#define XAUDIO2_RETURN_ON_FAIL(Result)						\
+	if (FAILED(Result))										\
+	{														\
+		const TCHAR* ErrorString = GetErrorString(Result);	\
+		AUDIO_PLATFORM_ERROR(ErrorString);					\
+		return false;										\
+	}
+
+namespace Audio
+{
+	void FXAudio2VoiceCallback::OnBufferEnd(void* BufferContext)
+	{
+		check(BufferContext);
+		IAudioMixerPlatformInterface* MixerPlatform = (IAudioMixerPlatformInterface*)BufferContext;
+		MixerPlatform->ReadNextBuffer();
+	}
+
+	FMixerPlatformXAudio2::FMixerPlatformXAudio2()
+		: XAudio2System(nullptr)
+		, OutputAudioStreamMasteringVoice(nullptr)
+		, OutputAudioStreamSourceVoice(nullptr)
+		, bIsComInitialized(false)
+		, bIsInitialized(false)
+		, bIsDeviceOpen(false)
+	{
+		// Build the channel map. Index corresponds to unreal audio mixer channel enumeration.
+		ChannelTypeMap.Add(SPEAKER_FRONT_LEFT);
+		ChannelTypeMap.Add(SPEAKER_FRONT_RIGHT);
+		ChannelTypeMap.Add(SPEAKER_FRONT_CENTER);
+		ChannelTypeMap.Add(SPEAKER_LOW_FREQUENCY);
+		ChannelTypeMap.Add(SPEAKER_BACK_LEFT);
+		ChannelTypeMap.Add(SPEAKER_BACK_RIGHT);
+		ChannelTypeMap.Add(SPEAKER_FRONT_LEFT_OF_CENTER);
+		ChannelTypeMap.Add(SPEAKER_FRONT_RIGHT_OF_CENTER);
+		ChannelTypeMap.Add(SPEAKER_BACK_CENTER);
+		ChannelTypeMap.Add(SPEAKER_SIDE_LEFT);
+		ChannelTypeMap.Add(SPEAKER_SIDE_RIGHT);
+		ChannelTypeMap.Add(SPEAKER_TOP_CENTER);
+		ChannelTypeMap.Add(SPEAKER_TOP_FRONT_LEFT);
+		ChannelTypeMap.Add(SPEAKER_TOP_FRONT_CENTER);
+		ChannelTypeMap.Add(SPEAKER_TOP_FRONT_RIGHT);
+		ChannelTypeMap.Add(SPEAKER_TOP_BACK_LEFT);
+		ChannelTypeMap.Add(SPEAKER_TOP_BACK_CENTER);
+		ChannelTypeMap.Add(SPEAKER_TOP_BACK_RIGHT);
+	}
+
+	FMixerPlatformXAudio2::~FMixerPlatformXAudio2()
+	{
+	}
+
+	const TCHAR* FMixerPlatformXAudio2::GetErrorString(HRESULT Result)
+	{
+		switch (Result)
+		{
+			case HRESULT(XAUDIO2_E_INVALID_CALL):			return TEXT("XAUDIO2_E_INVALID_CALL");
+			case HRESULT(XAUDIO2_E_XMA_DECODER_ERROR):		return TEXT("XAUDIO2_E_XMA_DECODER_ERROR");
+			case HRESULT(XAUDIO2_E_XAPO_CREATION_FAILED):	return TEXT("XAUDIO2_E_XAPO_CREATION_FAILED");
+			case HRESULT(XAUDIO2_E_DEVICE_INVALIDATED):		return TEXT("XAUDIO2_E_DEVICE_INVALIDATED");
+			case REGDB_E_CLASSNOTREG:						return TEXT("REGDB_E_CLASSNOTREG");
+			case CLASS_E_NOAGGREGATION:						return TEXT("CLASS_E_NOAGGREGATION");
+			case E_NOINTERFACE:								return TEXT("E_NOINTERFACE");
+			case E_POINTER:									return TEXT("E_POINTER");
+			case E_INVALIDARG:								return TEXT("E_INVALIDARG");
+			case E_OUTOFMEMORY:								return TEXT("E_OUTOFMEMORY");
+			default:										return TEXT("UKNOWN");
+		}
+	}
+
+	bool FMixerPlatformXAudio2::InitializeHardware()
+	{
+		if (bIsInitialized)
+		{
+			AUDIO_PLATFORM_ERROR(TEXT("XAudio2 already initialized."));
+			return false;
+
+		}
+
+		// Load ogg and vorbis dlls if they haven't been loaded yet
+		LoadVorbisLibraries();
+
+		bIsComInitialized = FWindowsPlatformMisc::CoInitialize();
+
+		XAUDIO2_RETURN_ON_FAIL(XAudio2Create(&XAudio2System, 0, (XAUDIO2_PROCESSOR)FPlatformAffinity::GetAudioThreadMask()));
+
+		bIsInitialized = true;
+
+		return true;
+	}
+
+	bool FMixerPlatformXAudio2::TeardownHardware()
+	{
+		if (!bIsInitialized)
+		{
+			AUDIO_PLATFORM_ERROR(TEXT("XAudio2 was already tore down."));
+			return false;
+		}
+
+		SAFE_RELEASE(XAudio2System);
+
+		if (bIsComInitialized)
+		{
+			FWindowsPlatformMisc::CoUninitialize();
+		}
+		bIsInitialized = false;
+
+		return true;
+	}
+
+	bool FMixerPlatformXAudio2::IsInitialized() const
+	{
+		return bIsInitialized;
+	}
+
+	bool FMixerPlatformXAudio2::GetNumOutputDevices(uint32& OutNumOutputDevices)
+	{
+		if (!bIsInitialized)
+		{
+			AUDIO_PLATFORM_ERROR(TEXT("XAudio2 was not initialized."));
+			return false;
+		}
+
+		check(XAudio2System);
+		XAUDIO2_RETURN_ON_FAIL(XAudio2System->GetDeviceCount(&OutNumOutputDevices));
+		return true;
+	}
+
+	bool FMixerPlatformXAudio2::GetOutputDeviceInfo(const uint32 InDeviceIndex, FAudioPlatformDeviceInfo& OutInfo)
+	{
+		if (!bIsInitialized)
+		{
+			AUDIO_PLATFORM_ERROR(TEXT("XAudio2 was not initialized."));
+			return false;
+		}
+
+		check(XAudio2System);
+
+		XAUDIO2_DEVICE_DETAILS DeviceDetails;
+		XAUDIO2_RETURN_ON_FAIL(XAudio2System->GetDeviceDetails(InDeviceIndex, &DeviceDetails));
+
+		OutInfo.Name = FString(DeviceDetails.DisplayName);
+		OutInfo.bIsSystemDefault = (InDeviceIndex == 0);
+
+		// Get the wave format to parse there rest of the device details
+		const WAVEFORMATEX& WaveFormatEx = DeviceDetails.OutputFormat.Format;
+		OutInfo.SampleRate = WaveFormatEx.nSamplesPerSec;
+		OutInfo.NumChannels = WaveFormatEx.nChannels;
+
+		// XAudio2 automatically converts the audio format to output device us so we don't need to do any format conversions
+		OutInfo.Format = EAudioMixerStreamDataFormat::Float;
+
+		OutInfo.OutputChannelArray.Reset();
+
+		// Extensible format supports surround sound so we need to parse the channel configuration to build our channel output array
+		if (WaveFormatEx.wFormatTag == WAVE_FORMAT_EXTENSIBLE)
+		{
+			// Cast to the extensible format to get access to extensible data
+			const WAVEFORMATEXTENSIBLE* WaveFormatExtensible = (WAVEFORMATEXTENSIBLE*)&WaveFormatEx;
+
+			// Loop through the extensible format channel flags in the standard order and build our output channel array
+			// From https://msdn.microsoft.com/en-us/library/windows/hardware/dn653308(v=vs.85).aspx
+			// The channels in the interleaved stream corresponding to these spatial positions must appear in the order specified above. This holds true even in the 
+			// case of a non-contiguous subset of channels. For example, if a stream contains left, bass enhance and right, then channel 1 is left, channel 2 is right, 
+			// and channel 3 is bass enhance. This enables the linkage of multi-channel streams to well-defined multi-speaker configurations.
+
+			uint32 ChanCount = 0;
+			for (uint32 ChannelTypeIndex = 0; ChannelTypeIndex < EAudioMixerChannel::ChannelTypeCount && ChanCount < WaveFormatEx.nChannels; ++ChannelTypeIndex)
+			{
+				if (WaveFormatExtensible->dwChannelMask & ChannelTypeMap[ChannelTypeIndex])
+				{
+					++ChanCount;
+					OutInfo.OutputChannelArray.Add((EAudioMixerChannel::Type)ChannelTypeIndex);
+				}
+			}
+
+			// Make sure we got the right number of output channels for the device
+			checkf(ChanCount == OutInfo.NumChannels, TEXT("Unknown channel type or channel not in proper order."));
+		}
+		else
+		{
+			// Non-extensible formats only support mono or stereo channel output
+			OutInfo.OutputChannelArray.Add(EAudioMixerChannel::FrontLeft);
+			if (OutInfo.NumChannels == 2)
+			{
+				OutInfo.OutputChannelArray.Add(EAudioMixerChannel::FrontRight);
+			}
+		}
+
+		UE_LOG(LogAudioMixerDebug, Log, TEXT("Audio Device Output Speaker Info:"), OutInfo.NumChannels);
+		UE_LOG(LogAudioMixerDebug, Log, TEXT("Name: %s"), *OutInfo.Name);
+		UE_LOG(LogAudioMixerDebug, Log, TEXT("Is Default: %s"), OutInfo.bIsSystemDefault ? TEXT("Yes") : TEXT("No"));
+		UE_LOG(LogAudioMixerDebug, Log, TEXT("Sample Rate: %d"), OutInfo.SampleRate);
+		UE_LOG(LogAudioMixerDebug, Log, TEXT("Channel Count: %d"), OutInfo.NumChannels);
+		UE_LOG(LogAudioMixerDebug, Log, TEXT("Channel Order:"));
+		for (int32 i = 0; i < OutInfo.NumChannels; ++i)
+		{
+			UE_LOG(LogAudioMixerDebug, Log, TEXT("%d: %s"), i, EAudioMixerChannel::ToString(OutInfo.OutputChannelArray[i]));
+		}
+
+		return true;
+	}
+
+	bool FMixerPlatformXAudio2::GetDefaultOutputDeviceIndex(uint32& OutDefaultDeviceIndex) const
+	{
+		OutDefaultDeviceIndex = 0;
+		return true;
+	}
+
+	bool FMixerPlatformXAudio2::OpenAudioStream(const FAudioMixerOpenStreamParams& Params)
+	{
+		if (!bIsInitialized)
+		{
+			AUDIO_PLATFORM_ERROR(TEXT("XAudio2 was not initialized."));
+			return false;
+		}
+
+		if (bIsDeviceOpen)
+		{
+			AUDIO_PLATFORM_ERROR(TEXT("XAudio2 audio stream already opened."));
+			return false;
+		}
+
+		checkf(Params.OutputDeviceIndex != INDEX_NONE, TEXT("OpenAudioStream params must specify a output device index."));
+		check(XAudio2System);
+		check(OutputAudioStreamMasteringVoice == nullptr);
+
+		WAVEFORMATEX Format = { 0 };
+
+		AudioStreamInfo.Reset();
+
+		AudioStreamInfo.OutputDeviceIndex = Params.OutputDeviceIndex;
+		AudioStreamInfo.NumOutputFrames = Params.NumFrames;
+		AudioStreamInfo.RequestedSampleRate = Params.SampleRate;
+		AudioStreamInfo.OutputDeviceIndex = Params.OutputDeviceIndex;
+		AudioStreamInfo.AudioMixer = Params.AudioMixer;
+
+		if (!GetOutputDeviceInfo(AudioStreamInfo.OutputDeviceIndex, AudioStreamInfo.DeviceInfo))
+		{
+			return false;
+		}
+
+		AudioStreamInfo.DeviceInfo.NumFrames = Params.NumFrames;
+		AudioStreamInfo.DeviceInfo.NumSamples = AudioStreamInfo.DeviceInfo.NumFrames * AudioStreamInfo.DeviceInfo.NumChannels;
+		AudioStreamInfo.DeviceInfo.SampleRate = AudioStreamInfo.RequestedSampleRate;
+
+		HRESULT Result = XAudio2System->CreateMasteringVoice(&OutputAudioStreamMasteringVoice, AudioStreamInfo.DeviceInfo.NumChannels, AudioStreamInfo.RequestedSampleRate, 0, AudioStreamInfo.OutputDeviceIndex, nullptr);
+		XAUDIO2_CLEANUP_ON_FAIL(Result);
+
+		// Xaudio2 on windows, no need for byte swap
+		AudioStreamInfo.bPerformByteSwap = false;
+
+		// Start the xaudio2 engine running, which will now allow us to start feeding audio to it
+		XAudio2System->StartEngine();
+
+		// Setup the format of the output source voice
+		Format.nChannels = AudioStreamInfo.DeviceInfo.NumChannels;
+		Format.nSamplesPerSec = AudioStreamInfo.DeviceInfo.SampleRate;
+		Format.wFormatTag = WAVE_FORMAT_IEEE_FLOAT;
+		Format.nAvgBytesPerSec = Format.nSamplesPerSec * sizeof(float) * Format.nChannels;
+		Format.nBlockAlign = sizeof(float) * Format.nChannels;
+		Format.wBitsPerSample = sizeof(float) * 8;
+
+		// Create the output source voice
+		Result = XAudio2System->CreateSourceVoice(&OutputAudioStreamSourceVoice, &Format, XAUDIO2_VOICE_NOSRC | XAUDIO2_VOICE_NOPITCH, 2.0f, &OutputVoiceCallback);
+		XAUDIO2_RETURN_ON_FAIL(Result);
+		
+		AudioStreamInfo.StreamState = EAudioOutputStreamState::Open;
+		bIsDeviceOpen = true;
+
+	Cleanup:
+		if (FAILED(Result))
+		{
+			CloseAudioStream();
+		}
+		return SUCCEEDED(Result);
+	}
+
+	FAudioPlatformDeviceInfo FMixerPlatformXAudio2::GetPlatformDeviceInfo() const
+	{
+		return AudioStreamInfo.DeviceInfo;
+	}
+
+	bool FMixerPlatformXAudio2::CloseAudioStream()
+	{
+		if (!bIsInitialized)
+		{
+			AUDIO_PLATFORM_ERROR(TEXT("XAudio2 was not initialized."));
+			return false;
+		}
+
+		if (bIsDeviceOpen)
+		{
+			if (!StopAudioStream())
+			{
+				return false;
+			}
+		}
+
+		check(XAudio2System);
+		XAudio2System->StopEngine();
+		check(OutputAudioStreamSourceVoice);
+		OutputAudioStreamSourceVoice->DestroyVoice();
+		OutputAudioStreamSourceVoice = nullptr;
+
+		check(OutputAudioStreamMasteringVoice);
+		OutputAudioStreamMasteringVoice->DestroyVoice();
+		OutputAudioStreamMasteringVoice = nullptr;
+
+		bIsDeviceOpen = false;
+
+		AudioStreamInfo.StreamState = EAudioOutputStreamState::Closed;
+
+		return true;
+	}
+
+	bool FMixerPlatformXAudio2::StartAudioStream()
+	{
+		// Start generating audio with our output source voice
+		BeginGeneratingAudio();
+
+		// If we already have a source voice, we can just restart it
+		if (OutputAudioStreamSourceVoice)
+		{
+			OutputAudioStreamSourceVoice->Start();
+			AudioStreamInfo.StreamState = EAudioOutputStreamState::Running;
+			return true;
+		}
+
+		return false;
+	}
+
+	bool FMixerPlatformXAudio2::StopAudioStream()
+	{
+		if (!bIsInitialized)
+		{
+			AUDIO_PLATFORM_ERROR(TEXT("XAudio2 was not initialized."));
+			return false;
+		}
+
+		check(XAudio2System);
+
+		// Signal that the thread that is running the update that we're stopping
+		if (!(AudioStreamInfo.StreamState == EAudioOutputStreamState::Stopped))
+		{
+			AudioStreamInfo.StreamState = EAudioOutputStreamState::Stopping;
+		}
+
+		if (OutputAudioStreamSourceVoice)
+		{
+			OutputAudioStreamSourceVoice->Stop(0);
+		}
+
+		check(AudioStreamInfo.StreamState == EAudioOutputStreamState::Stopped);
+
+		return true;
+	}
+
+	void FMixerPlatformXAudio2::SubmitBuffer(const TArray<float>& Buffer)
+	{
+		// Create a new xaudio2 buffer submission
+		XAUDIO2_BUFFER XAudio2Buffer = { 0 };
+		XAudio2Buffer.AudioBytes = AudioStreamInfo.DeviceInfo.NumSamples * sizeof(float);
+		XAudio2Buffer.pAudioData = (const BYTE*)Buffer.GetData();
+		XAudio2Buffer.pContext = this;
+
+		// Submit buffer to the output streaming voice
+		OutputAudioStreamSourceVoice->SubmitSourceBuffer(&XAudio2Buffer);
+	}
+
+	FName FMixerPlatformXAudio2::GetRuntimeFormat(USoundWave* InSoundWave)
+	{
+		if (InSoundWave->IsStreaming())
+		{
+			return FName(TEXT("OPUS"));
+		}
+
+		static FName NAME_OGG(TEXT("OGG"));
+		return NAME_OGG;
+	}
+
+	bool FMixerPlatformXAudio2::HasCompressedAudioInfoClass(USoundWave* InSoundWave)
+	{
+#if WITH_OGGVORBIS
+		return true;
+#else
+		return false;
+#endif
+	}
+
+	ICompressedAudioInfo* FMixerPlatformXAudio2::CreateCompressedAudioInfo(USoundWave* InSoundWave)
+	{
+		check(InSoundWave);
+
+		if (InSoundWave->IsStreaming())
+		{
+			return new FOpusAudioInfo();
+		}
+
+		ICompressedAudioInfo* CompressedInfo = new FVorbisAudioInfo();
+		if (!CompressedInfo)
+		{
+			UE_LOG(LogAudioMixer, Error, TEXT("Failed to create new FVorbisAudioInfo for SoundWave %s: out of memory."), *InSoundWave->GetName());
+			return nullptr;
+		}
+		return CompressedInfo;
+	}
+
+}
\ No newline at end of file
diff --git a/Engine/Source/Runtime/Windows/AudioMixerXAudio2/Private/AudioMixerPlatformXAudio2.h b/Engine/Source/Runtime/Windows/AudioMixerXAudio2/Private/AudioMixerPlatformXAudio2.h
new file mode 100644
--- /dev/null
+++ b/Engine/Source/Runtime/Windows/AudioMixerXAudio2/Private/AudioMixerPlatformXAudio2.h
@@ -0,0 +1,85 @@
+// Copyright 1998-2016 Epic Games, Inc. All Rights Reserved.
+
+#pragma once
+
+#include "AudioMixer.h"
+#include "AllowWindowsPlatformTypes.h"
+#include <xaudio2.h>
+#include "HideWindowsPlatformTypes.h"
+
+// Any platform defines
+namespace Audio
+{
+	class FMixerPlatformXAudio2;
+
+	/**
+	* FXAudio2VoiceCallback
+	* XAudio2 implementation of IXAudio2VoiceCallback
+	* This callback class is used to get event notifications on buffer end (when a buffer has finished processing).
+	* This is used to signal the I/O thread that it can request another buffer from the user callback.
+	*/
+	class FXAudio2VoiceCallback : public IXAudio2VoiceCallback
+	{
+	public:
+		FXAudio2VoiceCallback() {}
+		~FXAudio2VoiceCallback() {}
+
+	private:
+		void STDCALL OnVoiceProcessingPassStart(UINT32 BytesRequired) {}
+		void STDCALL OnVoiceProcessingPassEnd() {}
+		void STDCALL OnStreamEnd() {}
+		void STDCALL OnBufferStart(void* BufferContext) {}
+		void STDCALL OnLoopEnd(void* BufferContext) {}
+		void STDCALL OnVoiceError(void* BufferContext, HRESULT Error) {}
+
+		void STDCALL OnBufferEnd(void* BufferContext);
+
+	};
+
+	class FMixerPlatformXAudio2 : public IAudioMixerPlatformInterface
+	{
+
+	public:
+
+		FMixerPlatformXAudio2();
+		~FMixerPlatformXAudio2();
+
+		//~ Begin IAudioMixerPlatformInterface
+		EAudioMixerPlatformApi::Type GetPlatformApi() const override { return EAudioMixerPlatformApi::XAudio2; }
+		bool InitializeHardware() override;
+		bool TeardownHardware() override;
+		bool IsInitialized() const override;
+		bool GetNumOutputDevices(uint32& OutNumOutputDevices) override;
+		bool GetOutputDeviceInfo(const uint32 InDeviceIndex, FAudioPlatformDeviceInfo& OutInfo) override;
+		bool GetDefaultOutputDeviceIndex(uint32& OutDefaultDeviceIndex) const override;
+		bool OpenAudioStream(const FAudioMixerOpenStreamParams& Params) override;
+		bool CloseAudioStream() override;
+		bool StartAudioStream() override;
+		bool StopAudioStream() override;
+		FAudioPlatformDeviceInfo GetPlatformDeviceInfo() const override;
+		void SubmitBuffer(const TArray<float>& Buffer) override;
+		FName GetRuntimeFormat(USoundWave* InSoundWave) override;
+		bool HasCompressedAudioInfoClass(USoundWave* InSoundWave) override;
+		ICompressedAudioInfo* CreateCompressedAudioInfo(USoundWave* InSoundWave) override;
+		//~ End IAudioMixerPlatformInterface
+
+	private:
+
+		const TCHAR* GetErrorString(HRESULT Result);
+
+	private:
+		typedef TArray<long> TChannelTypeMap;
+
+		TChannelTypeMap ChannelTypeMap;
+		IXAudio2* XAudio2System;
+		IXAudio2MasteringVoice* OutputAudioStreamMasteringVoice;
+		IXAudio2SourceVoice* OutputAudioStreamSourceVoice;
+		FXAudio2VoiceCallback OutputVoiceCallback;
+		uint32 bIsComInitialized : 1;
+		uint32 bIsInitialized : 1;
+		uint32 bIsDeviceOpen : 1;
+
+	};
+
+}
+
diff --git a/Engine/Source/Runtime/Windows/AudioMixerXAudio2/Private/AudioMixerXAudio2.cpp b/Engine/Source/Runtime/Windows/AudioMixerXAudio2/Private/AudioMixerXAudio2.cpp
new file mode 100644
--- /dev/null
+++ b/Engine/Source/Runtime/Windows/AudioMixerXAudio2/Private/AudioMixerXAudio2.cpp
@@ -0,0 +1,25 @@
+// Copyright 1998-2016 Epic Games, Inc. All Rights Reserved.
+
+/**
+	Concrete implementation of FAudioDevice for XAudio2
+
+	See https://msdn.microsoft.com/en-us/library/windows/desktop/hh405049%28v=vs.85%29.aspx
+*/
+
+#include "AudioMixer.h"
+#include "AudioMixerDevice.h"
+#include "AudioMixerPlatformXAudio2.h"
+
+
+class FAudioMixerModuleXAudio2 : public IAudioDeviceModule
+{
+public:
+	virtual FAudioDevice* CreateAudioDevice() override
+	{
+		return new Audio::FMixerDevice(new Audio::FMixerPlatformXAudio2());
+	}
+};
+
+IMPLEMENT_MODULE(FAudioMixerModuleXAudio2, AudioMixerXAudio2);
+
+
diff --git a/Engine/Source/Runtime/Windows/XAudio2/Private/XAudio2Buffer.cpp b/Engine/Source/Runtime/Windows/XAudio2/Private/XAudio2Buffer.cpp
--- a/Engine/Source/Runtime/Windows/XAudio2/Private/XAudio2Buffer.cpp
+++ b/Engine/Source/Runtime/Windows/XAudio2/Private/XAudio2Buffer.cpp
@@ -75,11 +75,6 @@ struct FXMAInfo
 	FXAudio2SoundBuffer.
 ------------------------------------------------------------------------------------*/
 
-/** 
- * Constructor
- *
- * @param AudioDevice	audio device this sound buffer is going to be attached to.
- */
 FXAudio2SoundBuffer::FXAudio2SoundBuffer( FAudioDevice* InAudioDevice, ESoundFormat InSoundFormat )
 	: FSoundBuffer(InAudioDevice)
 	, SoundFormat(InSoundFormat)
@@ -101,11 +96,6 @@ FXAudio2SoundBuffer::FXAudio2SoundBuffer( FAudioDevice* InAudioDevice, ESoundFor
 	XWMA.XWMASeekDataSize = 0;
 }
 
-/**
- * Destructor 
- * 
- * Frees wave data and detaches itself from audio device.
- */
 FXAudio2SoundBuffer::~FXAudio2SoundBuffer( void )
 {
 	if( bAllocationInPermanentPool )
@@ -167,11 +157,6 @@ FXAudio2SoundBuffer::~FXAudio2SoundBuffer( void )
 	}
 }
 
-/**
- * Returns the size of this buffer in bytes.
- *
- * @return Size in bytes
- */
 int32 FXAudio2SoundBuffer::GetSize( void )
 {
 	int32 TotalSize = 0;
@@ -247,9 +232,7 @@ void FXAudio2SoundBuffer::EnsureRealtimeTaskCompletion()
 	}
 }
 
-/** 
- * Setup a WAVEFORMATEX structure
- */
+
 void FXAudio2SoundBuffer::InitWaveFormatEx( uint16 Format, USoundWave* Wave, bool bCheckPCMData )
 {
 	// Setup the format structure required for XAudio2
@@ -275,9 +258,6 @@ void FXAudio2SoundBuffer::InitWaveFormatEx( uint16 Format, USoundWave* Wave, boo
 	PCM.PCMFormat.nAvgBytesPerSec = NumChannels * sizeof( int16 ) * Wave->SampleRate;
 }
 
-/** 
- * Set up this buffer to contain and play XMA2 data
- */
 void FXAudio2SoundBuffer::InitXMA2( FXAudio2Device* XAudio2Device, USoundWave* Wave, FXMAInfo* XMAInfo )
 {
 #if XAUDIO_SUPPORTS_XMA2WAVEFORMATEX
@@ -309,9 +289,6 @@ void FXAudio2SoundBuffer::InitXMA2( FXAudio2Device* XAudio2Device, USoundWave* W
 #endif	//XAUDIO_SUPPORTS_XMA2WAVEFORMATEX
 }
 
-/** 
- * Set up this buffer to contain and play XWMA data
- */
 void FXAudio2SoundBuffer::InitXWMA( USoundWave* Wave, FXMAInfo* XMAInfo )
 {
 #if XAUDIO_SUPPORTS_XMA2WAVEFORMATEX
@@ -347,13 +324,6 @@ bool FXAudio2SoundBuffer::ReadCompressedInfo(USoundWave* SoundWave)
 	return DecompressionState->ReadCompressedInfo(SoundWave->ResourceData, SoundWave->ResourceSize, nullptr);
 }
 
-/**
- * Decompresses a chunk of compressed audio to the destination memory
- *
- * @param Destination		Memory to decompress to
- * @param bLooping			Whether to loop the sound seamlessly, or pad with zeroes
- * @return					Whether the sound looped or not
- */
 bool FXAudio2SoundBuffer::ReadCompressedData( uint8* Destination, bool bLooping )
 {
 	if (!DecompressionState)
@@ -381,13 +351,6 @@ void FXAudio2SoundBuffer::Seek( const float SeekTime )
 	}
 }
 
-/**
- * Static function used to create an OpenAL buffer and dynamically upload decompressed ogg vorbis data to.
- *
- * @param InWave		USoundWave to use as template and wave source
- * @param AudioDevice	audio device to attach created buffer to
- * @return FXAudio2SoundBuffer pointer if buffer creation succeeded, NULL otherwise
- */
 FXAudio2SoundBuffer* FXAudio2SoundBuffer::CreateQueuedBuffer( FXAudio2Device* XAudio2Device, USoundWave* Wave )
 {
 	check(XAudio2Device);
@@ -439,13 +402,6 @@ FXAudio2SoundBuffer* FXAudio2SoundBuffer::CreateQueuedBuffer( FXAudio2Device* XA
 	return Buffer;
 }
 
-/**
- * Static function used to create an OpenAL buffer and dynamically upload procedural data to.
- *
- * @param InWave		USoundWave to use as template and wave source
- * @param AudioDevice	audio device to attach created buffer to
- * @return FXAudio2SoundBuffer pointer if buffer creation succeeded, NULL otherwise
- */
 FXAudio2SoundBuffer* FXAudio2SoundBuffer::CreateProceduralBuffer( FXAudio2Device* XAudio2Device, USoundWave* Wave )
 {
 	// Always create a new buffer for real time decompressed sounds
@@ -464,13 +420,6 @@ FXAudio2SoundBuffer* FXAudio2SoundBuffer::CreateProceduralBuffer( FXAudio2Device
 	return( Buffer );
 }
 
-/**
- * Static function used to create an OpenAL buffer and upload raw PCM data to.
- *
- * @param InWave		USoundWave to use as template and wave source
- * @param AudioDevice	audio device to attach created buffer to
- * @return FXAudio2SoundBuffer pointer if buffer creation succeeded, NULL otherwise
- */
 FXAudio2SoundBuffer* FXAudio2SoundBuffer::CreatePreviewBuffer( FXAudio2Device* XAudio2Device, USoundWave* Wave, FXAudio2SoundBuffer* Buffer )
 {
 	FAudioDeviceManager* AudioDeviceManager = GEngine->GetAudioDeviceManager();
@@ -500,13 +449,6 @@ FXAudio2SoundBuffer* FXAudio2SoundBuffer::CreatePreviewBuffer( FXAudio2Device* X
 	return( Buffer );
 }
 
-/**
- * Static function used to create an OpenAL buffer and upload decompressed ogg vorbis data to.
- *
- * @param InWave		USoundWave to use as template and wave source
- * @param AudioDevice	audio device to attach created buffer to
- * @return FXAudio2SoundBuffer pointer if buffer creation succeeded, NULL otherwise
- */
 FXAudio2SoundBuffer* FXAudio2SoundBuffer::CreateNativeBuffer( FXAudio2Device* XAudio2Device, USoundWave* Wave )
 {
 	// Check to see if thread has finished decompressing on the other thread
@@ -575,13 +517,6 @@ FXAudio2SoundBuffer* FXAudio2SoundBuffer::CreateStreamingBuffer( FXAudio2Device*
 	return(Buffer);
 }
 
-/**
- * Static function used to create a buffer.
- *
- * @param InWave USoundWave to use as template and wave source
- * @param AudioDevice audio device to attach created buffer to
- * @return FXAudio2SoundBuffer pointer if buffer creation succeeded, NULL otherwise
- */
 FXAudio2SoundBuffer* FXAudio2SoundBuffer::Init( FAudioDevice* AudioDevice, USoundWave* Wave, bool bForceRealTime )
 {
 	// Can't create a buffer without any source data
@@ -667,5 +602,3 @@ FXAudio2SoundBuffer* FXAudio2SoundBuffer::Init( FAudioDevice* AudioDevice, USoun
 	return( Buffer );
 }
 
-// end
-
diff --git a/Engine/Source/Runtime/Windows/XAudio2/Private/XAudio2Source.cpp b/Engine/Source/Runtime/Windows/XAudio2/Private/XAudio2Source.cpp
--- a/Engine/Source/Runtime/Windows/XAudio2/Private/XAudio2Source.cpp
+++ b/Engine/Source/Runtime/Windows/XAudio2/Private/XAudio2Source.cpp
@@ -40,7 +40,6 @@ FXAudio2SoundSource::FXAudio2SoundSource(FAudioDevice* InAudioDevice)
 	, Source(nullptr)
 	, MaxEffectChainChannels(0)
 	, RealtimeAsyncTask(nullptr)
-	, VoiceId(-1)
 	, CurrentBuffer(0)
 	, bLoopCallback(false)
 	, bIsFinished(false)
@@ -583,6 +582,17 @@ bool FXAudio2SoundSource::Init(FWaveInstance* InWaveInstance)
 
 			Update();
 
+			// Initialize the total  number of frames of audio for this sound source
+			int32 NumBytes = InWaveInstance->WaveData->RawPCMDataSize;
+			NumTotalFrames = NumBytes / (Buffer->NumChannels * sizeof(int16));
+			StartFrame = 0;
+
+			if (InWaveInstance->StartTime > 0.0f)
+			{
+				const float StartFraction = InWaveInstance->StartTime / InWaveInstance->WaveData->GetDuration();
+				StartFrame = StartFraction * NumTotalFrames;
+			}
+
 			// Now set the source state to initialized so it can be played
 			// Initialization succeeded.
 			return true;
@@ -698,7 +708,7 @@ void FXAudio2SoundSource::GetMonoChannelVolumes(float ChannelVolumes[CHANNEL_MAT
 		}
 		check(AudioDevice->SpatializeProcessor != nullptr);
 
-		AudioDevice->SpatializeProcessor->SetSpatializationParameters(VoiceId, FAudioSpatializationParams(SpatializationParams.EmitterPosition, WaveInstance->Location));
+		AudioDevice->SpatializeProcessor->SetSpatializationParameters(VoiceId, SpatializationParams);
 		GetStereoChannelVolumes(ChannelVolumes, AttenuatedVolume);
 	}
 	else // Spatialize the mono stream using the normal 3d audio algorithm
@@ -868,7 +878,7 @@ void FXAudio2SoundSource::GetHexChannelVolumes(float ChannelVolumes[CHANNEL_MATR
 /** 
  * Maps a sound with a given number of channels to to expected speakers
  */
-void FXAudio2SoundSource::RouteDryToSpeakers(float ChannelVolumes[CHANNEL_MATRIX_COUNT])
+void FXAudio2SoundSource::RouteDryToSpeakers(float ChannelVolumes[CHANNEL_MATRIX_COUNT], const float InVolume)
 {
 	// Only need to account to the special cases that are not a simple match of channel to speaker
 	switch( Buffer->NumChannels )
@@ -890,6 +900,8 @@ void FXAudio2SoundSource::RouteDryToSpeakers(float ChannelVolumes[CHANNEL_MATRIX
 		break;
 
 		default:
+			// For all other channel counts, just apply the volume and let XAudio2 handle doing downmixing of the source
+			Source->SetVolume(InVolume);
 			break;
 	};
 }
@@ -1324,62 +1336,93 @@ void FXAudio2SoundSource::Update()
 	}
 	else
 	{
-		// Set the pitch on the xaudio2 source
-	AudioDevice->ValidateAPICall( TEXT( "SetFrequencyRatio" ), 
-		Source->SetFrequencyRatio( Pitch) );
+			// Set the pitch on the xaudio2 source
+		AudioDevice->ValidateAPICall( TEXT( "SetFrequencyRatio" ), 
+			Source->SetFrequencyRatio( Pitch) );
 
-	// Set whether to bleed to the rear speakers
-	SetStereoBleed();
+		// Set whether to bleed to the rear speakers
+		SetStereoBleed();
 
-	// Set the amount to bleed to the LFE speaker
-	SetLFEBleed();
+		// Set the amount to bleed to the LFE speaker
+		SetLFEBleed();
 
-	// Set the low pass filter frequency value
-	SetFilterFrequency();
+		// Set the low pass filter frequency value
+		SetFilterFrequency();
 
-	if (LastLPFFrequency != LPFFrequency)
-	{
-		// Apply the low pass filter
-		XAUDIO2_FILTER_PARAMETERS LPFParameters = { LowPassFilter, 1.0f, AudioDevice->GetLowPassFilterResonance() };
+		if (LastLPFFrequency != LPFFrequency)
+		{
+			// Apply the low pass filter
+			XAUDIO2_FILTER_PARAMETERS LPFParameters = { LowPassFilter, 1.0f, AudioDevice->GetLowPassFilterResonance() };
 
-		check(AudioDevice->SampleRate > 0.0f);
+			check(AudioDevice->SampleRate > 0.0f);
 
-		// Convert the frequency value to normalized radian frequency values where 0.0f to 2.0f sweeps 0.0hz to sample rate
-		// and 1.0f is the nyquist frequency. A normalized frequency of 1.0f is an effective bypass.
-		LPFParameters.Frequency = FMath::Clamp(2.0f * LPFFrequency / AudioDevice->SampleRate, 0.0f, 1.0f);
+			// Convert the frequency value to normalized radian frequency values where 0.0f to 2.0f sweeps 0.0hz to sample rate
+			// and 1.0f is the nyquist frequency. A normalized frequency of 1.0f is an effective bypass.
+			LPFParameters.Frequency = FMath::Clamp(2.0f * LPFFrequency / AudioDevice->SampleRate, 0.0f, 1.0f);
 
-		AudioDevice->ValidateAPICall(TEXT("SetFilterParameters"),
-									 Source->SetFilterParameters(&LPFParameters));
+			AudioDevice->ValidateAPICall(TEXT("SetFilterParameters"),
+										 Source->SetFilterParameters(&LPFParameters));
 
-		LastLPFFrequency = LPFFrequency;
-	}
+			LastLPFFrequency = LPFFrequency;
+		}
+
+		// Get the current xaudio2 source voice state to determine the number of frames played
+		XAUDIO2_VOICE_STATE VoiceState;
+		Source->GetState(&VoiceState);
 
-	// Initialize channel volumes
-	float ChannelVolumes[CHANNEL_MATRIX_COUNT] = { 0.0f };
+		// XAudio2's "samples" appear to actually be frames (1 interleaved time-slice)
+		NumFramesPlayed = VoiceState.SamplesPlayed;
 
-	const float Volume = FSoundSource::GetDebugVolume(WaveInstance->GetActualVolume());
+		// Initialize channel volumes
+		float ChannelVolumes[CHANNEL_MATRIX_COUNT] = { 0.0f };
 
-	GetChannelVolumes( ChannelVolumes, Volume );
+		const float Volume = FSoundSource::GetDebugVolume(WaveInstance->GetActualVolume());
 
-	// Send to the 5.1 channels
-	RouteDryToSpeakers( ChannelVolumes );
+		GetChannelVolumes( ChannelVolumes, Volume );
 
-	// Send to the reverb channel
-	if( bReverbApplied )
-	{
-		RouteToReverb( ChannelVolumes );
+		// Send to the 5.1 channels
+		RouteDryToSpeakers(ChannelVolumes, Volume);
+
+		// Send to the reverb channel
+		if( bReverbApplied )
+		{
+			RouteToReverb( ChannelVolumes );
+		}
+
+		// If this audio can have radio distortion applied, 
+		// send the volumes to the radio distortion voice. 
+		if( WaveInstance->bApplyRadioFilter )
+		{
+			RouteToRadio( ChannelVolumes );
+		}
 	}
 
-	// If this audio can have radio distortion applied, 
-	// send the volumes to the radio distortion voice. 
-	if( WaveInstance->bApplyRadioFilter )
+	FSoundSource::DrawDebugInfo();
+
+}
+
+float FXAudio2SoundSource::GetPlaybackPercent() const
+{
+	// If we didn't compute NumTotalFrames then there's no playback percent
+	if (NumTotalFrames == 0)
 	{
-		RouteToRadio( ChannelVolumes );
-	}
+		return 0.0f;
 	}
 
-	FSoundSource::DrawDebugInfo();
+	const int32 CurrentFrame = (StartFrame + NumFramesPlayed) % NumTotalFrames;
+
+	// Compute the percent based on frames played and total frames
+	const float Percent = (float)CurrentFrame / NumTotalFrames;
 
+	if (WaveInstance->LoopingMode == LOOP_Never)
+	{
+		return FMath::Clamp(Percent, 0.0f, 1.0f);
+	}
+	else
+	{
+		// Wrap the playback percent for looping sounds
+		return FMath::Fmod(Percent, 1.0f);
+	}
 }
 
 void FXAudio2SoundSource::Play()
@@ -1544,7 +1587,7 @@ bool FXAudio2SoundSource::IsFinished()
 	// A paused source is not finished.
 	if (Paused || !bInitialized)
 	{
-		return(false);
+		return false;
 	}
 
 
@@ -1553,21 +1596,21 @@ bool FXAudio2SoundSource::IsFinished()
 		return true;
 	}
 
-		if (bIsFinished)
-		{
-			WaveInstance->NotifyFinished();
-			return true;
-		}
-
-		if (bLoopCallback && WaveInstance->LoopingMode == LOOP_WithNotification)
-		{
-			WaveInstance->NotifyFinished();
-			bLoopCallback = false;
-		}
+	if (bIsFinished)
+	{
+		WaveInstance->NotifyFinished();
+		return true;
+	}
 
-		return false;
+	if (bLoopCallback && WaveInstance->LoopingMode == LOOP_WithNotification)
+	{
+		WaveInstance->NotifyFinished();
+		bLoopCallback = false;
 	}
 
+	return false;
+}
+
 bool FXAudio2SoundSource::IsUsingHrtfSpatializer()
 {
 	return bUsingHRTFSpatialization;
diff --git a/Engine/Source/Runtime/Windows/XAudio2/Private/XAudio2Support.h b/Engine/Source/Runtime/Windows/XAudio2/Private/XAudio2Support.h
--- a/Engine/Source/Runtime/Windows/XAudio2/Private/XAudio2Support.h
+++ b/Engine/Source/Runtime/Windows/XAudio2/Private/XAudio2Support.h
@@ -265,7 +265,7 @@ public:
 	 * @param bLooping			Whether to loop the sound seamlessly, or pad with zeroes
 	 * @return					Whether the sound looped or not
 	 */
-	bool ReadCompressedData( uint8* Destination, bool bLooping );
+	bool ReadCompressedData( uint8* Destination, bool bLooping ) override;
 
 	/**
 	 * Sets the point in time within the buffer to the specified time
@@ -273,7 +273,7 @@ public:
 	 *
 	 * @param SeekTime		Time in seconds from the beginning of sound to seek to
 	 */
-	void Seek( const float SeekTime );
+	void Seek( const float SeekTime ) override;
 
 	/**
 	 * Static function used to create an OpenAL buffer and dynamically upload decompressed ogg vorbis data to.
@@ -438,6 +438,9 @@ public:
 	 */
 	virtual void Update();
 
+	/** Returns the playback percent */
+	virtual float GetPlaybackPercent() const override;
+
 	/**
 	 * Plays the current wave instance.	
 	 */
@@ -514,7 +517,7 @@ public:
 	/** 
 	 * Maps a sound with a given number of channels to to expected speakers
 	 */
-	void RouteDryToSpeakers(float ChannelVolumes[CHANNEL_MATRIX_COUNT]);
+	void RouteDryToSpeakers(float ChannelVolumes[CHANNEL_MATRIX_COUNT], const float InVolume);
 
 	/** 
 	 * Maps the sound to the relevant reverb effect
@@ -613,9 +616,6 @@ protected:
 	/** Additional buffer info for XWMA sounds */
 	XAUDIO2_BUFFER_WMA XAudio2BufferXWMA[1];
 
-	/** Index of this sound source in the audio device sound source array. */
-	uint32 VoiceId;
-
 	/** Which sound buffer should be written to next - used for triple buffering. */
 	int32 CurrentBuffer;
 
@@ -1000,7 +1000,7 @@ struct FXAudioDeviceProperties : public IDeviceChangedListener
 	#define UE4_XAUDIO2_CHANNELMASK		FXAudioDeviceProperties::DeviceDetails.OutputFormat.dwChannelMask
 	#define UE4_XAUDIO2_SAMPLERATE		FXAudioDeviceProperties::DeviceDetails.OutputFormat.Format.nSamplesPerSec
 #else	//XAUDIO_SUPPORTS_DEVICE_DETAILS
-	#define UE4_XAUDIO2_NUMCHANNELS		8		// Up to 7.1 supported
-	#define UE4_XAUDIO2_CHANNELMASK		3		// Default to left and right speakers...
-	#define UE4_XAUDIO2_SAMPLERATE		44100	// Default to CD sample rate
+#define UE4_XAUDIO2_NUMCHANNELS		6
+#define UE4_XAUDIO2_CHANNELMASK		SPEAKER_5POINT1
+#define UE4_XAUDIO2_SAMPLERATE		44100
 #endif	//XAUDIO_SUPPORTS_DEVICE_DETAILS
diff --git a/Templates/FP_FirstPerson/Source/FP_FirstPerson/FP_FirstPersonGameMode.h b/Templates/FP_FirstPerson/Source/FP_FirstPerson/FP_FirstPersonGameMode.h
--- a/Templates/FP_FirstPerson/Source/FP_FirstPerson/FP_FirstPersonGameMode.h
+++ b/Templates/FP_FirstPerson/Source/FP_FirstPerson/FP_FirstPersonGameMode.h
@@ -1,10 +1,10 @@
 // Copyright 1998-2016 Epic Games, Inc. All Rights Reserved.
 #pragma once
-#include "GameFramework/GameMode.h"
+#include "GameFramework/GameModeBase.h"
 #include "FP_FirstPersonGameMode.generated.h"
 
 UCLASS(minimalapi)
-class AFP_FirstPersonGameMode : public AGameMode
+class AFP_FirstPersonGameMode : public AGameModeBase
 {
 	GENERATED_BODY()
 
diff --git a/Templates/TP_2DSideScroller/Source/TP_2DSideScroller/TP_2DSideScrollerGameMode.h b/Templates/TP_2DSideScroller/Source/TP_2DSideScroller/TP_2DSideScrollerGameMode.h
--- a/Templates/TP_2DSideScroller/Source/TP_2DSideScroller/TP_2DSideScrollerGameMode.h
+++ b/Templates/TP_2DSideScroller/Source/TP_2DSideScroller/TP_2DSideScrollerGameMode.h
@@ -1,7 +1,7 @@
 // Copyright 1998-2016 Epic Games, Inc. All Rights Reserved.
 
 #pragma once
-#include "GameFramework/GameMode.h"
+#include "GameFramework/GameModeBase.h"
 #include "TP_2DSideScrollerGameMode.generated.h"
 
 // The GameMode defines the game being played. It governs the game rules, scoring, what actors
@@ -10,7 +10,7 @@
 // This game mode just sets the default pawn to be the MyCharacter asset, which is a subclass of TP_2DSideScrollerCharacter
 
 UCLASS(minimalapi)
-class ATP_2DSideScrollerGameMode : public AGameMode
+class ATP_2DSideScrollerGameMode : public AGameModeBase
 {
 	GENERATED_BODY()
 public:
diff --git a/Templates/TP_FirstPerson/Source/TP_FirstPerson/TP_FirstPersonGameMode.h b/Templates/TP_FirstPerson/Source/TP_FirstPerson/TP_FirstPersonGameMode.h
--- a/Templates/TP_FirstPerson/Source/TP_FirstPerson/TP_FirstPersonGameMode.h
+++ b/Templates/TP_FirstPerson/Source/TP_FirstPerson/TP_FirstPersonGameMode.h
@@ -1,10 +1,10 @@
 // Copyright 1998-2016 Epic Games, Inc. All Rights Reserved.
 #pragma once
-#include "GameFramework/GameMode.h"
+#include "GameFramework/GameModeBase.h"
 #include "TP_FirstPersonGameMode.generated.h"
 
 UCLASS(minimalapi)
-class ATP_FirstPersonGameMode : public AGameMode
+class ATP_FirstPersonGameMode : public AGameModeBase
 {
 	GENERATED_BODY()
 
diff --git a/Templates/TP_Flying/Source/TP_Flying/TP_FlyingGameMode.h b/Templates/TP_Flying/Source/TP_Flying/TP_FlyingGameMode.h
--- a/Templates/TP_Flying/Source/TP_Flying/TP_FlyingGameMode.h
+++ b/Templates/TP_Flying/Source/TP_Flying/TP_FlyingGameMode.h
@@ -1,10 +1,10 @@
 // Copyright 1998-2016 Epic Games, Inc. All Rights Reserved.
 #pragma once
-#include "GameFramework/GameMode.h"
+#include "GameFramework/GameModeBase.h"
 #include "TP_FlyingGameMode.generated.h"
 
 UCLASS(minimalapi)
-class ATP_FlyingGameMode : public AGameMode
+class ATP_FlyingGameMode : public AGameModeBase
 {
 	GENERATED_BODY()
 
diff --git a/Templates/TP_Puzzle/Source/TP_Puzzle/TP_PuzzleGameMode.h b/Templates/TP_Puzzle/Source/TP_Puzzle/TP_PuzzleGameMode.h
--- a/Templates/TP_Puzzle/Source/TP_Puzzle/TP_PuzzleGameMode.h
+++ b/Templates/TP_Puzzle/Source/TP_Puzzle/TP_PuzzleGameMode.h
@@ -1,11 +1,11 @@
 // Copyright 1998-2016 Epic Games, Inc. All Rights Reserved.
 #pragma once
-#include "GameFramework/GameMode.h"
+#include "GameFramework/GameModeBase.h"
 #include "TP_PuzzleGameMode.generated.h"
 
 /** GameMode class to specify pawn and playercontroller */
 UCLASS(minimalapi)
-class ATP_PuzzleGameMode : public AGameMode
+class ATP_PuzzleGameMode : public AGameModeBase
 {
 	GENERATED_BODY()
 
diff --git a/Templates/TP_Rolling/Source/TP_Rolling/TP_RollingGameMode.h b/Templates/TP_Rolling/Source/TP_Rolling/TP_RollingGameMode.h
--- a/Templates/TP_Rolling/Source/TP_Rolling/TP_RollingGameMode.h
+++ b/Templates/TP_Rolling/Source/TP_Rolling/TP_RollingGameMode.h
@@ -1,10 +1,10 @@
 // Copyright 1998-2016 Epic Games, Inc. All Rights Reserved.
 #pragma once
-#include "GameFramework/GameMode.h"
+#include "GameFramework/GameModeBase.h"
 #include "TP_RollingGameMode.generated.h"
 
 UCLASS(minimalapi)
-class ATP_RollingGameMode : public AGameMode
+class ATP_RollingGameMode : public AGameModeBase
 {
 	GENERATED_BODY()
 
diff --git a/Templates/TP_SideScroller/Source/TP_SideScroller/TP_SideScrollerGameMode.h b/Templates/TP_SideScroller/Source/TP_SideScroller/TP_SideScrollerGameMode.h
--- a/Templates/TP_SideScroller/Source/TP_SideScroller/TP_SideScrollerGameMode.h
+++ b/Templates/TP_SideScroller/Source/TP_SideScroller/TP_SideScrollerGameMode.h
@@ -1,10 +1,10 @@
 // Copyright 1998-2016 Epic Games, Inc. All Rights Reserved.
 #pragma once
-#include "GameFramework/GameMode.h"
+#include "GameFramework/GameModeBase.h"
 #include "TP_SideScrollerGameMode.generated.h"
 
 UCLASS(minimalapi)
-class ATP_SideScrollerGameMode : public AGameMode
+class ATP_SideScrollerGameMode : public AGameModeBase
 {
 	GENERATED_BODY()
 
diff --git a/Templates/TP_ThirdPerson/Source/TP_ThirdPerson/TP_ThirdPersonGameMode.h b/Templates/TP_ThirdPerson/Source/TP_ThirdPerson/TP_ThirdPersonGameMode.h
--- a/Templates/TP_ThirdPerson/Source/TP_ThirdPerson/TP_ThirdPersonGameMode.h
+++ b/Templates/TP_ThirdPerson/Source/TP_ThirdPerson/TP_ThirdPersonGameMode.h
@@ -1,10 +1,10 @@
 // Copyright 1998-2016 Epic Games, Inc. All Rights Reserved.
 #pragma once
-#include "GameFramework/GameMode.h"
+#include "GameFramework/GameModeBase.h"
 #include "TP_ThirdPersonGameMode.generated.h"
 
 UCLASS(minimalapi)
-class ATP_ThirdPersonGameMode : public AGameMode
+class ATP_ThirdPersonGameMode : public AGameModeBase
 {
 	GENERATED_BODY()
 
diff --git a/Templates/TP_TopDown/Source/TP_TopDown/TP_TopDownGameMode.h b/Templates/TP_TopDown/Source/TP_TopDown/TP_TopDownGameMode.h
--- a/Templates/TP_TopDown/Source/TP_TopDown/TP_TopDownGameMode.h
+++ b/Templates/TP_TopDown/Source/TP_TopDown/TP_TopDownGameMode.h
@@ -1,10 +1,10 @@
 // Copyright 1998-2016 Epic Games, Inc. All Rights Reserved.
 #pragma once
-#include "GameFramework/GameMode.h"
+#include "GameFramework/GameModeBase.h"
 #include "TP_TopDownGameMode.generated.h"
 
 UCLASS(minimalapi)
-class ATP_TopDownGameMode : public AGameMode
+class ATP_TopDownGameMode : public AGameModeBase
 {
 	GENERATED_BODY()
 
diff --git a/Templates/TP_TwinStick/Source/TP_TwinStick/TP_TwinStickGameMode.h b/Templates/TP_TwinStick/Source/TP_TwinStick/TP_TwinStickGameMode.h
--- a/Templates/TP_TwinStick/Source/TP_TwinStick/TP_TwinStickGameMode.h
+++ b/Templates/TP_TwinStick/Source/TP_TwinStick/TP_TwinStickGameMode.h
@@ -1,12 +1,12 @@
 // Copyright 1998-2016 Epic Games, Inc. All Rights Reserved.
 #pragma once
 
-#include "GameFramework/GameMode.h"
+#include "GameFramework/GameModeBase.h"
 #include "TP_TwinStickGameMode.generated.h"
 
 
 UCLASS(minimalapi)
-class ATP_TwinStickGameMode : public AGameMode
+class ATP_TwinStickGameMode : public AGameModeBase
 {
 	GENERATED_BODY()
 
diff --git a/Templates/TP_Vehicle/Source/TP_Vehicle/TP_VehicleGameMode.h b/Templates/TP_Vehicle/Source/TP_Vehicle/TP_VehicleGameMode.h
--- a/Templates/TP_Vehicle/Source/TP_Vehicle/TP_VehicleGameMode.h
+++ b/Templates/TP_Vehicle/Source/TP_Vehicle/TP_VehicleGameMode.h
@@ -1,10 +1,10 @@
 // Copyright 1998-2016 Epic Games, Inc. All Rights Reserved.
 #pragma once
-#include "GameFramework/GameMode.h"
+#include "GameFramework/GameModeBase.h"
 #include "TP_VehicleGameMode.generated.h"
 
 UCLASS(minimalapi)
-class ATP_VehicleGameMode : public AGameMode
+class ATP_VehicleGameMode : public AGameModeBase
 {
 	GENERATED_BODY()
 
diff --git a/Templates/TP_VehicleAdv/Source/TP_VehicleAdv/TP_VehicleAdvGameMode.h b/Templates/TP_VehicleAdv/Source/TP_VehicleAdv/TP_VehicleAdvGameMode.h
--- a/Templates/TP_VehicleAdv/Source/TP_VehicleAdv/TP_VehicleAdvGameMode.h
+++ b/Templates/TP_VehicleAdv/Source/TP_VehicleAdv/TP_VehicleAdvGameMode.h
@@ -1,10 +1,10 @@
 // Copyright 1998-2016 Epic Games, Inc. All Rights Reserved.
 #pragma once
-#include "GameFramework/GameMode.h"
+#include "GameFramework/GameModeBase.h"
 #include "TP_VehicleAdvGameMode.generated.h"
 
 UCLASS(minimalapi)
-class ATP_VehicleAdvGameMode : public AGameMode
+class ATP_VehicleAdvGameMode : public AGameModeBase
 {
 	GENERATED_BODY()
 
