UE4 - Change Warning to Display
#tests compile

diff --git a/Engine/Source/Runtime/Windows/D3D11RHI/Private/Windows/WindowsD3D11Device.cpp b/Engine/Source/Runtime/Windows/D3D11RHI/Private/Windows/WindowsD3D11Device.cpp
--- a/Engine/Source/Runtime/Windows/D3D11RHI/Private/Windows/WindowsD3D11Device.cpp
+++ b/Engine/Source/Runtime/Windows/D3D11RHI/Private/Windows/WindowsD3D11Device.cpp
@@ -941,7 +941,7 @@ void FD3D11DynamicRHI::InitD3DDevice()
 			// Note: 0x1300 corresponds to Maxwell hardware or above
 			if ( GRHIDeviceId < 0x1300 ) 
 			{
-				UE_LOG(LogD3D11RHI, Warning, TEXT("Timestamp queries are currently disabled on this hardware due to instability. Realtime GPU stats will not be available. You can override this behaviour by setting r.NVIDIATimestampWorkaround to 0"));
+				UE_LOG(LogD3D11RHI, Display, TEXT("Timestamp queries are currently disabled on this hardware due to instability. Realtime GPU stats will not be available. You can override this behaviour by setting r.NVIDIATimestampWorkaround to 0"));
 				GSupportsTimestampRenderQueries = false;
 			}
 		}
