Copying //UE4/Dev-AnimPhys to //UE4/Dev-Main (Source: //UE4/Dev-AnimPhys @ 3624379)

diff --git a/Engine/Source/Runtime/Windows/XAudio2/Private/XAudio2Device.cpp b/Engine/Source/Runtime/Windows/XAudio2/Private/XAudio2Device.cpp
--- a/Engine/Source/Runtime/Windows/XAudio2/Private/XAudio2Device.cpp
+++ b/Engine/Source/Runtime/Windows/XAudio2/Private/XAudio2Device.cpp
@@ -62,10 +62,6 @@ const float* FXAudioDeviceProperties::OutputMixMatrix = NULL;
 XAUDIO2_DEVICE_DETAILS FXAudioDeviceProperties::DeviceDetails;
 #endif	//XAUDIO_SUPPORTS_DEVICE_DETAILS
 
-#if PLATFORM_WINDOWS
-FMMNotificationClient* FXAudioDeviceProperties::NotificationClient = nullptr;
-#endif
-
 /*------------------------------------------------------------------------------------
 	FAudioDevice Interface.
 ------------------------------------------------------------------------------------*/
@@ -291,6 +287,34 @@ void FXAudio2Device::TeardownHardware()
 
 void FXAudio2Device::UpdateHardware()
 {
+	// If the audio device changed, we need to tear down and restart the audio engine state
+	if (DeviceProperties->DidAudioDeviceChange())
+	{
+		// Flush stops all sources so sources can be safely deleted below.
+		Flush(nullptr);
+
+		// Remove the effects manager
+		if (Effects)
+		{
+			delete Effects;
+			Effects = nullptr;
+		}
+	
+		// Teardown hardware
+		TeardownHardware();
+		
+		// Restart the hardware
+		InitializeHardware();
+
+		// Recreate the effects manager
+		Effects = CreateEffectsManager();
+
+		// Now reset and restart the sound source objects
+		FreeSources.Reset();
+		Sources.Reset();
+
+		InitSoundSources();
+	}
 }
 
 void FXAudio2Device::UpdateAudioClock()
diff --git a/Engine/Source/Runtime/Windows/XAudio2/Private/XAudio2Support.h b/Engine/Source/Runtime/Windows/XAudio2/Private/XAudio2Support.h
--- a/Engine/Source/Runtime/Windows/XAudio2/Private/XAudio2Support.h
+++ b/Engine/Source/Runtime/Windows/XAudio2/Private/XAudio2Support.h
@@ -70,6 +70,10 @@ public:
 
 	HRESULT STDMETHODCALLTYPE OnDefaultDeviceChanged(EDataFlow flow, ERole role, LPCWSTR pwstrDeviceId) override
 	{
+		for (IDeviceChangedListener* Listener : Listeners)
+		{
+			Listener->OnDefaultDeviceChanged();
+		}
 		return S_OK;
 	}
 
@@ -753,7 +757,7 @@ struct FXAudioDeviceProperties final : public IDeviceChangedListener
 #endif	//XAUDIO_SUPPORTS_DEVICE_DETAILS
 
 #if PLATFORM_WINDOWS
-	static FMMNotificationClient* NotificationClient;
+	FMMNotificationClient* NotificationClient;
 #endif
 
 	// For calculating speaker maps for 3d audio
@@ -781,15 +785,7 @@ struct FXAudioDeviceProperties final : public IDeviceChangedListener
 		, bAllowNewVoices(true)
 	{
 #if PLATFORM_WINDOWS
-		if (NotificationClient == nullptr)
-		{
-			NotificationClient = new FMMNotificationClient();
-	}
-		else
-		{
-			NotificationClient->AddRef();
-		}
-
+		NotificationClient = new FMMNotificationClient();
 		NotificationClient->RegisterDeviceChangedListener(this);
 #endif
 	}
@@ -797,8 +793,12 @@ struct FXAudioDeviceProperties final : public IDeviceChangedListener
 	virtual ~FXAudioDeviceProperties()
 	{
 #if PLATFORM_WINDOWS
-		NotificationClient->UnRegisterDeviceDeviceChangedListener(this);
-		NotificationClient->Release();
+		if (NotificationClient)
+		{
+			NotificationClient->UnRegisterDeviceDeviceChangedListener(this);
+			NotificationClient->Release();
+			NotificationClient = nullptr;
+		}
 #endif
 
 		// Make sure we've free'd all of our active voices at this point!
@@ -852,6 +852,14 @@ struct FXAudioDeviceProperties final : public IDeviceChangedListener
 #endif // XAUDIO_SUPPORTS_DEVICE_DETAILS
 	}
 
+	void OnDefaultDeviceChanged() override
+	{
+		bDeviceChanged = true;
+
+		// Immediately disallow new voices to be created
+		bAllowNewVoices = false;
+	}
+
 	bool DidAudioDeviceChange()
 	{
 		bool bChanged = bDeviceChanged;
