Fixed compilation error in VS 2017 due to bad member function declaration syntax

diff --git a/Engine/Source/Runtime/Windows/D3D11RHI/Private/D3D11RHIPrivate.h b/Engine/Source/Runtime/Windows/D3D11RHI/Private/D3D11RHIPrivate.h
--- a/Engine/Source/Runtime/Windows/D3D11RHI/Private/D3D11RHIPrivate.h
+++ b/Engine/Source/Runtime/Windows/D3D11RHI/Private/D3D11RHIPrivate.h
@@ -337,8 +337,8 @@ public:
 	virtual void RHIBeginUpdateMultiFrameResource(FTextureRHIParamRef Texture) override;
 	virtual void RHIEndUpdateMultiFrameResource(FTextureRHIParamRef Texture) override;
 
-	virtual void FD3D11DynamicRHI::RHIBeginUpdateMultiFrameResource(FUnorderedAccessViewRHIParamRef UAV) override;
-	virtual void FD3D11DynamicRHI::RHIEndUpdateMultiFrameResource(FUnorderedAccessViewRHIParamRef UAV) override;
+	virtual void RHIBeginUpdateMultiFrameResource(FUnorderedAccessViewRHIParamRef UAV) override;
+	virtual void RHIEndUpdateMultiFrameResource(FUnorderedAccessViewRHIParamRef UAV) override;
 
 	virtual FSamplerStateRHIRef RHICreateSamplerState(const FSamplerStateInitializerRHI& Initializer) final override;
 	virtual FRasterizerStateRHIRef RHICreateRasterizerState(const FRasterizerStateInitializerRHI& Initializer) final override;
