Fixup for deterministic audio generation issue.
Submitted on behalf of Rich.Whitehouse
#test prefilght automated test.

diff --git a/Engine/Source/Runtime/Android/AndroidAudio/Private/AndroidAudioBuffer.cpp b/Engine/Source/Runtime/Android/AndroidAudio/Private/AndroidAudioBuffer.cpp
--- a/Engine/Source/Runtime/Android/AndroidAudio/Private/AndroidAudioBuffer.cpp
+++ b/Engine/Source/Runtime/Android/AndroidAudio/Private/AndroidAudioBuffer.cpp
@@ -163,6 +163,8 @@ FSLESSoundBuffer* FSLESSoundBuffer::CreateNativeBuffer( FSLESAudioDevice* AudioD
 	Buffer->NumChannels		= InWave->NumChannels;
 	Buffer->SampleRate		= InWave->GetSampleRateForCurrentPlatform();
 
+	check(!InWave->RawPCMData || InWave->RawPCMDataSize);
+
 	// Take ownership the PCM data
 	Buffer->AudioData = InWave->RawPCMData;
 	Buffer->BufferSize = InWave->RawPCMDataSize;
diff --git a/Engine/Source/Runtime/AudioMixer/Private/AudioMixerBuffer.cpp b/Engine/Source/Runtime/AudioMixer/Private/AudioMixerBuffer.cpp
--- a/Engine/Source/Runtime/AudioMixer/Private/AudioMixerBuffer.cpp
+++ b/Engine/Source/Runtime/AudioMixer/Private/AudioMixerBuffer.cpp
@@ -22,6 +22,7 @@ namespace Audio
 	{
 		if (InWave->DecompressionType == EDecompressionType::DTYPE_Native || InWave->DecompressionType == EDecompressionType::DTYPE_Preview)
 		{
+			check(!InWave->RawPCMData || InWave->RawPCMDataSize);
 			Data = InWave->RawPCMData;
 			DataSize = InWave->RawPCMDataSize;
  			InWave->RawPCMData = nullptr;
diff --git a/Engine/Source/Runtime/AudioMixer/Private/AudioMixerSource.cpp b/Engine/Source/Runtime/AudioMixer/Private/AudioMixerSource.cpp
--- a/Engine/Source/Runtime/AudioMixer/Private/AudioMixerSource.cpp
+++ b/Engine/Source/Runtime/AudioMixer/Private/AudioMixerSource.cpp
@@ -63,6 +63,7 @@ namespace Audio
 
 		if (InWaveInstance->WaveData->DecompressionType != DTYPE_Procedural)
 		{
+			check(!InWaveInstance->WaveData->RawPCMData || InWaveInstance->WaveData->RawPCMDataSize);
 			const int32 NumBytes = InWaveInstance->WaveData->RawPCMDataSize;
 			NumFrames = NumBytes / (InWaveInstance->WaveData->NumChannels * sizeof(int16));
 		}
diff --git a/Engine/Source/Runtime/Engine/Classes/Sound/SoundWave.h b/Engine/Source/Runtime/Engine/Classes/Sound/SoundWave.h
--- a/Engine/Source/Runtime/Engine/Classes/Sound/SoundWave.h
+++ b/Engine/Source/Runtime/Engine/Classes/Sound/SoundWave.h
@@ -215,10 +215,6 @@ class ENGINE_API USoundWave : public USoundBase
 
 #endif // WITH_EDITORONLY_DATA
 
-	/** Size of RawPCMData, or what RawPCMData would be if the sound was fully decompressed */
-	UPROPERTY()
-	int32 RawPCMDataSize;
-
 	/**
 	 * Subtitle cues.  If empty, use SpokenText as the subtitle.  Will often be empty,
 	 * as the contents of the subtitle is commonly identical to what is spoken.
@@ -287,6 +283,9 @@ public:
 	/** Pointer to 16 bit PCM data - used to decompress data to and preview sounds */
 	uint8*						RawPCMData;
 
+	/** Size of RawPCMData, or what RawPCMData would be if the sound was fully decompressed */
+	int32						RawPCMDataSize;
+
 	/** Memory containing the data copied from the compressed bulk data */
 	uint8*						ResourceData;
 
diff --git a/Engine/Source/Runtime/Engine/Private/GameFramework/HapticFeedbackEffect.cpp b/Engine/Source/Runtime/Engine/Private/GameFramework/HapticFeedbackEffect.cpp
--- a/Engine/Source/Runtime/Engine/Private/GameFramework/HapticFeedbackEffect.cpp
+++ b/Engine/Source/Runtime/Engine/Private/GameFramework/HapticFeedbackEffect.cpp
@@ -175,6 +175,7 @@ void UHapticFeedbackEffect_SoundWave::PrepareSoundWaveBuffer()
 	}
 	AD->Precache(SoundWave, true, false);
 	SoundWave->InitAudioResource(AD->GetRuntimeFormat(SoundWave));
+	check(!SoundWave->RawPCMData || SoundWave->RawPCMDataSize);
 	uint8* PCMData = SoundWave->RawPCMData;
 	int32 SampleRate = SoundWave->GetSampleRateForCurrentPlatform();
 	int TargetFrequency = 320;
diff --git a/Engine/Source/Runtime/Engine/Private/SoundWave.cpp b/Engine/Source/Runtime/Engine/Private/SoundWave.cpp
--- a/Engine/Source/Runtime/Engine/Private/SoundWave.cpp
+++ b/Engine/Source/Runtime/Engine/Private/SoundWave.cpp
@@ -108,6 +108,7 @@ USoundWave::USoundWave(const FObjectInitializer& ObjectInitializer)
 	CompressionQuality = 40;
 	SubtitlePriority = DEFAULT_SUBTITLE_PRIORITY;
 	ResourceState = ESoundWaveResourceState::NeedsFree;
+	RawPCMDataSize = 0;
 
 	// Default this to true since most sound wave types don't need precaching
 	bIsPrecacheDone = true;
@@ -131,6 +132,7 @@ void USoundWave::GetResourceSizeEx(FResourceSizeEx& CumulativeResourceSize)
 	{
 		if (LocalAudioDevice->HasCompressedAudioInfoClass(this) && DecompressionType == DTYPE_Native)
 		{
+			check(!RawPCMData || RawPCMDataSize);
 			CumulativeResourceSize.AddDedicatedSystemMemoryBytes(RawPCMDataSize);
 		}
 		else 
diff --git a/Engine/Source/Runtime/Mac/CoreAudio/Private/CoreAudioBuffer.cpp b/Engine/Source/Runtime/Mac/CoreAudio/Private/CoreAudioBuffer.cpp
--- a/Engine/Source/Runtime/Mac/CoreAudio/Private/CoreAudioBuffer.cpp
+++ b/Engine/Source/Runtime/Mac/CoreAudio/Private/CoreAudioBuffer.cpp
@@ -232,6 +232,8 @@ FCoreAudioSoundBuffer* FCoreAudioSoundBuffer::CreatePreviewBuffer( FCoreAudioDev
 	// Create new buffer.
 	Buffer = new FCoreAudioSoundBuffer( CoreAudioDevice, SoundFormat_PCMPreview );
 	
+	check(!Wave->RawPCMData || Wave->RawPCMDataSize);
+
 	// Take ownership the PCM data
 	Buffer->PCMData = Wave->RawPCMData;
 	Buffer->PCMDataSize = Wave->RawPCMDataSize;
@@ -260,6 +262,8 @@ FCoreAudioSoundBuffer* FCoreAudioSoundBuffer::CreateNativeBuffer( FCoreAudioDevi
 	// Create new buffer.
 	FCoreAudioSoundBuffer* Buffer = new FCoreAudioSoundBuffer( CoreAudioDevice, SoundFormat_PCM );
 
+	check(!Wave->RawPCMData || Wave->RawPCMDataSize);
+
 	// Take ownership the PCM data
 	Buffer->PCMData = Wave->RawPCMData;
 	Buffer->PCMDataSize = Wave->RawPCMDataSize;
diff --git a/Engine/Source/Runtime/Windows/XAudio2/Private/XAudio2Buffer.cpp b/Engine/Source/Runtime/Windows/XAudio2/Private/XAudio2Buffer.cpp
--- a/Engine/Source/Runtime/Windows/XAudio2/Private/XAudio2Buffer.cpp
+++ b/Engine/Source/Runtime/Windows/XAudio2/Private/XAudio2Buffer.cpp
@@ -429,6 +429,8 @@ FXAudio2SoundBuffer* FXAudio2SoundBuffer::CreatePreviewBuffer( FXAudio2Device* X
 	// Create new buffer.
 	Buffer = new FXAudio2SoundBuffer( XAudio2Device, SoundFormat_PCMPreview );
 
+	check(!Wave->RawPCMData || Wave->RawPCMDataSize);
+
 	// Take ownership the PCM data
 	Buffer->PCM.PCMData = Wave->RawPCMData;
 	Buffer->PCM.PCMDataSize = Wave->RawPCMDataSize;
@@ -452,6 +454,8 @@ FXAudio2SoundBuffer* FXAudio2SoundBuffer::CreateNativeBuffer( FXAudio2Device* XA
 	// Create new buffer.
 	FXAudio2SoundBuffer* Buffer = new FXAudio2SoundBuffer( XAudio2Device, SoundFormat_PCM );
 
+	check(!Wave->RawPCMData || Wave->RawPCMDataSize);
+
 	// Take ownership the PCM data
 	Buffer->PCM.PCMData = Wave->RawPCMData;
 	Buffer->PCM.PCMDataSize = Wave->RawPCMDataSize;
diff --git a/Engine/Source/Runtime/Windows/XAudio2/Private/XAudio2Source.cpp b/Engine/Source/Runtime/Windows/XAudio2/Private/XAudio2Source.cpp
--- a/Engine/Source/Runtime/Windows/XAudio2/Private/XAudio2Source.cpp
+++ b/Engine/Source/Runtime/Windows/XAudio2/Private/XAudio2Source.cpp
@@ -593,6 +593,8 @@ bool FXAudio2SoundSource::Init(FWaveInstance* InWaveInstance)
 
 			Update();
 
+			check(!InWaveInstance->WaveData->RawPCMData || InWaveInstance->WaveData->RawPCMDataSize);
+
 			// Initialize the total  number of frames of audio for this sound source
 			int32 NumBytes = InWaveInstance->WaveData->RawPCMDataSize;
 			NumTotalFrames = NumBytes / (Buffer->NumChannels * sizeof(int16));
