Copying //UE4/Dev-Audio to //UE4/Dev-Main (Source: //UE4/Dev-Audio @ 4332921)

diff --git a/Engine/Source/Runtime/Engine/Private/AudioDevice.cpp b/Engine/Source/Runtime/Engine/Private/AudioDevice.cpp
--- a/Engine/Source/Runtime/Engine/Private/AudioDevice.cpp
+++ b/Engine/Source/Runtime/Engine/Private/AudioDevice.cpp
@@ -4898,12 +4898,24 @@ void FAudioDevice::RegisterSoundClass(USoundClass* InSoundClass)
 	}
 }
 
-void FAudioDevice::UnregisterSoundClass(USoundClass* SoundClass)
+void FAudioDevice::UnregisterSoundClass(USoundClass* InSoundClass)
 {
-	check(IsInAudioThread());
-	if (SoundClass)
+	if (InSoundClass)
 	{
-		SoundClasses.Remove(SoundClass);
+		if (!IsInAudioThread())
+		{
+			DECLARE_CYCLE_STAT(TEXT("FAudioThreadTask.UnregisterSoundClass"), STAT_AudioUnregisterSoundClass, STATGROUP_AudioThreadCommands);
+
+			FAudioDevice* AudioDevice = this;
+			FAudioThread::RunCommandOnAudioThread([AudioDevice, InSoundClass]()
+			{
+				AudioDevice->UnregisterSoundClass(InSoundClass);
+			}, GET_STATID(STAT_AudioUnregisterSoundClass));
+
+			return;
+		}
+
+		SoundClasses.Remove(InSoundClass);
 	}
 }
 
diff --git a/Engine/Source/Runtime/Engine/Private/SoundEffectBase.cpp b/Engine/Source/Runtime/Engine/Private/SoundEffectBase.cpp
--- a/Engine/Source/Runtime/Engine/Private/SoundEffectBase.cpp
+++ b/Engine/Source/Runtime/Engine/Private/SoundEffectBase.cpp
@@ -51,7 +51,10 @@ void FSoundEffectBase::RegisterWithPreset(USoundEffectPreset* InParentPreset)
 
 void FSoundEffectBase::UnregisterWithPreset()
 {
-	ParentPreset->RemoveEffectInstance(this);
+	if (ParentPreset)
+	{
+		ParentPreset->RemoveEffectInstance(this);
+	}
 }
 
 bool FSoundEffectBase::IsParentPreset(USoundEffectPreset* InPreset) const
diff --git a/Engine/Source/Runtime/Engine/Private/SoundEffectPreset.cpp b/Engine/Source/Runtime/Engine/Private/SoundEffectPreset.cpp
--- a/Engine/Source/Runtime/Engine/Private/SoundEffectPreset.cpp
+++ b/Engine/Source/Runtime/Engine/Private/SoundEffectPreset.cpp
@@ -45,7 +45,10 @@ void USoundEffectPreset::AddEffectInstance(FSoundEffectBase* InSource)
 
 void USoundEffectPreset::RemoveEffectInstance(FSoundEffectBase* InSource)
 {
-	Instances.Remove(InSource);
+	if (Instances.Contains(InSource))
+	{
+		Instances.Remove(InSource);
+	}
 }
 
 #if WITH_EDITORONLY_DATA
diff --git a/Engine/Source/Runtime/Engine/Private/SoundSourceBus.cpp b/Engine/Source/Runtime/Engine/Private/SoundSourceBus.cpp
--- a/Engine/Source/Runtime/Engine/Private/SoundSourceBus.cpp
+++ b/Engine/Source/Runtime/Engine/Private/SoundSourceBus.cpp
@@ -8,13 +8,24 @@
 #include "UObject/UObjectIterator.h"
 #include "ActiveSound.h"
 
+USoundSourceBus::USoundSourceBus(const FObjectInitializer& ObjectInitializer)
+	: Super(ObjectInitializer)
+{
+	// This is a bus. This will result in the decompression type to be set as DTYPE_Bus. Audio won't be generated from this object but from instance data in audio mixer.
+	bIsBus = true;
+
+	Init();
+}
+
 void USoundSourceBus::PostLoad()
 {
 	Super::PostLoad();
 
-	// This is a bus. This will result in the decompression type to be set as DTYPE_Bus. Audio won't be generated from this object but from instance data in audio mixer.
-	bIsBus = true;
+	Init();
+}
 
+void USoundSourceBus::Init()
+{
 	// Allow users to manually set the source bus duration
 	Duration = GetDuration();
 
@@ -25,31 +36,23 @@ void USoundSourceBus::PostLoad()
 	// Note source buses can't ever be truly virtual as they are procedurally generated.
 	bVirtualizeWhenSilent = !bAutoDeactivateWhenSilent;
 
-	// Use the main/default audio device for storing and retrieving sound class properties
-	FAudioDeviceManager* AudioDeviceManager = (GEngine ? GEngine->GetAudioDeviceManager() : nullptr);
-
 	// Set the channels equal to the users channel count choice
 	switch (SourceBusChannels)
 	{
-		case ESourceBusChannels::Mono:
-			NumChannels = 1;
-			break;
+	case ESourceBusChannels::Mono:
+		NumChannels = 1;
+		break;
 
-		case ESourceBusChannels::Stereo:
-			NumChannels = 2;
-			break;
+	case ESourceBusChannels::Stereo:
+		NumChannels = 2;
+		break;
 	}
 }
 
-void USoundSourceBus::BeginDestroy()
-{
-	Super::BeginDestroy();
-}
-
 #if WITH_EDITOR
 void USoundSourceBus::PostEditChangeProperty(FPropertyChangedEvent& PropertyChangedEvent)
 {
-	// stub
+	Init();
 }
 #endif
 
diff --git a/Engine/Source/Runtime/Engine/Private/SoundSubmix.cpp b/Engine/Source/Runtime/Engine/Private/SoundSubmix.cpp
--- a/Engine/Source/Runtime/Engine/Private/SoundSubmix.cpp
+++ b/Engine/Source/Runtime/Engine/Private/SoundSubmix.cpp
@@ -20,6 +20,8 @@ TSharedPtr<ISoundSubmixAudioEditor> USoundSubmix::SoundSubmixAudioEditor = nullp
 USoundSubmix::USoundSubmix(const FObjectInitializer& ObjectInitializer)
 	: Super(ObjectInitializer)
 {
+	EnvelopeFollowerAttackTime = 10;
+	EnvelopeFollowerReleaseTime = 500;
 }
 
 void USoundSubmix::StartRecordingOutput(const UObject* WorldContextObject, float ExpectedDuration)
@@ -124,6 +126,85 @@ void USoundSubmix::StopRecordingOutput(FAudioDevice* InDevice, EAudioRecordingEx
 	}
 }
 
+void USoundSubmix::StartEnvelopeFollowing(const UObject* WorldContextObject)
+{
+	if (!GEngine)
+	{
+		return;
+	}
+
+	// Find device for this specific audio recording thing.
+	UWorld* ThisWorld = GEngine->GetWorldFromContextObject(WorldContextObject, EGetWorldErrorMode::LogAndReturnNull);
+	FAudioDevice* AudioDevice = ThisWorld->GetAudioDevice();
+
+	StartEnvelopeFollowing(AudioDevice);
+}
+
+void USoundSubmix::StartEnvelopeFollowing(FAudioDevice* InAudioDevice)
+{
+	if (InAudioDevice)
+	{
+		InAudioDevice->StartEnvelopeFollowing(this);
+
+		// Track which audio devices were envelope following with. Note this may be multiple devices if running with multiple PIE.
+		//EnvelopeFollowingDevices.AddUnique(InAudioDevice);
+	}
+}
+
+void USoundSubmix::StopEnvelopeFollowing(const UObject* WorldContextObject)
+{
+	if (!GEngine)
+	{
+		return;
+	}
+
+	// Find device for this specific audio recording thing.
+	UWorld* ThisWorld = GEngine->GetWorldFromContextObject(WorldContextObject, EGetWorldErrorMode::LogAndReturnNull);
+	FAudioDevice* AudioDevice = ThisWorld->GetAudioDevice();
+
+	StopEnvelopeFollowing(AudioDevice);
+}
+
+void USoundSubmix::StopEnvelopeFollowing(FAudioDevice* InAudioDevice)
+{
+	if (InAudioDevice)
+	{
+		InAudioDevice->StopEnvelopeFollowing(this);
+
+		// Stop tracking audio device for envelope following data
+		//EnvelopeFollowingDevices.RemoveSingleSwap(InAudioDevice);
+	}
+}
+
+void USoundSubmix::AddEnvelopeFollowerDelegate(const UObject* WorldContextObject, const FOnSubmixEnvelopeBP& OnSubmixEnvelopeBP)
+{
+	if (!GEngine)
+	{
+		return;
+	}
+
+	// Find device for this specific audio recording thing.
+	UWorld* ThisWorld = GEngine->GetWorldFromContextObject(WorldContextObject, EGetWorldErrorMode::LogAndReturnNull);
+	FAudioDevice* AudioDevice = ThisWorld->GetAudioDevice();
+	if (AudioDevice)
+	{
+		AudioDevice->AddEnvelopeFollowerDelegate(this, OnSubmixEnvelopeBP);
+
+// 		FSubmixEnvelopeData* EnvelopeData = EnvelopeFollowingData.Find(AudioDevice);
+// 		if (EnvelopeData)
+// 		{
+// 			EnvelopeData->OnSubmixEnvelope.AddUnique(OnSubmixEnvelopeBP);
+// 		}
+// 		else
+// 		{
+// 			FSubmixEnvelopeData& NewEnvelopeData = EnvelopeFollowingData.Add(AudioDevice, FSubmixEnvelopeData());
+// 			NewEnvelopeData.OnSubmixEnvelope.AddUnique(OnSubmixEnvelopeBP);
+// 		}
+	}
+
+}
+
+
 FString USoundSubmix::GetDesc()
 {
 	return FString(TEXT("Sound submix"));
diff --git a/Engine/Source/Runtime/Engine/Public/AudioDevice.h b/Engine/Source/Runtime/Engine/Public/AudioDevice.h
--- a/Engine/Source/Runtime/Engine/Public/AudioDevice.h
+++ b/Engine/Source/Runtime/Engine/Public/AudioDevice.h
@@ -543,6 +543,9 @@ public:
 	 */
 	void Update(bool bGameTicking);
 
+	/** Update called on game thread. */
+	virtual void UpdateGameThread() {}
+
 	/**
 	 * Suspend/resume all sounds (global pause for device suspend/resume, etc.)
 	 *
@@ -1154,18 +1157,36 @@ public:
 	/** This is called by a USoundSubmix to start recording a submix instance on this device. */
 	virtual void StartRecording(USoundSubmix* InSubmix, float ExpectedRecordingDuration) 
 	{
-		UE_LOG(LogAudio, Error, TEXT("Submix recording only works with the audio mixer. Please run using -audiomixer to use submix recording."));
+		UE_LOG(LogAudio, Error, TEXT("Submix recording only works with the audio mixer. Please run using -audiomixer to or set INI file use submix recording."));
 	}
 
 	/** This is called by a USoundSubmix when we stop recording a submix on this device. */
 	virtual Audio::AlignedFloatBuffer& StopRecording(USoundSubmix* InSubmix, float& OutNumChannels, float& OutSampleRate) 
 	{
-		UE_LOG(LogAudio, Error, TEXT("Submix recording only works with the audio mixer. Please run using -audiomixer to use submix recording."));
+		UE_LOG(LogAudio, Error, TEXT("Submix recording only works with the audio mixer. Please run using -audiomixer to or set INI file use submix recording."));
 		
 		static Audio::AlignedFloatBuffer InvalidBuffer;
 		return InvalidBuffer;
 	}
 
+	/** This is called by a USoundSubmix to start envelope following on a submix isntance on this device. */
+	virtual void StartEnvelopeFollowing(USoundSubmix* InSubmix)
+	{
+		UE_LOG(LogAudio, Error, TEXT("Envelope following submixes only works with the audio mixer. Please run using -audiomixer or set INI file to use submix recording."));
+	}
+
+	/** This is called by a USoundSubmix when we stop envelope following a submix instance on this device. */
+	virtual void StopEnvelopeFollowing(USoundSubmix* InSubmix)
+	{
+		UE_LOG(LogAudio, Error, TEXT("Envelope following submixes only works with the audio mixer. Please run using -audiomixer or set INI file to use submix recording."));
+	}
+
+	/** Adds an envelope follower delegate to the submix for this audio device. */
+	virtual void AddEnvelopeFollowerDelegate(USoundSubmix* InSubmix, const FOnSubmixEnvelopeBP& OnSubmixEnvelopeBP)
+	{
+		UE_LOG(LogAudio, Error, TEXT("Envelope following submixes only works with the audio mixer. Please run using -audiomixer or set INI file to use submix recording."));
+	}
+
 protected:
 	friend class FSoundSource;
 
diff --git a/Engine/Source/Runtime/Windows/AudioMixerXAudio2/Private/AudioMixerPlatformWindows.cpp b/Engine/Source/Runtime/Windows/AudioMixerXAudio2/Private/AudioMixerPlatformWindows.cpp
--- a/Engine/Source/Runtime/Windows/AudioMixerXAudio2/Private/AudioMixerPlatformWindows.cpp
+++ b/Engine/Source/Runtime/Windows/AudioMixerXAudio2/Private/AudioMixerPlatformWindows.cpp
@@ -18,6 +18,15 @@
 #include <mmdeviceapi.h>
 #include <functiondiscoverykeys_devpkey.h>
 
+static int32 GEnableDetailedWindowsDeviceLogging = 0;
+FAutoConsoleVariableRef CVarEnableDetailedWindowsDeviceLogging(
+	TEXT("au.EnableDetailedWindowsDeviceLogging"),
+	GEnableDetailedWindowsDeviceLogging,
+	TEXT("Enables detailed windows device logging.\n")
+	TEXT("0: Not Enabled, 1: Enabled"),
+	ECVF_Default);
+
+
 class FWindowsMMNotificationClient final : public IMMNotificationClient
 {
 public:
@@ -49,6 +58,11 @@ public:
 
 	HRESULT STDMETHODCALLTYPE OnDefaultDeviceChanged(EDataFlow InFlow, ERole InRole, LPCWSTR pwstrDeviceId) override
 	{
+		if (GEnableDetailedWindowsDeviceLogging)
+		{
+			UE_LOG(LogAudioMixer, Warning, TEXT("OnDefaultDeviceChanged: %d, %d, %s"), InFlow, InRole, pwstrDeviceId);
+		}
+
 		Audio::EAudioDeviceRole AudioDeviceRole;
 
 		if (InRole == eConsole)
@@ -93,6 +107,11 @@ public:
 
 	HRESULT STDMETHODCALLTYPE OnDeviceAdded(LPCWSTR pwstrDeviceId) override
 	{
+		if (GEnableDetailedWindowsDeviceLogging)
+		{
+			UE_LOG(LogAudioMixer, Warning, TEXT("OnDeviceAdded: %s"), pwstrDeviceId);
+		}
+
 		for (Audio::IAudioMixerDeviceChangedLister* Listener : Listeners)
 		{
 			Listener->OnDeviceAdded(FString(pwstrDeviceId));
@@ -102,6 +121,11 @@ public:
 
 	HRESULT STDMETHODCALLTYPE OnDeviceRemoved(LPCWSTR pwstrDeviceId) override
 	{
+		if (GEnableDetailedWindowsDeviceLogging)
+		{
+			UE_LOG(LogAudioMixer, Warning, TEXT("OnDeviceRemoved: %s"), pwstrDeviceId);
+		}
+
 		for (Audio::IAudioMixerDeviceChangedLister* Listener : Listeners)
 		{
 			Listener->OnDeviceRemoved(FString(pwstrDeviceId));
@@ -111,11 +135,30 @@ public:
 
 	HRESULT STDMETHODCALLTYPE OnDeviceStateChanged(LPCWSTR pwstrDeviceId, DWORD dwNewState) override
 	{
+		if (GEnableDetailedWindowsDeviceLogging)
+		{
+			UE_LOG(LogAudioMixer, Warning, TEXT("OnDeviceStateChanged: %s, %d"), pwstrDeviceId, dwNewState);
+		}
+
 		if (dwNewState == DEVICE_STATE_DISABLED || dwNewState == DEVICE_STATE_UNPLUGGED || dwNewState == DEVICE_STATE_NOTPRESENT)
 		{
 			for (Audio::IAudioMixerDeviceChangedLister* Listener : Listeners)
 			{
-				Listener->OnDeviceRemoved(FString(pwstrDeviceId));
+				switch (dwNewState)
+				{
+				case DEVICE_STATE_DISABLED:
+					Listener->OnDeviceStateChanged(FString(pwstrDeviceId), Audio::EAudioDeviceState::Disabled);
+					break;
+
+				case DEVICE_STATE_UNPLUGGED:
+					Listener->OnDeviceStateChanged(FString(pwstrDeviceId), Audio::EAudioDeviceState::Unplugged);
+					break;
+				case DEVICE_STATE_NOTPRESENT:
+					Listener->OnDeviceStateChanged(FString(pwstrDeviceId), Audio::EAudioDeviceState::NotPresent);
+					break;
+				default:
+					break;
+				}
 			}
 		}
 		return S_OK;
@@ -123,6 +166,11 @@ public:
 
 	HRESULT STDMETHODCALLTYPE OnPropertyValueChanged(LPCWSTR pwstrDeviceId, const PROPERTYKEY key)
 	{
+		if (GEnableDetailedWindowsDeviceLogging)
+		{
+			UE_LOG(LogAudioMixer, Warning, TEXT("OnPropertyValueChanged: %s, %d"), pwstrDeviceId, key.pid);
+		}
+
 		FString ChangedId = FString(pwstrDeviceId);
 
 		// look for ids we care about!
@@ -204,13 +252,21 @@ namespace Audio
 
 	void FMixerPlatformXAudio2::OnDefaultRenderDeviceChanged(const EAudioDeviceRole InAudioDeviceRole, const FString& DeviceId)
 	{
+		if (!AllowDeviceSwap())
+		{
+			return;
+		}
+
 		if (AudioDeviceSwapCriticalSection.TryLock())
 		{
+			UE_LOG(LogAudioMixer, Warning, TEXT("Changing default audio render device to new device: %s."), *DeviceId);
+
 			NewAudioDeviceId = "";
 			bMoveAudioStreamToNewAudioDevice = true;
 
 			AudioDeviceSwapCriticalSection.Unlock();
 		}
+
 	}
 
 	void FMixerPlatformXAudio2::OnDeviceAdded(const FString& DeviceId)
@@ -221,6 +277,8 @@ namespace Audio
 			// move our audio stream to this newly added device.
 			if (AudioStreamInfo.DeviceInfo.DeviceId != OriginalAudioDeviceId && DeviceId == OriginalAudioDeviceId)
 			{
+				UE_LOG(LogAudioMixer, Warning, TEXT("Original audio device re-added. Moving audio back to original audio device %s."), *OriginalAudioDeviceId);
+
 				NewAudioDeviceId = OriginalAudioDeviceId;
 				bMoveAudioStreamToNewAudioDevice = true;
 			}
@@ -236,6 +294,8 @@ namespace Audio
 			// If the device we're currently using was removed... then switch to the new default audio device.
 			if (AudioStreamInfo.DeviceInfo.DeviceId == DeviceId)
 			{
+				UE_LOG(LogAudioMixer, Warning, TEXT("Audio device removed, falling back to other windows default device."));
+
 				NewAudioDeviceId = "";
 				bMoveAudioStreamToNewAudioDevice = true;
 			}
diff --git a/Engine/Source/Runtime/Windows/AudioMixerXAudio2/Private/AudioMixerPlatformXAudio2.cpp b/Engine/Source/Runtime/Windows/AudioMixerXAudio2/Private/AudioMixerPlatformXAudio2.cpp
--- a/Engine/Source/Runtime/Windows/AudioMixerXAudio2/Private/AudioMixerPlatformXAudio2.cpp
+++ b/Engine/Source/Runtime/Windows/AudioMixerXAudio2/Private/AudioMixerPlatformXAudio2.cpp
@@ -60,6 +60,7 @@ namespace Audio
 		, OutputAudioStreamMasteringVoice(nullptr)
 		, OutputAudioStreamSourceVoice(nullptr)
 		, bMoveAudioStreamToNewAudioDevice(false)
+		, LastDeviceSwapTime(0.0)
 		, bIsComInitialized(false)
 		, bIsInitialized(false)
 		, bIsDeviceOpen(false)
@@ -113,6 +114,27 @@ namespace Audio
 		}
 	}
 
+	bool FMixerPlatformXAudio2::AllowDeviceSwap()
+	{
+		double CurrentTime = FPlatformTime::Seconds();
+
+		// If we're already in the process of swapping, we do not want to "double-trigger" a swap
+		if (bMoveAudioStreamToNewAudioDevice)
+		{
+			LastDeviceSwapTime = CurrentTime;
+			return false;
+		}
+
+		// Some devices spam device swap notifications, so we want to rate-limit them to prevent double/tripple triggering.
+		static const int32 MinSwapTimeMs = 10;
+		if (CurrentTime - LastDeviceSwapTime > (double)MinSwapTimeMs / 1000.0)
+		{
+			LastDeviceSwapTime = CurrentTime;
+			return true;
+		}
+		return false;
+	}
+
 	bool FMixerPlatformXAudio2::InitializeHardware()
 	{
 		if (bIsInitialized)
@@ -124,7 +146,23 @@ namespace Audio
 
 #if PLATFORM_WINDOWS
 		bIsComInitialized = FWindowsPlatformMisc::CoInitialize();
-#endif //#if PLATFORM_WINDOWS
+#if PLATFORM_64BITS
+		// Work around the fact the x64 version of XAudio2_7.dll does not properly ref count
+		// by forcing it to be always loaded
+
+		// Load the xaudio2 library and keep a handle so we can free it on teardown
+		// Note: windows internally ref-counts the library per call to load library so 
+		// when we call FreeLibrary, it will only free it once the refcount is zero
+		XAudio2Dll = LoadLibraryA("XAudio2_7.dll");
+
+		// returning null means we failed to load XAudio2, which means everything will fail
+		if (XAudio2Dll == nullptr)
+		{
+			UE_LOG(LogInit, Warning, TEXT("Failed to load XAudio2 dll"));
+			return false;
+		}
+#endif // #if PLATFORM_64BITS
+#endif // #if PLATFORM_WINDOWS
 
 		uint32 Flags = 0;
 
@@ -158,12 +196,25 @@ namespace Audio
 		SAFE_RELEASE(XAudio2System);
 
 #if PLATFORM_WINDOWS
+
+#if PLATFORM_64BITS
+		if (XAudio2Dll)
+		{
+			if (!FreeLibrary(XAudio2Dll))
+			{
+				UE_LOG(LogAudio, Warning, TEXT("Failed to free XAudio2 Dll"));
+			}
+		}
+#endif
+
 		if (bIsComInitialized)
 		{
 			FWindowsPlatformMisc::CoUninitialize();
 		}
 #endif
 
+
+
 		bIsInitialized = false;
 
 		return true;
@@ -214,16 +265,7 @@ namespace Audio
 		const WAVEFORMATEX& WaveFormatEx = DeviceDetails.OutputFormat.Format;
 		OutInfo.SampleRate = WaveFormatEx.nSamplesPerSec;
 
-		bool bIsMono = (WaveFormatEx.nChannels == 1);
-		// We are going to default to stereo for mono devices, then mix to mono on buffer submission (automatically done by xaudio2)
-		if (bIsMono == 1)
-		{
-			OutInfo.NumChannels = 2;
-		}
-		else
-		{
-			OutInfo.NumChannels = WaveFormatEx.nChannels;
-		}
+		OutInfo.NumChannels = FMath::Clamp((int32)WaveFormatEx.nChannels, 2, 8);
 
 		// XAudio2 automatically converts the audio format to output device us so we don't need to do any format conversions
 		OutInfo.Format = EAudioMixerStreamDataFormat::Float;
@@ -244,7 +286,7 @@ namespace Audio
 
 			check(EAudioMixerChannel::ChannelTypeCount == ChannelTypeMap.Num());
 			uint32 ChanCount = 0;
-			for (uint32 ChannelTypeIndex = 0;  ChannelTypeIndex < EAudioMixerChannel::ChannelTypeCount &&  ChanCount < WaveFormatEx.nChannels; ++ChannelTypeIndex)
+			for (uint32 ChannelTypeIndex = 0; ChannelTypeIndex < EAudioMixerChannel::ChannelTypeCount && ChanCount < (uint32)OutInfo.NumChannels; ++ChannelTypeIndex)
 			{
 				if (WaveFormatExtensible->dwChannelMask & ChannelTypeMap[ChannelTypeIndex])
 				{
@@ -253,14 +295,58 @@ namespace Audio
 				}
 			}
 
-			if (ChanCount != WaveFormatEx.nChannels)
+			// We didnt match channel masks for all channels, revert to a default ordering
+			if (ChanCount < (uint32)OutInfo.NumChannels)
 			{
-				UE_LOG(LogAudioMixer, Warning, TEXT("Did not find the channel type flags for audio device '%s'. Filling in with default channel type mappings"), *OutInfo.Name);
-				while ((int32)ChanCount < OutInfo.OutputChannelArray.Num())
+				UE_LOG(LogAudioMixer, Warning, TEXT("Did not find the channel type flags for audio device '%s'. Reverting to a default channel ordering."), *OutInfo.Name);
+
+				OutInfo.OutputChannelArray.Reset();
+
+				static EAudioMixerChannel::Type DefaultChannelOrdering[] = {
+					EAudioMixerChannel::FrontLeft,
+					EAudioMixerChannel::FrontRight,
+					EAudioMixerChannel::FrontCenter,
+					EAudioMixerChannel::LowFrequency,
+					EAudioMixerChannel::SideLeft,
+					EAudioMixerChannel::SideRight,
+					EAudioMixerChannel::BackLeft,
+					EAudioMixerChannel::BackRight,
+				};
+
+				EAudioMixerChannel::Type* ChannelOrdering = DefaultChannelOrdering;
+
+				// Override channel ordering for some special cases
+				if (OutInfo.NumChannels == 4)
+				{
+					static EAudioMixerChannel::Type DefaultChannelOrderingQuad[] = {
+						EAudioMixerChannel::FrontLeft,
+						EAudioMixerChannel::FrontRight,
+						EAudioMixerChannel::BackLeft,
+						EAudioMixerChannel::BackRight,
+					};
+
+					ChannelOrdering = DefaultChannelOrderingQuad;
+				}
+				else if (OutInfo.NumChannels == 6)
 				{
-					// Fill in the channel array with speaker types as a guess. This may not be accurate.
-					OutInfo.OutputChannelArray.Add((EAudioMixerChannel::Type)ChanCount++);
+					static EAudioMixerChannel::Type DefaultChannelOrdering51[] = {
+						EAudioMixerChannel::FrontLeft,
+						EAudioMixerChannel::FrontRight,
+						EAudioMixerChannel::FrontCenter,
+						EAudioMixerChannel::LowFrequency,
+						EAudioMixerChannel::BackLeft,
+						EAudioMixerChannel::BackRight,
+					};
+
+					ChannelOrdering = DefaultChannelOrdering51;
 				}
+
+				check(OutInfo.NumChannels <= 8);
+				for (int32 Index = 0; Index < OutInfo.NumChannels; ++Index)
+				{
+					OutInfo.OutputChannelArray.Add(ChannelOrdering[Index]);
+				}
+
 			}
 		}
 		else
@@ -277,7 +363,8 @@ namespace Audio
 		UE_LOG(LogAudioMixer, Display, TEXT("Name: %s"), *OutInfo.Name);
 		UE_LOG(LogAudioMixer, Display, TEXT("Is Default: %s"), OutInfo.bIsSystemDefault ? TEXT("Yes") : TEXT("No"));
 		UE_LOG(LogAudioMixer, Display, TEXT("Sample Rate: %d"), OutInfo.SampleRate);
-		UE_LOG(LogAudioMixer, Display, TEXT("Channel Count: %d (was Mono? %s)"), OutInfo.NumChannels, bIsMono ? TEXT("Yes") : TEXT("No"));
+		UE_LOG(LogAudioMixer, Display, TEXT("Channel Count Used: %d"), OutInfo.NumChannels);
+		UE_LOG(LogAudioMixer, Display, TEXT("Device Channel Count: %d"), WaveFormatEx.nChannels);
 		UE_LOG(LogAudioMixer, Display, TEXT("Channel Order:"));
 		for (int32 i = 0; i < OutInfo.NumChannels; ++i)
 		{
diff --git a/Engine/Source/Runtime/Windows/AudioMixerXAudio2/Private/AudioMixerPlatformXAudio2.h b/Engine/Source/Runtime/Windows/AudioMixerXAudio2/Private/AudioMixerPlatformXAudio2.h
--- a/Engine/Source/Runtime/Windows/AudioMixerXAudio2/Private/AudioMixerPlatformXAudio2.h
+++ b/Engine/Source/Runtime/Windows/AudioMixerXAudio2/Private/AudioMixerPlatformXAudio2.h
@@ -84,9 +84,13 @@ namespace Audio
 	private:
 
 		const TCHAR* GetErrorString(HRESULT Result);
+		bool AllowDeviceSwap();
 
 		typedef TArray<long> TChannelTypeMap;
 		
+		// Handle to XAudio2DLL
+		HMODULE XAudio2Dll;
+
 		// Bool indicating that the default audio device changed
 		// And that we need to restart the audio device.
 		FThreadSafeBool bDeviceChanged;
@@ -100,6 +104,8 @@ namespace Audio
 		FString OriginalAudioDeviceId;
 		FString NewAudioDeviceId;
 		FThreadSafeBool bMoveAudioStreamToNewAudioDevice;
+		double LastDeviceSwapTime;
+
 		uint32 bIsComInitialized : 1;
 		uint32 bIsInitialized : 1;
 		uint32 bIsDeviceOpen : 1;
