Fix Clang link error.

#rnx

diff --git a/Engine/Source/Runtime/Windows/D3D11RHI/Private/D3D11Texture.cpp b/Engine/Source/Runtime/Windows/D3D11RHI/Private/D3D11Texture.cpp
--- a/Engine/Source/Runtime/Windows/D3D11RHI/Private/D3D11Texture.cpp
+++ b/Engine/Source/Runtime/Windows/D3D11RHI/Private/D3D11Texture.cpp
@@ -239,6 +239,8 @@ TD3D11Texture2D<BaseResourceType>::~TD3D11Texture2D()
 #endif
 }
 
+template TD3D11Texture2D<FD3D11BaseTexture2D>::~TD3D11Texture2D();
+
 FD3D11Texture3D::~FD3D11Texture3D()
 {
 	D3D11TextureDeleted( *this );
