1. When going from a device to a device of a different channel count, audio for some sources is heavily distorted.
2. When changing default device with multiple devices, the audio remains on the previous device without moving to a new device.

diff --git a/Engine/Source/Runtime/AudioMixer/Private/AudioMixerSourceManager.cpp b/Engine/Source/Runtime/AudioMixer/Private/AudioMixerSourceManager.cpp
--- a/Engine/Source/Runtime/AudioMixer/Private/AudioMixerSourceManager.cpp
+++ b/Engine/Source/Runtime/AudioMixer/Private/AudioMixerSourceManager.cpp
@@ -1780,7 +1780,7 @@ namespace Audio
 
 	FMixerSourceManager::FSourceDownmixData& FMixerSourceManager::InitializeDownmixForSource(const int32 SourceId, const int32 NumInputChannels, const int32 NumOutputChannels, const int32 InNumOutputFrames)
 	{
-		DownmixDataArray[SourceId].ResetData(NumInputChannels);
+		DownmixDataArray[SourceId].ResetData(NumInputChannels, NumOutputChannels);
 		return DownmixDataArray[SourceId];
 	}
 
diff --git a/Engine/Source/Runtime/AudioMixer/Private/AudioMixerSourceManager.h b/Engine/Source/Runtime/AudioMixer/Private/AudioMixerSourceManager.h
--- a/Engine/Source/Runtime/AudioMixer/Private/AudioMixerSourceManager.h
+++ b/Engine/Source/Runtime/AudioMixer/Private/AudioMixerSourceManager.h
@@ -376,8 +376,9 @@ namespace Audio
 				DeviceSubmixInfo.Reset(NumInputChannels, NumDeviceOutputChannels, NumFrames);
 			}
 
-			void ResetData(const uint32 InNumInputChannels)
+			void ResetData(const uint32 InNumInputChannels, int32 InNumDeviceChannels)
 			{
+				NumDeviceChannels = InNumDeviceChannels;
 				NumInputChannels = InNumInputChannels;
 				PostEffectBuffers = nullptr;
 
diff --git a/Engine/Source/Runtime/Windows/AudioMixerXAudio2/Private/AudioMixerPlatformXAudio2.cpp b/Engine/Source/Runtime/Windows/AudioMixerXAudio2/Private/AudioMixerPlatformXAudio2.cpp
--- a/Engine/Source/Runtime/Windows/AudioMixerXAudio2/Private/AudioMixerPlatformXAudio2.cpp
+++ b/Engine/Source/Runtime/Windows/AudioMixerXAudio2/Private/AudioMixerPlatformXAudio2.cpp
@@ -140,6 +140,26 @@ namespace Audio
 		return false;
 	}
 
+	bool FMixerPlatformXAudio2::ResetXAudio2System()
+	{
+		SAFE_RELEASE(XAudio2System);
+
+		uint32 Flags = 0;
+
+#if WITH_XMA2
+		// We need to raise this flag explicitly to prevent initializing SHAPE twice, because we are allocating SHAPE in FXMAAudioInfo
+		Flags |= XAUDIO2_DO_NOT_USE_SHAPE;
+#endif
+
+		if (FAILED(XAudio2Create(&XAudio2System, Flags, (XAUDIO2_PROCESSOR)FPlatformAffinity::GetAudioThreadMask())))
+		{
+			XAudio2System = nullptr;
+			return false;
+		}
+
+		return true;
+	}
+
 	bool FMixerPlatformXAudio2::InitializeHardware()
 	{
 		if (bIsInitialized)
@@ -676,6 +696,13 @@ namespace Audio
 
 		if (NumDevices > 0)
 		{
+			if (!ResetXAudio2System())
+			{
+				// Reinitializing the XAudio2System failed, so we have to exit here.
+				StartRunningNullDevice();
+				return true;
+			}
+
 			// Now get info on the new audio device we're trying to reset to
 			uint32 DeviceIndex = 0;
 			if (!InNewDeviceId.IsEmpty())
diff --git a/Engine/Source/Runtime/Windows/AudioMixerXAudio2/Private/AudioMixerPlatformXAudio2.h b/Engine/Source/Runtime/Windows/AudioMixerXAudio2/Private/AudioMixerPlatformXAudio2.h
--- a/Engine/Source/Runtime/Windows/AudioMixerXAudio2/Private/AudioMixerPlatformXAudio2.h
+++ b/Engine/Source/Runtime/Windows/AudioMixerXAudio2/Private/AudioMixerPlatformXAudio2.h
@@ -87,6 +87,10 @@ namespace Audio
 		const TCHAR* GetErrorString(HRESULT Result);
 		bool AllowDeviceSwap();
 
+		// Used to teardown and reinitialize XAudio2.
+		// This must be done to repopulate the playback device list in XAudio 2.7.
+		bool ResetXAudio2System();
+
 		typedef TArray<long> TChannelTypeMap;
 		
 #if PLATFORM_WINDOWS
