[NO NON-PII COMMIT MESSAGE]

diff --git a/Engine/Source/Runtime/AudioMixer/Private/AudioMixer.cpp b/Engine/Source/Runtime/AudioMixer/Private/AudioMixer.cpp
--- a/Engine/Source/Runtime/AudioMixer/Private/AudioMixer.cpp
+++ b/Engine/Source/Runtime/AudioMixer/Private/AudioMixer.cpp
@@ -235,6 +235,7 @@ namespace Audio
 		: bWarnedBufferUnderrun(false)
 		, AudioRenderThread(nullptr)
 		, AudioRenderEvent(nullptr)
+		, bIsInDeviceSwap(false)
 		, AudioFadeEvent(nullptr)
 		, CurrentBufferReadIndex(INDEX_NONE)
 		, CurrentBufferWriteIndex(INDEX_NONE)
@@ -364,9 +365,24 @@ namespace Audio
 	{
 		LLM_SCOPE(ELLMTag::AudioMixer);
 
+		// If we are flushing buffers for our output voice and this is being called on the audio thread directly,
+		// early exit.
+		if (bIsInDeviceSwap)
+		{
+			return;
+		}
+
+		// If we are currently swapping devices and OnBufferEnd is being triggered in an XAudio2Thread,
+		// early exit.
+		if (!DeviceSwapCriticalSection.TryLock())
+		{
+			return;
+		}
+
 		// Don't read any more audio if we're not running or changing device
-		if (AudioStreamInfo.StreamState != EAudioOutputStreamState::Running || !DeviceSwapCriticalSection.TryLock())
+		if (AudioStreamInfo.StreamState != EAudioOutputStreamState::Running)
 		{
+			DeviceSwapCriticalSection.Unlock();
 			return;
 		}
 
diff --git a/Engine/Source/Runtime/AudioMixer/Public/AudioMixer.h b/Engine/Source/Runtime/AudioMixer/Public/AudioMixer.h
--- a/Engine/Source/Runtime/AudioMixer/Public/AudioMixer.h
+++ b/Engine/Source/Runtime/AudioMixer/Public/AudioMixer.h
@@ -515,6 +515,9 @@ namespace Audio
 		/** Critical Section used for times when we need the render loop to halt for the device swap. */
 		FCriticalSection DeviceSwapCriticalSection;
 
+		/** This is used if we are attempting to TryLock on DeviceSwapCriticalSection, but a buffer callback is being called in the current thread. */
+		FThreadSafeBool bIsInDeviceSwap;
+
 		/** Event allows you to block until fadeout is complete. */
 		FEvent* AudioFadeEvent;
 
diff --git a/Engine/Source/Runtime/Windows/AudioMixerXAudio2/Private/AudioMixerPlatformXAudio2.cpp b/Engine/Source/Runtime/Windows/AudioMixerXAudio2/Private/AudioMixerPlatformXAudio2.cpp
--- a/Engine/Source/Runtime/Windows/AudioMixerXAudio2/Private/AudioMixerPlatformXAudio2.cpp
+++ b/Engine/Source/Runtime/Windows/AudioMixerXAudio2/Private/AudioMixerPlatformXAudio2.cpp
@@ -679,6 +679,17 @@ namespace Audio
 			// we have to wait for it here.
 			FScopeLock ScopeLock(&DeviceSwapCriticalSection);
 
+			// Now that we've properly locked, raise the bIsInDeviceSwap flag
+			// in case FlushSourceBuffers() calls OnBufferEnd on this thread,
+			// and DeviceSwapCriticalSection.TryLock() is still returning true
+			bIsInDeviceSwap = true;
+
+			// Flush all buffers. Because we've locked DeviceSwapCriticalSection, ReadNextBuffer will early exit and we will not submit any additional buffers.
+			if (OutputAudioStreamSourceVoice)
+			{
+				OutputAudioStreamSourceVoice->FlushSourceBuffers();
+			}
+
 			if (OutputAudioStreamSourceVoice)
 			{
 				// Then destroy the current audio stream source voice
@@ -692,6 +703,8 @@ namespace Audio
 				OutputAudioStreamMasteringVoice->DestroyVoice();
 				OutputAudioStreamMasteringVoice = nullptr;
 			}
+
+			bIsInDeviceSwap = false;
 		}
 
 		if (NumDevices > 0)
