Fix for swapping to legacy backend after launching with audio mixer on PC when there are no audio devices available.

diff --git a/Engine/Source/Runtime/Windows/XAudio2/Private/XAudio2Device.cpp b/Engine/Source/Runtime/Windows/XAudio2/Private/XAudio2Device.cpp
--- a/Engine/Source/Runtime/Windows/XAudio2/Private/XAudio2Device.cpp
+++ b/Engine/Source/Runtime/Windows/XAudio2/Private/XAudio2Device.cpp
@@ -307,7 +307,7 @@ void FXAudio2Device::TeardownHardware()
 void FXAudio2Device::UpdateHardware()
 {
 	// If the audio device changed, we need to tear down and restart the audio engine state
-	if (DeviceProperties->DidAudioDeviceChange())
+	if (DeviceProperties && DeviceProperties->DidAudioDeviceChange())
 	{
 		//Cache the current audio clock.
 		CachedAudioClockStartTime = GetAudioClock();
