RL - Pretty sure this flush is not needed as the documentation states: "When UpdateSubresource returns, the application is free to change or even free the data pointed to by pSrcData because the method has already copied/snapped away the original contents. "
This fixes a performance regression with HLR in DX11 because the PrimitiveUniform Shader Parameters take this path now.

diff --git a/Engine/Source/Runtime/Windows/D3D11RHI/Private/D3D11StructuredBuffer.cpp b/Engine/Source/Runtime/Windows/D3D11RHI/Private/D3D11StructuredBuffer.cpp
--- a/Engine/Source/Runtime/Windows/D3D11RHI/Private/D3D11StructuredBuffer.cpp
+++ b/Engine/Source/Runtime/Windows/D3D11RHI/Private/D3D11StructuredBuffer.cpp
@@ -201,9 +201,6 @@ void FD3D11DynamicRHI::RHIUnlockStructuredBuffer(FStructuredBufferRHIParamRef St
 			// Copy the contents of the temporary memory buffer allocated for writing into the VB.
 			Direct3DDeviceIMContext->UpdateSubresource(StructuredBuffer->Resource,LockedKey.Subresource,NULL,LockedData->GetData(),LockedData->Pitch,0);
 
-			// Check the copy is finished before freeing...
-			Direct3DDeviceIMContext->Flush();
-
 			// Free the temporary memory buffer.
 			LockedData->FreeData();
 		}
