remove stall on dx11 rhithread caused by csvprofiler
-FRealtimeGPUProfilerEvent no longer forces Query to wait.
-Add AbsoluteTime Queries to FD3D11RenderQueryBatcher, like occlusion batches
-At the end of every frame a readback of all the available queries are scheduled

diff --git a/Engine/Source/Runtime/Engine/Private/SceneUtils.cpp b/Engine/Source/Runtime/Engine/Private/SceneUtils.cpp
--- a/Engine/Source/Runtime/Engine/Private/SceneUtils.cpp
+++ b/Engine/Source/Runtime/Engine/Private/SceneUtils.cpp
@@ -264,13 +264,14 @@ public:
 
 	bool GatherQueryResults(FRHICommandListImmediate& RHICmdList)
 	{
+		QUICK_SCOPE_CYCLE_COUNTER(STAT_SceneUtils_GatherQueryResults);
 		// Get the query results which are still outstanding
 		check(GFrameNumberRenderThread != FrameNumber);
 		if ( HasQueriesAllocated() )
 		{
 			if (StartResultMicroseconds == InvalidQueryResult)
 			{
-				if (!RHICmdList.GetRenderQueryResult(StartQuery, StartResultMicroseconds, true))
+				if (!RHICmdList.GetRenderQueryResult(StartQuery, StartResultMicroseconds, false))
 				{
 					StartResultMicroseconds = InvalidQueryResult;
 				}
@@ -278,7 +279,7 @@ public:
 			}
 			if (EndResultMicroseconds == InvalidQueryResult)
 			{
-				if (!RHICmdList.GetRenderQueryResult(EndQuery, EndResultMicroseconds, true))
+				if (!RHICmdList.GetRenderQueryResult(EndQuery, EndResultMicroseconds, false))
 				{
 					EndResultMicroseconds = InvalidQueryResult;
 				}
diff --git a/Engine/Source/Runtime/Windows/D3D11RHI/Private/D3D11Query.cpp b/Engine/Source/Runtime/Windows/D3D11RHI/Private/D3D11Query.cpp
--- a/Engine/Source/Runtime/Windows/D3D11RHI/Private/D3D11Query.cpp
+++ b/Engine/Source/Runtime/Windows/D3D11RHI/Private/D3D11Query.cpp
@@ -44,12 +44,17 @@ public:
 
 	inline void Add(FRenderQueryRHIRef NewQuery)
 	{
+		if(!ActiveBatches.Num())
+		{
+			StartNewBatch();
+		}
 		FRenderQueryBatch& CurrentBatch = ActiveBatches.Last();
 		CurrentBatch.Queries.Add(NewQuery);
 	}
 
 	void PollQueryResults()
 	{
+		QUICK_SCOPE_CYCLE_COUNTER(STAT_PollQueryResults);
 		FD3D11DynamicRHI* D3D11RHI = static_cast<FD3D11DynamicRHI*>(GDynamicRHI);
 
 		for (int32 BatchIdx = 0; BatchIdx < ActiveBatches.Num(); ++BatchIdx)
@@ -80,6 +85,7 @@ public:
 
 	void PollQueryResults_RenderThread()
 	{
+		QUICK_SCOPE_CYCLE_COUNTER(STAT_PollQueryResults_RenderThread);
 		FScopedRHIThreadStaller RHITStaller(FRHICommandListExecutor::GetImmediateCommandList());
 		PollQueryResults();
 	}
@@ -187,21 +193,30 @@ bool FD3D11DynamicRHI::RHIGetRenderQueryResult(FRenderQueryRHIParamRef QueryRHI,
 {
 	check(IsInRenderingThread());
 	FD3D11RenderQuery* Query = ResourceCast(QueryRHI);
-
-	static uint32 LastQueryBatchId = 0xffffffff;
-	uint32 CurrentQueryBatchId = FD3D11RenderQueryBatcher::Get().GetCurrentId();
-	if (LastQueryBatchId != CurrentQueryBatchId && Query->QueryType == RQT_Occlusion)
+	bool bSuccess = true;
+	if(Query->QueryType == RQT_AbsoluteTime && !bWait)
 	{
-		LastQueryBatchId = CurrentQueryBatchId;
-		FD3D11RenderQueryBatcher::Get().PollQueryResults_RenderThread();
+		if(!Query->bResultIsCached)
+		{
+			return false;
+		}
 	}
-
-	bool bSuccess = true;
-	if(!Query->bResultIsCached)
+	else
 	{
-		bSuccess = GetQueryData(Query->Resource,&Query->Result,sizeof(Query->Result),Query->QueryType,bWait,true);
+		static uint32 LastQueryBatchId = 0xffffffff;
+		uint32 CurrentQueryBatchId = FD3D11RenderQueryBatcher::Get().GetCurrentId();
+		if (LastQueryBatchId != CurrentQueryBatchId && Query->QueryType == RQT_Occlusion)
+		{
+			LastQueryBatchId = CurrentQueryBatchId;
+			FD3D11RenderQueryBatcher::Get().PollQueryResults_RenderThread();
+		}
 
-		Query->bResultIsCached = bSuccess;
+		if(!Query->bResultIsCached)
+		{
+			bSuccess = GetQueryData(Query->Resource,&Query->Result,sizeof(Query->Result),Query->QueryType,bWait,true);
+
+			Query->bResultIsCached = bSuccess;
+		}
 	}
 
 	if(Query->QueryType == RQT_AbsoluteTime)
@@ -243,6 +258,10 @@ void FD3D11DynamicRHI::RHIEndRenderQuery(FRenderQueryRHIParamRef QueryRHI)
 {
 	FD3D11RenderQuery* Query = ResourceCast(QueryRHI);
 	Query->bResultIsCached = false; // for occlusion queries, this is redundant with the one in begin
+	if(Query->QueryType == RQT_AbsoluteTime)
+	{
+		FD3D11RenderQueryBatcher::Get().Add(QueryRHI);
+	}
 	Direct3DDeviceIMContext->End(Query->Resource);
 
 	//@todo - d3d debug spews warnings about OQ's that are being issued but not polled, need to investigate
diff --git a/Engine/Source/Runtime/Windows/D3D11RHI/Private/Windows/WindowsD3D11Device.cpp b/Engine/Source/Runtime/Windows/D3D11RHI/Private/Windows/WindowsD3D11Device.cpp
--- a/Engine/Source/Runtime/Windows/D3D11RHI/Private/Windows/WindowsD3D11Device.cpp
+++ b/Engine/Source/Runtime/Windows/D3D11RHI/Private/Windows/WindowsD3D11Device.cpp
@@ -1676,6 +1676,8 @@ void FD3D11DynamicRHI::InitD3DDevice()
 
 void FD3D11DynamicRHI::RHIPerFrameRHIFlushComplete()
 {
+	RHIPollRenderQueryResults();
+
 	extern void D3D11RHIQueryBatcherPerFrameCleanup();
 
 	D3D11RHIQueryBatcherPerFrameCleanup();
