Make a copy of source data if run UpdateTexture3D on RHI thread

diff --git a/Engine/Source/Runtime/Windows/D3D11RHI/Private/D3D11Texture.cpp b/Engine/Source/Runtime/Windows/D3D11RHI/Private/D3D11Texture.cpp
--- a/Engine/Source/Runtime/Windows/D3D11RHI/Private/D3D11Texture.cpp
+++ b/Engine/Source/Runtime/Windows/D3D11RHI/Private/D3D11Texture.cpp
@@ -1961,9 +1961,12 @@ void FD3D11DynamicRHI::UpdateTexture3D_RenderThread(
 	}
 	else
 	{
-		RunOnRHIThread([this, Texture, MipIndex, UpdateRegion, SourceRowPitch, SourceDepthPitch, SourceData]()
+		const SIZE_T SourceDataSize = SourceDepthPitch * UpdateRegion.Depth;
+		uint8* SourceDataCopy = (uint8*)FMemory::Malloc(SourceDataSize);
+		RunOnRHIThread([this, Texture, MipIndex, UpdateRegion, SourceRowPitch, SourceDepthPitch, SourceDataCopy]()
 		{
-			RHIUpdateTexture3D(Texture, MipIndex, UpdateRegion, SourceRowPitch, SourceDepthPitch, SourceData);
+			RHIUpdateTexture3D(Texture, MipIndex, UpdateRegion, SourceRowPitch, SourceDepthPitch, SourceDataCopy);
+			FMemory::Free(SourceDataCopy);
 		});
 	}
 }
