Make a copy of source data if run UpdateTexture2D on RHI thread

diff --git a/Engine/Source/Runtime/Windows/D3D11RHI/Private/D3D11Texture.cpp b/Engine/Source/Runtime/Windows/D3D11RHI/Private/D3D11Texture.cpp
--- a/Engine/Source/Runtime/Windows/D3D11RHI/Private/D3D11Texture.cpp
+++ b/Engine/Source/Runtime/Windows/D3D11RHI/Private/D3D11Texture.cpp
@@ -1891,9 +1891,12 @@ void FD3D11DynamicRHI::UpdateTexture2D_RenderThread(
 	}
 	else
 	{
-		RunOnRHIThread([this, Texture, MipIndex, UpdateRegion, SourcePitch, SourceData]()
+		const SIZE_T SourceDataSize = SourcePitch * UpdateRegion.Height;
+		uint8* SourceDataCopy = (uint8*)FMemory::Malloc(SourceDataSize);
+		FMemory::Memcpy(SourceDataCopy, SourceData, SourceDataSize);
+		RunOnRHIThread([this, Texture, MipIndex, UpdateRegion, SourcePitch, SourceDataCopy]()
 		{
-			RHIUpdateTexture2D(Texture, MipIndex, UpdateRegion, SourcePitch, SourceData);
+			RHIUpdateTexture2D(Texture, MipIndex, UpdateRegion, SourcePitch, SourceDataCopy);
 		});
 	}
 }
@@ -1963,6 +1966,7 @@ void FD3D11DynamicRHI::UpdateTexture3D_RenderThread(
 	{
 		const SIZE_T SourceDataSize = SourceDepthPitch * UpdateRegion.Depth;
 		uint8* SourceDataCopy = (uint8*)FMemory::Malloc(SourceDataSize);
+		FMemory::Memcpy(SourceDataCopy, SourceData, SourceDataSize);
 		RunOnRHIThread([this, Texture, MipIndex, UpdateRegion, SourceRowPitch, SourceDepthPitch, SourceDataCopy]()
 		{
 			RHIUpdateTexture3D(Texture, MipIndex, UpdateRegion, SourceRowPitch, SourceDepthPitch, SourceDataCopy);
