Fixing editor crash on startup with D3D11 RHI.

diff --git a/Engine/Source/Runtime/Windows/D3D11RHI/Private/D3D11Device.cpp b/Engine/Source/Runtime/Windows/D3D11RHI/Private/D3D11Device.cpp
--- a/Engine/Source/Runtime/Windows/D3D11RHI/Private/D3D11Device.cpp
+++ b/Engine/Source/Runtime/Windows/D3D11RHI/Private/D3D11Device.cpp
@@ -210,12 +210,6 @@ FD3D11DynamicRHI::FD3D11DynamicRHI(IDXGIFactory1* InDXGIFactory1,D3D_FEATURE_LEV
 		GMaxTextureArrayLayers = D3D11_REQ_TEXTURE2D_ARRAY_AXIS_DIMENSION;
 		GRHISupportsMSAADepthSampleAccess = true;
 		GRHISupportsRHIThread = !!EXPERIMENTAL_D3D11_RHITHREAD;
-
-		// Disable the RHI thread by default for devices that will likely suffer in performance
-		if (IsRHIDeviceIntel() || FPlatformMisc::NumberOfCoresIncludingHyperthreads() <= 2)
-		{
-			GRHISupportsRHIThread = false;
-		}
 	}
 	else if (FeatureLevel >= D3D_FEATURE_LEVEL_10_0)
 	{
diff --git a/Engine/Source/Runtime/Windows/D3D11RHI/Private/Windows/WindowsD3D11Device.cpp b/Engine/Source/Runtime/Windows/D3D11RHI/Private/Windows/WindowsD3D11Device.cpp
--- a/Engine/Source/Runtime/Windows/D3D11RHI/Private/Windows/WindowsD3D11Device.cpp
+++ b/Engine/Source/Runtime/Windows/D3D11RHI/Private/Windows/WindowsD3D11Device.cpp
@@ -1592,6 +1592,12 @@ void FD3D11DynamicRHI::InitD3DDevice()
 		}
 #endif // INTEL_METRICSDISCOVERY
 
+		// Disable the RHI thread by default for devices that will likely suffer in performance
+		if (IsRHIDeviceIntel() || FPlatformMisc::NumberOfCoresIncludingHyperthreads() <= 2)
+		{
+			GRHISupportsRHIThread = false;
+		}
+
 		SetupAfterDeviceCreation();
 
 		// Notify all initialized FRenderResources that there's a valid RHI device to create their RHI resources for now.
