Fix a crash when texture2d is lock/unlocked on different threads

diff --git a/Engine/Source/Runtime/Windows/D3D11RHI/Private/D3D11Commands.cpp b/Engine/Source/Runtime/Windows/D3D11RHI/Private/D3D11Commands.cpp
--- a/Engine/Source/Runtime/Windows/D3D11RHI/Private/D3D11Commands.cpp
+++ b/Engine/Source/Runtime/Windows/D3D11RHI/Private/D3D11Commands.cpp
@@ -2155,20 +2155,6 @@ void FD3D11DynamicRHI::RHIFlushComputeShaderCache()
 	EndUAVOverlap();
 }
 
-FD3D11DynamicRHI::FD3D11LockTracker& FD3D11DynamicRHI::GetThreadLocalLockTracker()
-{
-#if EXPERIMENTAL_D3D11_RHITHREAD
-	if (IsInRHIThread())
-	{
-		return LockTrackers[1];
-	}
-	else
-#endif
-	{
-		return LockTrackers[0];
-	}
-}
-
 //*********************** StagingBuffer Implementation ***********************//
 
 FStagingBufferRHIRef FD3D11DynamicRHI::RHICreateStagingBuffer()
diff --git a/Engine/Source/Runtime/Windows/D3D11RHI/Private/D3D11IndexBuffer.cpp b/Engine/Source/Runtime/Windows/D3D11RHI/Private/D3D11IndexBuffer.cpp
--- a/Engine/Source/Runtime/Windows/D3D11RHI/Private/D3D11IndexBuffer.cpp
+++ b/Engine/Source/Runtime/Windows/D3D11RHI/Private/D3D11IndexBuffer.cpp
@@ -134,7 +134,7 @@ void* FD3D11DynamicRHI::RHILockIndexBuffer(FIndexBufferRHIParamRef IndexBufferRH
 	}
 	
 	// Add the lock to the lock map.
-	GetThreadLocalLockTracker().Add(LockedKey, LockedData);
+	AddLockedData(MoveTemp(LockedKey), LockedData);
 
 	// Return the offset pointer
 	return (void*)((uint8*)LockedData.GetData() + Offset);
@@ -149,11 +149,9 @@ void FD3D11DynamicRHI::RHIUnlockIndexBuffer(FIndexBufferRHIParamRef IndexBufferR
 	IndexBuffer->Resource->GetDesc(&Desc);
 	const bool bIsDynamic = (Desc.Usage == D3D11_USAGE_DYNAMIC);
 
-	// Find the outstanding lock for this IB.
-	FD3D11LockTracker& OutstandingLocks = GetThreadLocalLockTracker();
-	FD3D11LockedKey LockedKey(IndexBuffer->Resource);
-	FD3D11LockedData* LockedData = OutstandingLocks.Find(LockedKey);
-	checkf(LockedData, TEXT("Index buffer is either not locked or locked on a different thread"));
+	// Find the outstanding lock for this IB and remove it from the tracker.
+	FD3D11LockedData LockedData;
+	verifyf(RemoveLockedData(FD3D11LockedKey(IndexBuffer->Resource), LockedData), TEXT("Index buffer is not locked"));
 
 	if(bIsDynamic)
 	{
@@ -163,25 +161,21 @@ void FD3D11DynamicRHI::RHIUnlockIndexBuffer(FIndexBufferRHIParamRef IndexBufferR
 	else
 	{
 		// If the static IB lock involved a staging resource, it was locked for reading.
-		if(LockedData->StagingResource)
+		if(LockedData.StagingResource)
 		{
 			// Unmap the staging buffer's memory.
-			ID3D11Buffer* StagingBuffer = (ID3D11Buffer*)LockedData->StagingResource.GetReference();
+			ID3D11Buffer* StagingBuffer = (ID3D11Buffer*)LockedData.StagingResource.GetReference();
 			Direct3DDeviceIMContext->Unmap(StagingBuffer,0);
 		}
 		else 
 		{
 			// Copy the contents of the temporary memory buffer allocated for writing into the IB.
-			Direct3DDeviceIMContext->UpdateSubresource(IndexBuffer->Resource,LockedKey.Subresource,NULL,LockedData->GetData(),LockedData->Pitch,0);
+			Direct3DDeviceIMContext->UpdateSubresource(IndexBuffer->Resource,0,NULL,LockedData.GetData(),LockedData.Pitch,0);
 
 			// Free the temporary memory buffer.
-			LockedData->FreeData();
+			LockedData.FreeData();
 		}
 	}
-
-	// Remove the FD3D11LockedData from the lock map.
-	// If the lock involved a staging resource, this releases it.
-	OutstandingLocks.Remove(LockedKey);
 }
 
 void FD3D11DynamicRHI::RHITransferIndexBufferUnderlyingResource(FIndexBufferRHIParamRef DestIndexBuffer, FIndexBufferRHIParamRef SrcIndexBuffer)
diff --git a/Engine/Source/Runtime/Windows/D3D11RHI/Private/D3D11RHIPrivate.h b/Engine/Source/Runtime/Windows/D3D11RHI/Private/D3D11RHIPrivate.h
--- a/Engine/Source/Runtime/Windows/D3D11RHI/Private/D3D11RHIPrivate.h
+++ b/Engine/Source/Runtime/Windows/D3D11RHI/Private/D3D11RHIPrivate.h
@@ -698,7 +698,20 @@ public:
 		return GPUProfilingData.CheckGpuHeartbeat();
 	}
 
-	FD3D11LockTracker& GetThreadLocalLockTracker();
+	template <typename KeyType, typename ValType,
+		typename = typename TEnableIf<TIsSame<FD3D11LockedKey, typename TDecay<KeyType>::Type>::Value && TIsSame<FD3D11LockedData, typename TDecay<ValType>::Type>::Value>::Type>
+	void AddLockedData(KeyType&& Key, ValType&& LockedData)
+	{
+		FScopeLock Lock(&LockTrackerCS);
+		LockTracker.Add(Forward<KeyType>(Key), Forward<ValType>(LockedData));
+	}
+
+	template <typename KeyType, typename = typename TEnableIf<TIsSame<FD3D11LockedKey, typename TDecay<KeyType>::Type>::Value>::Type>
+	bool RemoveLockedData(KeyType&& Key, FD3D11LockedData& OutLockedData)
+	{
+		FScopeLock Lock(&LockTrackerCS);
+		return LockTracker.RemoveAndCopyValue(Forward<KeyType>(Key), OutLockedData);
+	}
 
 	bool IsQuadBufferStereoEnabled();
 	void DisableQuadBufferStereo();
@@ -795,7 +808,8 @@ protected:
 	FD3D11StateCache StateCache;
 
 	/** Tracks outstanding locks on each thread */
-	FD3D11LockTracker LockTrackers[D3D11_NUM_THREAD_LOCAL_CACHES];
+	FD3D11LockTracker LockTracker;
+	FCriticalSection LockTrackerCS;
 
 	/** A list of all viewport RHIs that have been created. */
 	TArray<FD3D11Viewport*> Viewports;
diff --git a/Engine/Source/Runtime/Windows/D3D11RHI/Private/D3D11StructuredBuffer.cpp b/Engine/Source/Runtime/Windows/D3D11RHI/Private/D3D11StructuredBuffer.cpp
--- a/Engine/Source/Runtime/Windows/D3D11RHI/Private/D3D11StructuredBuffer.cpp
+++ b/Engine/Source/Runtime/Windows/D3D11RHI/Private/D3D11StructuredBuffer.cpp
@@ -161,7 +161,7 @@ void* FD3D11DynamicRHI::RHILockStructuredBuffer(FStructuredBufferRHIParamRef Str
 	}
 
 	// Add the lock to the lock map.
-	GetThreadLocalLockTracker().Add(LockedKey,LockedData);
+	AddLockedData(MoveTemp(LockedKey), LockedData);
 
 	// Return the offset pointer
 	return (void*)((uint8*)LockedData.GetData() + Offset);
@@ -177,10 +177,8 @@ void FD3D11DynamicRHI::RHIUnlockStructuredBuffer(FStructuredBufferRHIParamRef St
 	const bool bIsDynamic = (Desc.Usage == D3D11_USAGE_DYNAMIC);
 
 	// Find the outstanding lock for this VB.
-	FD3D11LockTracker& OutstandingLocks = GetThreadLocalLockTracker();
-	FD3D11LockedKey LockedKey(StructuredBuffer->Resource);
-	FD3D11LockedData* LockedData = OutstandingLocks.Find(LockedKey);
-	checkf(LockedData, TEXT("Structured buffer is either not locked or locked on a different thread"));
+	FD3D11LockedData LockedData;
+	verifyf(RemoveLockedData(FD3D11LockedKey(StructuredBuffer->Resource), LockedData), TEXT("Structured buffer is not locked"));
 
 	if(bIsDynamic)
 	{
@@ -190,23 +188,19 @@ void FD3D11DynamicRHI::RHIUnlockStructuredBuffer(FStructuredBufferRHIParamRef St
 	else
 	{
 		// If the static VB lock involved a staging resource, it was locked for reading.
-		if(LockedData->StagingResource)
+		if(LockedData.StagingResource)
 		{
 			// Unmap the staging buffer's memory.
-			ID3D11Buffer* StagingBuffer = (ID3D11Buffer*)LockedData->StagingResource.GetReference();
+			ID3D11Buffer* StagingBuffer = (ID3D11Buffer*)LockedData.StagingResource.GetReference();
 			Direct3DDeviceIMContext->Unmap(StagingBuffer,0);
 		}
 		else 
 		{
 			// Copy the contents of the temporary memory buffer allocated for writing into the VB.
-			Direct3DDeviceIMContext->UpdateSubresource(StructuredBuffer->Resource,LockedKey.Subresource,NULL,LockedData->GetData(),LockedData->Pitch,0);
+			Direct3DDeviceIMContext->UpdateSubresource(StructuredBuffer->Resource,0,NULL,LockedData.GetData(),LockedData.Pitch,0);
 
 			// Free the temporary memory buffer.
-			LockedData->FreeData();
+			LockedData.FreeData();
 		}
 	}
-
-	// Remove the FD3D11LockedData from the lock map.
-	// If the lock involved a staging resource, this releases it.
-	OutstandingLocks.Remove(LockedKey);
 }
diff --git a/Engine/Source/Runtime/Windows/D3D11RHI/Private/D3D11Texture.cpp b/Engine/Source/Runtime/Windows/D3D11RHI/Private/D3D11Texture.cpp
--- a/Engine/Source/Runtime/Windows/D3D11RHI/Private/D3D11Texture.cpp
+++ b/Engine/Source/Runtime/Windows/D3D11RHI/Private/D3D11Texture.cpp
@@ -1654,7 +1654,7 @@ ETextureReallocationStatus FD3D11DynamicRHI::CancelAsyncReallocateTexture2D_Rend
 }
 
 template<typename RHIResourceType>
-void* TD3D11Texture2D<RHIResourceType>::Lock(uint32 MipIndex,uint32 ArrayIndex,EResourceLockMode LockMode,uint32& DestStride)
+void* TD3D11Texture2D<RHIResourceType>::Lock(uint32 MipIndex,uint32 ArrayIndex,EResourceLockMode LockMode,uint32& DestStride,bool bForceLockDeferred)
 {
 	SCOPE_CYCLE_COUNTER(STAT_D3D11LockTextureTime);
 
@@ -1681,7 +1681,7 @@ void* TD3D11Texture2D<RHIResourceType>::Lock(uint32 MipIndex,uint32 ArrayIndex,E
 #endif
 	if( LockMode == RLM_WriteOnly )
 	{
-		if (Flags & TexCreate_CPUWritable)
+		if (!bForceLockDeferred && (Flags & TexCreate_CPUWritable))
 		{
 			D3D11_MAPPED_SUBRESOURCE MappedTexture;
 			VERIFYD3D11RESULT_EX(D3DRHI->GetDeviceContext()->Map(GetResource(), Subresource, D3D11_MAP_WRITE, 0, &MappedTexture), D3DRHI->GetDevice());
@@ -1693,10 +1693,12 @@ void* TD3D11Texture2D<RHIResourceType>::Lock(uint32 MipIndex,uint32 ArrayIndex,E
 			// If we're writing to the texture, allocate a system memory buffer to receive the new contents.
 			LockedData.AllocData(MipBytes);
 			LockedData.Pitch = DestStride = NumBlocksX * BlockBytes;
+			LockedData.bLockDeferred = true;
 		}
 	}
 	else
 	{
+		check(!bForceLockDeferred);
 		// If we're reading from the texture, we create a staging resource, copy the texture contents to it, and map it.
 
 		// Create the staging texture.
@@ -1742,7 +1744,17 @@ void* TD3D11Texture2D<RHIResourceType>::Lock(uint32 MipIndex,uint32 ArrayIndex,E
 	}
 
 	// Add the lock to the outstanding lock list.
-	D3DRHI->GetThreadLocalLockTracker().Add(FD3D11LockedKey(GetResource(),Subresource),LockedData);
+	if (!bForceLockDeferred)
+	{
+		D3DRHI->AddLockedData(FD3D11LockedKey(GetResource(), Subresource), LockedData);
+	}
+	else
+	{
+		RunOnRHIThread([this, Subresource, LockedData]() mutable
+		{
+			D3DRHI->AddLockedData(FD3D11LockedKey(GetResource(), Subresource), MoveTemp(LockedData));
+		});
+	}
 
 	return (void*)LockedData.GetData();
 }
@@ -1755,11 +1767,9 @@ void TD3D11Texture2D<RHIResourceType>::Unlock(uint32 MipIndex,uint32 ArrayIndex)
 	// Calculate the subresource index corresponding to the specified mip-map.
 	const uint32 Subresource = D3D11CalcSubresource(MipIndex,ArrayIndex,this->GetNumMips());
 
-	// Find the object that is tracking this lock
-	FD3D11DynamicRHI::FD3D11LockTracker& OutstandingLocks = D3DRHI->GetThreadLocalLockTracker();
-	const FD3D11LockedKey LockedKey(GetResource(),Subresource);
-	FD3D11LockedData* LockedData = OutstandingLocks.Find(LockedKey);
-	checkf(LockedData, TEXT("Texture is either not locked or locked on a different thread"));
+	// Find the object that is tracking this lock and remove it from outstanding list
+	FD3D11LockedData LockedData;
+	verifyf(D3DRHI->RemoveLockedData(FD3D11LockedKey(GetResource(), Subresource), LockedData), TEXT("Texture is not locked"));
 
 #if PLATFORM_SUPPORTS_VIRTUAL_TEXTURES
 	if (D3DRHI->HandleSpecialUnlock(MipIndex, GetFlags(), GetResource(), RawTextureMemory))
@@ -1768,19 +1778,20 @@ void TD3D11Texture2D<RHIResourceType>::Unlock(uint32 MipIndex,uint32 ArrayIndex)
 	}
 	else
 #endif
-	if (Flags & TexCreate_CPUWritable)
+	if (!LockedData.bLockDeferred && (Flags & TexCreate_CPUWritable))
 	{
 		D3DRHI->GetDeviceContext()->Unmap(GetResource(), 0);
 	}
-	else if(!LockedData->StagingResource)
+	else if(!LockedData.StagingResource)
 	{
 		// If we're writing, we need to update the subresource
-		D3DRHI->GetDeviceContext()->UpdateSubresource(GetResource(), Subresource, NULL, LockedData->GetData(), LockedData->Pitch, 0);
-		LockedData->FreeData();
+		D3DRHI->GetDeviceContext()->UpdateSubresource(GetResource(), Subresource, NULL, LockedData.GetData(), LockedData.Pitch, 0);
+		LockedData.FreeData();
+	}
+	else
+	{
+		D3DRHI->GetDeviceContext()->Unmap(LockedData.StagingResource, 0);
 	}
-
-	// Remove the lock from the outstanding lock list.
-	OutstandingLocks.Remove(LockedKey);
 }
 
 void* FD3D11DynamicRHI::RHILockTexture2D(FTexture2DRHIParamRef TextureRHI,uint32 MipIndex,EResourceLockMode LockMode,uint32& DestStride,bool bLockWithinMiptail)
@@ -1806,13 +1817,15 @@ void* FD3D11DynamicRHI::LockTexture2D_RenderThread(
 	{
 		LockedTexture = RHILockTexture2D(Texture, MipIndex, LockMode, DestStride, bLockWithinMiptail);
 	}
-	else
+	else if (LockMode == RLM_ReadOnly)
 	{
-		RunOnRHIThread([this, Texture, MipIndex, LockMode, &DestStride, bLockWithinMiptail, &LockedTexture]()
-		{
-			LockedTexture = RHILockTexture2D(Texture, MipIndex, LockMode, DestStride, bLockWithinMiptail);
-		});
 		FRHICommandListExecutor::GetImmediateCommandList().ImmediateFlush(EImmediateFlushType::FlushRHIThread);
+		LockedTexture = RHILockTexture2D(Texture, MipIndex, LockMode, DestStride, bLockWithinMiptail);
+	}
+	else
+	{
+		FD3D11Texture2D* TextureD3D11 = ResourceCast(Texture);
+		LockedTexture = TextureD3D11->Lock(MipIndex, 0, LockMode, DestStride, true);
 	}
 	return LockedTexture;
 }
@@ -2035,13 +2048,16 @@ void* FD3D11DynamicRHI::RHILockTextureCubeFace_RenderThread(
 	{
 		LockedTexture = RHILockTextureCubeFace(Texture, FaceIndex, ArrayIndex, MipIndex, LockMode, DestStride, bLockWithinMiptail);
 	}
-	else
+	else if (LockMode == RLM_ReadOnly)
 	{
-		RunOnRHIThread([this, Texture, FaceIndex, ArrayIndex, MipIndex, LockMode, &DestStride, bLockWithinMiptail, &LockedTexture]()
-		{
-			LockedTexture = RHILockTextureCubeFace(Texture, FaceIndex, ArrayIndex, MipIndex, LockMode, DestStride, bLockWithinMiptail);
-		});
 		FRHICommandListExecutor::GetImmediateCommandList().ImmediateFlush(EImmediateFlushType::FlushRHIThread);
+		LockedTexture = RHILockTextureCubeFace(Texture, FaceIndex, ArrayIndex, MipIndex, LockMode, DestStride, bLockWithinMiptail);
+	}
+	else
+	{
+		FD3D11TextureCube* TextureCube = ResourceCast(Texture);
+		const uint32 D3DFace = GetD3D11CubeFace((ECubeFace)FaceIndex);
+		LockedTexture = TextureCube->Lock(MipIndex, D3DFace + ArrayIndex * 6, LockMode, DestStride, true);
 	}
 	return LockedTexture;
 }
diff --git a/Engine/Source/Runtime/Windows/D3D11RHI/Private/D3D11VertexBuffer.cpp b/Engine/Source/Runtime/Windows/D3D11RHI/Private/D3D11VertexBuffer.cpp
--- a/Engine/Source/Runtime/Windows/D3D11RHI/Private/D3D11VertexBuffer.cpp
+++ b/Engine/Source/Runtime/Windows/D3D11RHI/Private/D3D11VertexBuffer.cpp
@@ -174,7 +174,7 @@ void* FD3D11DynamicRHI::RHILockVertexBuffer(FVertexBufferRHIParamRef VertexBuffe
 	}
 
 	// Add the lock to the lock map.
-	GetThreadLocalLockTracker().Add(LockedKey,LockedData);
+	AddLockedData(MoveTemp(LockedKey), LockedData);
 
 	// Return the offset pointer
 	return (void*)((uint8*)LockedData.GetData() + Offset);
@@ -190,10 +190,8 @@ void FD3D11DynamicRHI::RHIUnlockVertexBuffer(FVertexBufferRHIParamRef VertexBuff
 	const bool bIsDynamic = (Desc.Usage == D3D11_USAGE_DYNAMIC);
 
 	// Find the outstanding lock for this VB.
-	FD3D11LockTracker& OutstandingLocks = GetThreadLocalLockTracker();
-	FD3D11LockedKey LockedKey(VertexBuffer->Resource);
-	FD3D11LockedData* LockedData = OutstandingLocks.Find(LockedKey);
-	checkf(LockedData, TEXT("Vertex buffer is either not locked or locked on a different thread"));
+	FD3D11LockedData LockedData;
+	verifyf(RemoveLockedData(FD3D11LockedKey(VertexBuffer->Resource), LockedData), TEXT("Vertex buffer is not locked"));
 
 	if(bIsDynamic)
 	{
@@ -203,25 +201,21 @@ void FD3D11DynamicRHI::RHIUnlockVertexBuffer(FVertexBufferRHIParamRef VertexBuff
 	else
 	{
 		// If the static VB lock involved a staging resource, it was locked for reading.
-		if(LockedData->StagingResource)
+		if(LockedData.StagingResource)
 		{
 			// Unmap the staging buffer's memory.
-			ID3D11Buffer* StagingBuffer = (ID3D11Buffer*)LockedData->StagingResource.GetReference();
+			ID3D11Buffer* StagingBuffer = (ID3D11Buffer*)LockedData.StagingResource.GetReference();
 			Direct3DDeviceIMContext->Unmap(StagingBuffer,0);
 		}
 		else 
 		{
 			// Copy the contents of the temporary memory buffer allocated for writing into the VB.
-			Direct3DDeviceIMContext->UpdateSubresource(VertexBuffer->Resource,LockedKey.Subresource,NULL,LockedData->GetData(),LockedData->Pitch,0);
+			Direct3DDeviceIMContext->UpdateSubresource(VertexBuffer->Resource,0,NULL,LockedData.GetData(),LockedData.Pitch,0);
 
 			// Free the temporary memory buffer.
-			LockedData->FreeData();
+			LockedData.FreeData();
 		}
 	}
-
-	// Remove the FD3D11LockedData from the lock map.
-	// If the lock involved a staging resource, this releases it.
-	OutstandingLocks.Remove(LockedKey);
 }
 
 void FD3D11DynamicRHI::RHICopyVertexBuffer(FVertexBufferRHIParamRef SourceBufferRHI,FVertexBufferRHIParamRef DestBufferRHI)
diff --git a/Engine/Source/Runtime/Windows/D3D11RHI/Public/D3D11Resources.h b/Engine/Source/Runtime/Windows/D3D11RHI/Public/D3D11Resources.h
--- a/Engine/Source/Runtime/Windows/D3D11RHI/Public/D3D11Resources.h
+++ b/Engine/Source/Runtime/Windows/D3D11RHI/Public/D3D11Resources.h
@@ -418,7 +418,7 @@ public:
 	 * Locks one of the texture's mip-maps.
 	 * @return A pointer to the specified texture data.
 	 */
-	void* Lock(uint32 MipIndex,uint32 ArrayIndex,EResourceLockMode LockMode,uint32& DestStride);
+	void* Lock(uint32 MipIndex,uint32 ArrayIndex,EResourceLockMode LockMode,uint32& DestStride,bool bForceLockDeferred = false);
 
 	/** Unlocks a previously locked mip-map. */
 	void Unlock(uint32 MipIndex,uint32 ArrayIndex);
diff --git a/Engine/Source/Runtime/Windows/D3D11RHI/Public/D3D11Util.h b/Engine/Source/Runtime/Windows/D3D11RHI/Public/D3D11Util.h
--- a/Engine/Source/Runtime/Windows/D3D11RHI/Public/D3D11Util.h
+++ b/Engine/Source/Runtime/Windows/D3D11RHI/Public/D3D11Util.h
@@ -167,6 +167,7 @@ struct FD3D11LockedData
 	// constructor
 	FD3D11LockedData()
 		: bAllocDataWasUsed(false)
+		, bLockDeferred(false)
 	{
 	}
 
@@ -201,6 +202,9 @@ private:
 	uint8* Data;
 	// then FreeData
 	bool bAllocDataWasUsed;
+public:
+	// Whether the lock op is deferred
+	bool bLockDeferred;
 };
 
 /**
