D3D11 RHI - simplifying some code so it is easier to maintain

diff --git a/Engine/Source/Runtime/Windows/D3D11RHI/Private/D3D11IndexBuffer.cpp b/Engine/Source/Runtime/Windows/D3D11RHI/Private/D3D11IndexBuffer.cpp
--- a/Engine/Source/Runtime/Windows/D3D11RHI/Private/D3D11IndexBuffer.cpp
+++ b/Engine/Source/Runtime/Windows/D3D11RHI/Private/D3D11IndexBuffer.cpp
@@ -134,7 +134,7 @@ void* FD3D11DynamicRHI::RHILockIndexBuffer(FIndexBufferRHIParamRef IndexBufferRH
 	}
 	
 	// Add the lock to the lock map.
-	AddLockedData(MoveTemp(LockedKey), LockedData);
+	AddLockedData(LockedKey, LockedData);
 
 	// Return the offset pointer
 	return (void*)((uint8*)LockedData.GetData() + Offset);
diff --git a/Engine/Source/Runtime/Windows/D3D11RHI/Private/D3D11RHIPrivate.h b/Engine/Source/Runtime/Windows/D3D11RHI/Private/D3D11RHIPrivate.h
--- a/Engine/Source/Runtime/Windows/D3D11RHI/Private/D3D11RHIPrivate.h
+++ b/Engine/Source/Runtime/Windows/D3D11RHI/Private/D3D11RHIPrivate.h
@@ -698,19 +698,16 @@ public:
 		return GPUProfilingData.CheckGpuHeartbeat();
 	}
 
-	template <typename KeyType, typename ValType,
-		typename = typename TEnableIf<TIsSame<FD3D11LockedKey, typename TDecay<KeyType>::Type>::Value && TIsSame<FD3D11LockedData, typename TDecay<ValType>::Type>::Value>::Type>
-	void AddLockedData(KeyType&& Key, ValType&& LockedData)
+	void AddLockedData(const FD3D11LockedKey& Key, const FD3D11LockedData& LockedData)
 	{
 		FScopeLock Lock(&LockTrackerCS);
-		LockTracker.Add(Forward<KeyType>(Key), Forward<ValType>(LockedData));
+		LockTracker.Add(Key, LockedData);
 	}
 
-	template <typename KeyType, typename = typename TEnableIf<TIsSame<FD3D11LockedKey, typename TDecay<KeyType>::Type>::Value>::Type>
-	bool RemoveLockedData(KeyType&& Key, FD3D11LockedData& OutLockedData)
+	bool RemoveLockedData(const FD3D11LockedKey& Key, FD3D11LockedData& OutLockedData)
 	{
 		FScopeLock Lock(&LockTrackerCS);
-		return LockTracker.RemoveAndCopyValue(Forward<KeyType>(Key), OutLockedData);
+		return LockTracker.RemoveAndCopyValue(Key, OutLockedData);
 	}
 
 	bool IsQuadBufferStereoEnabled();
diff --git a/Engine/Source/Runtime/Windows/D3D11RHI/Private/D3D11StructuredBuffer.cpp b/Engine/Source/Runtime/Windows/D3D11RHI/Private/D3D11StructuredBuffer.cpp
--- a/Engine/Source/Runtime/Windows/D3D11RHI/Private/D3D11StructuredBuffer.cpp
+++ b/Engine/Source/Runtime/Windows/D3D11RHI/Private/D3D11StructuredBuffer.cpp
@@ -161,7 +161,7 @@ void* FD3D11DynamicRHI::RHILockStructuredBuffer(FStructuredBufferRHIParamRef Str
 	}
 
 	// Add the lock to the lock map.
-	AddLockedData(MoveTemp(LockedKey), LockedData);
+	AddLockedData(LockedKey, LockedData);
 
 	// Return the offset pointer
 	return (void*)((uint8*)LockedData.GetData() + Offset);
diff --git a/Engine/Source/Runtime/Windows/D3D11RHI/Private/D3D11Texture.cpp b/Engine/Source/Runtime/Windows/D3D11RHI/Private/D3D11Texture.cpp
--- a/Engine/Source/Runtime/Windows/D3D11RHI/Private/D3D11Texture.cpp
+++ b/Engine/Source/Runtime/Windows/D3D11RHI/Private/D3D11Texture.cpp
@@ -1750,9 +1750,9 @@ void* TD3D11Texture2D<RHIResourceType>::Lock(uint32 MipIndex,uint32 ArrayIndex,E
 	}
 	else
 	{
-		RunOnRHIThread([this, Subresource, LockedData]() mutable
+		RunOnRHIThread([this, Subresource, LockedData]()
 		{
-			D3DRHI->AddLockedData(FD3D11LockedKey(GetResource(), Subresource), MoveTemp(LockedData));
+			D3DRHI->AddLockedData(FD3D11LockedKey(GetResource(), Subresource), LockedData);
 		});
 	}
 
diff --git a/Engine/Source/Runtime/Windows/D3D11RHI/Private/D3D11VertexBuffer.cpp b/Engine/Source/Runtime/Windows/D3D11RHI/Private/D3D11VertexBuffer.cpp
--- a/Engine/Source/Runtime/Windows/D3D11RHI/Private/D3D11VertexBuffer.cpp
+++ b/Engine/Source/Runtime/Windows/D3D11RHI/Private/D3D11VertexBuffer.cpp
@@ -174,7 +174,7 @@ void* FD3D11DynamicRHI::RHILockVertexBuffer(FVertexBufferRHIParamRef VertexBuffe
 	}
 
 	// Add the lock to the lock map.
-	AddLockedData(MoveTemp(LockedKey), LockedData);
+	AddLockedData(LockedKey, LockedData);
 
 	// Return the offset pointer
 	return (void*)((uint8*)LockedData.GetData() + Offset);
