Adjusting D3D11 RHI toggle threshold based on 8.30 performance data.

diff --git a/Engine/Source/Runtime/Windows/D3D11RHI/Private/Windows/WindowsD3D11Device.cpp b/Engine/Source/Runtime/Windows/D3D11RHI/Private/Windows/WindowsD3D11Device.cpp
--- a/Engine/Source/Runtime/Windows/D3D11RHI/Private/Windows/WindowsD3D11Device.cpp
+++ b/Engine/Source/Runtime/Windows/D3D11RHI/Private/Windows/WindowsD3D11Device.cpp
@@ -1593,7 +1593,7 @@ void FD3D11DynamicRHI::InitD3DDevice()
 #endif // INTEL_METRICSDISCOVERY
 
 		// Disable the RHI thread by default for devices that will likely suffer in performance
-		if (IsRHIDeviceIntel() || FPlatformMisc::NumberOfCoresIncludingHyperthreads() <= 2)
+		if (IsRHIDeviceIntel() || FPlatformMisc::NumberOfCores() < 4)
 		{
 			GRHISupportsRHIThread = false;
 		}
