Fix for crash encountered when AudioMixer is toggled on twice in the same session
aaron.mcleran, rob.gay

diff --git a/Engine/Source/Runtime/Windows/AudioMixerXAudio2/Private/AudioMixerPlatformXAudio2.cpp b/Engine/Source/Runtime/Windows/AudioMixerXAudio2/Private/AudioMixerPlatformXAudio2.cpp
--- a/Engine/Source/Runtime/Windows/AudioMixerXAudio2/Private/AudioMixerPlatformXAudio2.cpp
+++ b/Engine/Source/Runtime/Windows/AudioMixerXAudio2/Private/AudioMixerPlatformXAudio2.cpp
@@ -229,6 +229,10 @@ namespace Audio
 
 		SAFE_RELEASE(XAudio2System);
 
+#if WITH_XMA2
+		FXMAAudioInfo::Shutdown();
+#endif
+
 #if PLATFORM_WINDOWS
 
 #if PLATFORM_64BITS
