Resolving a memory leak in texture updates on the RHI thread.

diff --git a/Engine/Source/Runtime/Windows/D3D11RHI/Private/D3D11Texture.cpp b/Engine/Source/Runtime/Windows/D3D11RHI/Private/D3D11Texture.cpp
--- a/Engine/Source/Runtime/Windows/D3D11RHI/Private/D3D11Texture.cpp
+++ b/Engine/Source/Runtime/Windows/D3D11RHI/Private/D3D11Texture.cpp
@@ -1910,6 +1910,7 @@ void FD3D11DynamicRHI::UpdateTexture2D_RenderThread(
 		RunOnRHIThread([this, Texture, MipIndex, UpdateRegion, SourcePitch, SourceDataCopy]()
 		{
 			RHIUpdateTexture2D(Texture, MipIndex, UpdateRegion, SourcePitch, SourceDataCopy);
+			FMemory::Free(SourceDataCopy);
 		});
 	}
 }
