[Andriy.Tylychko 's unshelved CL 6053257]
Implements pause/unpause in windows video recording
Stops platform dependent video recorder from editing timestamps of shared FGameplayMediaEncoderSample
[peter.lockhart additions]
Fixed crash where we try to determine if a sample is a video type - after moving the object, the underlying Sample object is null, so we can't query it for its flags

#rnx

diff --git a/Engine/Source/Runtime/Windows/WindowsPlatformFeatures/Private/HighlightRecorder.cpp b/Engine/Source/Runtime/Windows/WindowsPlatformFeatures/Private/HighlightRecorder.cpp
--- a/Engine/Source/Runtime/Windows/WindowsPlatformFeatures/Private/HighlightRecorder.cpp
+++ b/Engine/Source/Runtime/Windows/WindowsPlatformFeatures/Private/HighlightRecorder.cpp
@@ -124,7 +124,7 @@ FTimespan FHighlightRecorder::GetRecordingTime() const
 	return FTimespan::FromSeconds(FPlatformTime::Seconds()) - RecordingStartTime - TotalPausedDuration;
 }
 
-void FHighlightRecorder::OnMediaSample(const FGameplayMediaEncoderSample& Sample)
+void FHighlightRecorder::OnMediaSample(const FGameplayMediaEncoderSample& InSample)
 {
 	// We might be paused, so don't do anything
 	if (State != EState::Recording)
@@ -133,9 +133,9 @@ void FHighlightRecorder::OnMediaSample(const FGameplayMediaEncoderSample& Sample
 	}
 
 	// Only start pushing video frames once we receive a key frame
-	if (NumPushedFrames == 0 && Sample.GetType() == EMediaType::Video)
+	if (NumPushedFrames == 0 && InSample.GetType() == EMediaType::Video)
 	{
-		if (!Sample.IsVideoKeyFrame())
+		if (!InSample.IsVideoKeyFrame())
 		{
 			return;
 		}
@@ -143,7 +143,14 @@ void FHighlightRecorder::OnMediaSample(const FGameplayMediaEncoderSample& Sample
 		++NumPushedFrames;
 	}
 
-	RingBuffer.Push(Sample);
+	FGameplayMediaEncoderSample SampleCopy = InSample.Clone();
+
+	if (TotalPausedDuration != 0)
+	{
+		SampleCopy.SetTime(SampleCopy.GetTime() - TotalPausedDuration);
+	}
+
+	RingBuffer.Push(MoveTemp(SampleCopy));
 }
 
 bool FHighlightRecorder::SaveHighlight(const TCHAR* Filename, FDoneCallback InDoneCallback, double MaxDurationSecs)
diff --git a/Engine/Source/Runtime/Windows/WindowsPlatformFeatures/Private/WmfRingBuffer.cpp b/Engine/Source/Runtime/Windows/WindowsPlatformFeatures/Private/WmfRingBuffer.cpp
--- a/Engine/Source/Runtime/Windows/WindowsPlatformFeatures/Private/WmfRingBuffer.cpp
+++ b/Engine/Source/Runtime/Windows/WindowsPlatformFeatures/Private/WmfRingBuffer.cpp
@@ -8,7 +8,7 @@ DEFINE_LOG_CATEGORY(WmfRingBuffer);
 
 WINDOWSPLATFORMFEATURES_START
 
-void FWmfRingBuffer::Push(const FGameplayMediaEncoderSample& Sample)
+void FWmfRingBuffer::Push(FGameplayMediaEncoderSample&& Sample)
 {
 	check(MaxDuration != 0);
 
@@ -21,7 +21,7 @@ void FWmfRingBuffer::Push(const FGameplayMediaEncoderSample& Sample)
 	}
 	else if (Samples.Num() == 0)
 	{
-		Samples.Add(Sample);
+		Samples.Add(MoveTemp(Sample));
 	}
 	else
 		// video encoding takes longer than audio thus video encoded samples arrive later and need to be placed
@@ -29,13 +29,15 @@ void FWmfRingBuffer::Push(const FGameplayMediaEncoderSample& Sample)
 	{
 		FTimespan TimestampToInsert = Sample.GetTime();
 
+		bool bIsSampleVideoKeyFrame = Sample.IsVideoKeyFrame();
+
 		int i = Samples.Num() - 1;
 		while (true)
 		{
 			if (Samples[i].GetType() == Sample.GetType() // don't change order of same stream cos video is not strictly timestamp-ordered (B-Frames)
 				|| TimestampToInsert >= Samples[i].GetTime())
 			{
-				Samples.Insert(Sample, i + 1);
+				Samples.Insert(MoveTemp(Sample), i + 1);
 				break;
 			}
 
@@ -43,7 +45,7 @@ void FWmfRingBuffer::Push(const FGameplayMediaEncoderSample& Sample)
 			{
 				if (Sample.GetType() == EMediaType::Video)
 				{
-					Samples.Insert(Sample, 0);
+					Samples.Insert(MoveTemp(Sample), 0);
 				}
 				// else: the first sample in the buffer always should be video key-frame, just drop audio sample
 				// this can happen during debugging due to big gaps in timestamps
@@ -55,7 +57,7 @@ void FWmfRingBuffer::Push(const FGameplayMediaEncoderSample& Sample)
 		}
 
 		// cleanup only on receiving key-frame and only when ring-buffer if full
-		if (!bCleanupPaused && Sample.IsVideoKeyFrame())
+		if (!bCleanupPaused && bIsSampleVideoKeyFrame)
 		{
 			while (GetDuration() > GetMaxDuration())
 			{
@@ -116,7 +118,7 @@ TArray<FGameplayMediaEncoderSample> FWmfRingBuffer::GetCopy()
 	TArray<FGameplayMediaEncoderSample> Copy;
 	Copy.Reset(Samples.Num());
 
-	for (const auto& Sample: Samples)
+	for (const auto& Sample : Samples)
 	{
 		// cloning is required because during saving a copy of ring buffer we modify samples timestamps.
 		// using samples references shared between ring buffer and its copy being saved would cause ring
diff --git a/Engine/Source/Runtime/Windows/WindowsPlatformFeatures/Private/WmfRingBuffer.h b/Engine/Source/Runtime/Windows/WindowsPlatformFeatures/Private/WmfRingBuffer.h
--- a/Engine/Source/Runtime/Windows/WindowsPlatformFeatures/Private/WmfRingBuffer.h
+++ b/Engine/Source/Runtime/Windows/WindowsPlatformFeatures/Private/WmfRingBuffer.h
@@ -11,14 +11,18 @@ class FWmfRingBuffer
 {
 public:
 	FTimespan GetMaxDuration() const
-	{ return MaxDuration; }
+	{
+		return MaxDuration;
+	}
 
 	void SetMaxDuration(FTimespan InMaxDuration)
-	{ MaxDuration = InMaxDuration; }
+	{
+		MaxDuration = InMaxDuration;
+	}
 
 	FTimespan GetDuration() const;
 
-	void Push(const FGameplayMediaEncoderSample& Sample);
+	void Push(FGameplayMediaEncoderSample&& Sample);
 
 	void PauseCleanup(bool bPause);
 
