When a video has been recorded, and we aren't auto continuing on to the next video, it is important to reset the internal windows recorder, so upon next "start recording", the recorder is recreated. Otherwise, our next video will contain frames from the previous video, because FWindowsVideoRecordingSystem::NewRecording only creates a new ring buffer/writer if the underlying recorder is "stopped"

#rnx

diff --git a/Engine/Source/Runtime/Windows/WindowsPlatformFeatures/Private/WindowsVideoRecordingSystem.cpp b/Engine/Source/Runtime/Windows/WindowsPlatformFeatures/Private/WindowsVideoRecordingSystem.cpp
--- a/Engine/Source/Runtime/Windows/WindowsPlatformFeatures/Private/WindowsVideoRecordingSystem.cpp
+++ b/Engine/Source/Runtime/Windows/WindowsPlatformFeatures/Private/WindowsVideoRecordingSystem.cpp
@@ -248,6 +248,10 @@ void FWindowsVideoRecordingSystem::FinalizeCallbackOnGameThread(bool bSaved, boo
 		NextRecording();
 		RecordState = EVideoRecordingState::Recording;
 	}
+	else
+	{
+		Recorder->Stop();
+	}
 
 	if (bBroadcast)
 	{
