HDR - Add HDR display modes to Windows editor when running with D3D11 RHI and on Windows 10 > 1803

diff --git a/Engine/Shaders/Private/CompositeUIPixelShader.usf b/Engine/Shaders/Private/CompositeUIPixelShader.usf
--- a/Engine/Shaders/Private/CompositeUIPixelShader.usf
+++ b/Engine/Shaders/Private/CompositeUIPixelShader.usf
@@ -72,6 +72,32 @@ float3 ColorLookupTable(float3 LinearColor)
 	return ColorSpaceLUT.SampleLevel(ColorSpaceLUTSampler, UVW, 0).rgb * LUT_SCALING;
 }
 
+// Given an input using Rec709 primaries and a sRGB tranfer function where
+// 1.0 maps to paperwhitenits(160) remap to Rec2020 with PQ transfer function.
+void HDREditorConvert(
+	FScreenVertexOutput Input,
+	out float4 OutColor : SV_Target0
+	)
+{
+	OutColor = 0;
+	float3 SceneColor = Texture2DSample(SceneTexture, SceneSampler, Input.UV).xyz ;
+
+	// Get color into Rec2020 space
+	SceneColor = sRGBToLinear( SceneColor) ;
+	SceneColor = SceneColor * UILevel ; 
+
+	SceneColor = mul(sRGB_2_XYZ_MAT, SceneColor);
+	SceneColor = mul(XYZ_2_Rec2020_MAT, SceneColor);
+
+	// Transform into P3 (optional instead of 2020)
+	//	SceneColor = mul(sRGB_2_XYZ_MAT, SceneColor);
+	//	SceneColor = mul(XYZ_2_P3D65_MAT, SceneColor);
+
+	// Do PQ Curve
+	SceneColor = LinearToST2084( SceneColor ); 
+	OutColor.rgb = SceneColor;
+}
+
 void Main(
 	FScreenVertexOutput Input,
 	out float4 OutColor : SV_Target0
diff --git a/Engine/Shaders/Private/PostProcessTonemap.usf b/Engine/Shaders/Private/PostProcessTonemap.usf
--- a/Engine/Shaders/Private/PostProcessTonemap.usf
+++ b/Engine/Shaders/Private/PostProcessTonemap.usf
@@ -50,6 +50,7 @@ float4 SceneUVMinMax;
 float4 SceneBloomUVMinMax;
 
 float4 ScreenPosToScenePixel;
+float4 HDRParams;
 
 #if EYEADAPTATION_EXPOSURE_FIX != 1
 // Default value used when eye adaptation is disabled (e.g fixed exposure) or not available.
@@ -563,6 +564,15 @@ float4 TonemapCommonPS(
 	}
 #endif //DEBUG_TONEMAPPER
 
+
+	// If HDR in the editor need to convert from PQ space to linear sRGB
+	BRANCH
+	if(HDRParams.x == 1.0f)
+	{
+		OutColor.rgb = ST2084ToLinear(OutColor.rgb);
+		OutColor.rgb = LinearToSrgbBranchless( OutColor.rgb / HDRParams.y );
+	}
+
 	return OutColor;
 }
 
diff --git a/Engine/Source/Editor/UnrealEd/Classes/Settings/EditorExperimentalSettings.h b/Engine/Source/Editor/UnrealEd/Classes/Settings/EditorExperimentalSettings.h
--- a/Engine/Source/Editor/UnrealEd/Classes/Settings/EditorExperimentalSettings.h
+++ b/Engine/Source/Editor/UnrealEd/Classes/Settings/EditorExperimentalSettings.h
@@ -20,6 +20,14 @@ class UNREALED_API UEditorExperimentalSettings
 
 public:
 
+	/** Allows the editor to run on HDR monitors on Windows 10 */
+	UPROPERTY(EditAnywhere, config, Category = HDR, meta = (ConfigRestartRequired = true, DisplayName = "Enable Editor Support for HDR Monitors"))
+	bool bHDREditor;
+
+	/** The brightness of the slate UI on HDR monitors */
+	UPROPERTY(EditAnywhere, config, Category = HDR, meta = (ClampMin = "100.0", ClampMax = "300.0", UIMin = "100.0", UIMax = "300.0"))
+	float HDREditorNITLevel;
+
 	/** Allows usage of the procedural foliage system */
 	UPROPERTY(EditAnywhere, config, Category = Foliage, meta = (DisplayName = "Procedural Foliage"))
 	bool bProceduralFoliage;
@@ -142,8 +150,8 @@ public:
 protected:
 
 	// UObject overrides
-
-	virtual void PostEditChangeProperty( struct FPropertyChangedEvent& PropertyChangedEvent ) override;
+	virtual void PostEditChangeProperty(struct FPropertyChangedEvent& PropertyChangedEvent) override;
+	virtual void PostInitProperties() override;
 
 private:
 
diff --git a/Engine/Source/Editor/UnrealEd/Private/EditorEngine.cpp b/Engine/Source/Editor/UnrealEd/Private/EditorEngine.cpp
--- a/Engine/Source/Editor/UnrealEd/Private/EditorEngine.cpp
+++ b/Engine/Source/Editor/UnrealEd/Private/EditorEngine.cpp
@@ -1991,6 +1991,14 @@ bool UEditorEngine::UpdateSingleViewportClient(FEditorViewportClient* InViewport
 {
 	bool bUpdatedNonRealtimeViewport = false;
 
+	// When rendering the viewport we need to know whether the final result will be shown on a HDR display. This affects the final post processing step
+	FSceneViewport *SceneViewport = static_cast<FSceneViewport*>(InViewportClient->Viewport);
+	TSharedPtr<SWindow> Window = SceneViewport->FindWindow();
+	if (Window)
+	{
+		InViewportClient->Viewport->SetHDRMode(Window->GetIsHDR());
+	}
+
 	// Always submit view information for content streaming 
 	// otherwise content for editor view can be streamed out if there are other views (ex: thumbnails)
 	if (InViewportClient->IsPerspective())
diff --git a/Engine/Source/Editor/UnrealEd/Private/EditorViewportClient.cpp b/Engine/Source/Editor/UnrealEd/Private/EditorViewportClient.cpp
--- a/Engine/Source/Editor/UnrealEd/Private/EditorViewportClient.cpp
+++ b/Engine/Source/Editor/UnrealEd/Private/EditorViewportClient.cpp
@@ -3634,6 +3634,8 @@ void FEditorViewportClient::Draw(FViewport* InViewport, FCanvas* Canvas)
 
 	ViewFamily.EngineShowFlags = EngineShowFlags;
 
+	ViewFamily.bIsHDR = Viewport->IsHDRViewport();
+	
 	UpdateDebugViewModeShaders();
 
 	if( ModeTools->GetActiveMode( FBuiltinEditorModes::EM_InterpEdit ) == 0 || !AllowsCinematicControl() )
diff --git a/Engine/Source/Editor/UnrealEd/Private/Settings/SettingsClasses.cpp b/Engine/Source/Editor/UnrealEd/Private/Settings/SettingsClasses.cpp
--- a/Engine/Source/Editor/UnrealEd/Private/Settings/SettingsClasses.cpp
+++ b/Engine/Source/Editor/UnrealEd/Private/Settings/SettingsClasses.cpp
@@ -111,8 +111,22 @@ USkeletalMeshEditorSettings::USkeletalMeshEditorSettings(const FObjectInitialize
 /* UEditorExperimentalSettings interface
  *****************************************************************************/
 
+static TAutoConsoleVariable<int32> CVarEditorHDRSupport(
+	TEXT("Editor.HDRSupport"),
+	0,
+	TEXT("Sets whether or not we should allow the editor to run on HDR monitors"),
+	ECVF_Default);
+
+static TAutoConsoleVariable<float> CVarEditorHDRNITLevel(
+	TEXT("Editor.HDRNITLevel"),
+	150.0f,
+	TEXT("Sets The desired NIT level of the editor when running on HDR"),
+	ECVF_Default);
+
 UEditorExperimentalSettings::UEditorExperimentalSettings( const FObjectInitializer& ObjectInitializer )
 	: Super(ObjectInitializer)
+	, bHDREditor(false)
+	, HDREditorNITLevel(160.0f)
 	, bLandscapeLayerSystem(false)
 	, bEnableLocalizationDashboard(true)
 	, bUseOpenCLForConvexHullDecomp(false)
@@ -120,6 +134,13 @@ UEditorExperimentalSettings::UEditorExperimentalSettings( const FObjectInitializ
 {
 }
 
+void UEditorExperimentalSettings::PostInitProperties()
+{
+	CVarEditorHDRSupport->Set(bHDREditor ? 1 : 0, ECVF_SetByProjectSetting);
+	CVarEditorHDRNITLevel->Set(HDREditorNITLevel, ECVF_SetByProjectSetting);
+	Super::PostInitProperties();
+}
+
 void UEditorExperimentalSettings::PostEditChangeProperty( struct FPropertyChangedEvent& PropertyChangedEvent )
 {
 	static const FName NAME_EQS = GET_MEMBER_NAME_CHECKED(UEditorExperimentalSettings, bEQSEditor);
@@ -139,7 +160,14 @@ void UEditorExperimentalSettings::PostEditChangeProperty( struct FPropertyChange
 			FModuleManager::Get().LoadModule(TEXT("EnvironmentQueryEditor"));
 		}
 	}
-
+	else if (Name == FName(TEXT("bHDREditor")))
+	{
+		CVarEditorHDRSupport->Set(bHDREditor ? 1 : 0, ECVF_SetByProjectSetting);
+	}
+	else if (Name == FName(TEXT("HDREditorNITLevel")))
+	{
+		CVarEditorHDRNITLevel->Set(HDREditorNITLevel, ECVF_SetByProjectSetting);
+	}
 	if (!FUnrealEdMisc::Get().IsDeletePreferences())
 	{
 		SaveConfig();
diff --git a/Engine/Source/Runtime/D3D12RHI/Private/D3D12Commands.cpp b/Engine/Source/Runtime/D3D12RHI/Private/D3D12Commands.cpp
--- a/Engine/Source/Runtime/D3D12RHI/Private/D3D12Commands.cpp
+++ b/Engine/Source/Runtime/D3D12RHI/Private/D3D12Commands.cpp
@@ -1935,6 +1935,11 @@ uint32 FD3D12DynamicRHI::RHIGetGPUFrameCycles()
 	return GGPUFrameTime;
 }
 
+EColorSpace FD3D12DynamicRHI::RHIGetColorSpace(FViewportRHIParamRef Viewport)
+{
+	return EColorSpace::ERec709sRGB;
+}
+
 void FD3D12DynamicRHI::RHIExecuteCommandList(FRHICommandList* CmdList)
 {
 	check(0); // this path has gone stale and needs updated methods, starting at ERCT_SetScissorRect
diff --git a/Engine/Source/Runtime/D3D12RHI/Private/D3D12RHIPrivate.h b/Engine/Source/Runtime/D3D12RHI/Private/D3D12RHIPrivate.h
--- a/Engine/Source/Runtime/D3D12RHI/Private/D3D12RHIPrivate.h
+++ b/Engine/Source/Runtime/D3D12RHI/Private/D3D12RHIPrivate.h
@@ -356,6 +356,7 @@ public:
 	virtual FViewportRHIRef RHICreateViewport(void* WindowHandle, uint32 SizeX, uint32 SizeY, bool bIsFullscreen, EPixelFormat PreferredPixelFormat) final override;
 	virtual void RHIResizeViewport(FViewportRHIParamRef Viewport, uint32 SizeX, uint32 SizeY, bool bIsFullscreen) final override;
 	virtual void RHIResizeViewport(FViewportRHIParamRef ViewportRHI, uint32 SizeX, uint32 SizeY, bool bIsFullscreen, EPixelFormat PreferredPixelFormat) final override;
+	virtual EColorSpace RHIGetColorSpace(FViewportRHIParamRef Viewport) final override;
 	virtual void RHITick(float DeltaTime) final override;
 	virtual void RHISetStreamOutTargets(uint32 NumTargets, const FVertexBufferRHIParamRef* VertexBuffers, const uint32* Offsets) final override;
 	virtual void RHIBlockUntilGPUIdle() final override;
diff --git a/Engine/Source/Runtime/Engine/Private/Slate/SceneViewport.cpp b/Engine/Source/Runtime/Engine/Private/Slate/SceneViewport.cpp
--- a/Engine/Source/Runtime/Engine/Private/Slate/SceneViewport.cpp
+++ b/Engine/Source/Runtime/Engine/Private/Slate/SceneViewport.cpp
@@ -1935,6 +1935,16 @@ void FSceneViewport::InitDynamicRHI()
 
 		static const auto CVarDefaultBackBufferPixelFormat = IConsoleManager::Get().FindTConsoleVariableDataInt(TEXT("r.DefaultBackBufferPixelFormat"));
 		EPixelFormat SceneTargetFormat = EDefaultBackBufferPixelFormat::Convert2PixelFormat(EDefaultBackBufferPixelFormat::FromInt(CVarDefaultBackBufferPixelFormat->GetValueOnRenderThread()));
+	
+#if WITH_EDITOR
+		// HDR Editor needs to be in float format if running with HDR
+		// TODO - Just set this when we know we want HDR capabilites
+		static auto CVarHDREnable = IConsoleManager::Get().FindConsoleVariable(TEXT("Editor.HDRSupport"));
+		if(CVarHDREnable && (CVarHDREnable->GetInt() != 0))
+		{
+			SceneTargetFormat = PF_FloatRGBA;
+		}
+#endif
 
 		FRHIResourceCreateInfo CreateInfo;
 		FTexture2DRHIRef BufferedRTRHI;
diff --git a/Engine/Source/Runtime/Engine/Public/SceneView.h b/Engine/Source/Runtime/Engine/Public/SceneView.h
--- a/Engine/Source/Runtime/Engine/Public/SceneView.h
+++ b/Engine/Source/Runtime/Engine/Public/SceneView.h
@@ -1481,6 +1481,9 @@ public:
 	 */
 	bool bWorldIsPaused;
 
+	/** When enabled, the post processing will output in HDR space */
+	bool bIsHDR;
+
 	/** Gamma correction used when rendering this family. Default is 1.0 */
 	float GammaCorrection;
 	
diff --git a/Engine/Source/Runtime/Engine/Public/UnrealClient.h b/Engine/Source/Runtime/Engine/Public/UnrealClient.h
--- a/Engine/Source/Runtime/Engine/Public/UnrealClient.h
+++ b/Engine/Source/Runtime/Engine/Public/UnrealClient.h
@@ -572,6 +572,12 @@ public:
 	/** Returns true if this is an FSlateSceneViewport */
 	bool IsSlateViewport() const { return bIsSlateViewport; }
 
+	/** Returns true if this viewport should be rendered in HDR */
+	bool IsHDRViewport() const { return bIsHDR; }
+
+	/** Sets HDR Status of Viewport */
+	void SetHDRMode( bool bHDR) { bIsHDR = bHDR; }
+
 	/** The current version of the running instance */
 	FString AppVersionString;
 
@@ -693,6 +699,9 @@ protected:
 	/** If true this viewport is an FSlateSceneViewport */
 	uint32 bIsSlateViewport : 1;
 
+	/** If true this viewport is being displayed on a HDR monitor */
+	uint32 bIsHDR : 1;
+
 	/** true if we should draw game viewports (has no effect on Editor viewports) */
 	ENGINE_API static bool bIsGameRenderingEnabled;
 
@@ -1179,7 +1188,8 @@ public:
 		FTexture2DRHIRef ShaderResourceTextureRHI;
 
 		FRHIResourceCreateInfo CreateInfo;
-		RHICreateTargetableShaderResource2D( SizeX, SizeY, PF_A2B10G10R10, 1, TexCreate_None, TexCreate_RenderTargetable, false, CreateInfo, RenderTargetTextureRHI, ShaderResourceTextureRHI );
+		RHICreateTargetableShaderResource2D(SizeX, SizeY, PF_FloatRGB, 1, TexCreate_None, TexCreate_RenderTargetable, false, CreateInfo, RenderTargetTextureRHI, ShaderResourceTextureRHI);
+//		RHICreateTargetableShaderResource2D(SizeX, SizeY, PF_A2B10G10R10, 1, TexCreate_None, TexCreate_RenderTargetable, false, CreateInfo, RenderTargetTextureRHI, ShaderResourceTextureRHI);
 	}
 
 	// @todo UE4 DLL: Without these functions we get unresolved linker errors with FRenderResource
diff --git a/Engine/Source/Runtime/NullDrv/Public/NullRHI.h b/Engine/Source/Runtime/NullDrv/Public/NullRHI.h
--- a/Engine/Source/Runtime/NullDrv/Public/NullRHI.h
+++ b/Engine/Source/Runtime/NullDrv/Public/NullRHI.h
@@ -525,6 +525,12 @@ public:
 	{
 
 	}
+
+	virtual EColorSpace RHIGetColorSpace(FViewportRHIParamRef Viewport) final override
+	{
+		return EColorSpace::ERec709sRGB;
+	}
+
 	virtual void RHITick(float DeltaTime) final override
 	{
 
diff --git a/Engine/Source/Runtime/RHI/Private/DynamicRHI.cpp b/Engine/Source/Runtime/RHI/Private/DynamicRHI.cpp
--- a/Engine/Source/Runtime/RHI/Private/DynamicRHI.cpp
+++ b/Engine/Source/Runtime/RHI/Private/DynamicRHI.cpp
@@ -354,3 +354,9 @@ void FDynamicRHI::RHIUpdateShaderResourceView(FShaderResourceViewRHIParamRef SRV
 {
 	UE_LOG(LogRHI, Fatal, TEXT("RHIUpdateShaderResourceView isn't implemented for the current RHI"));
 }
+
+
+EColorSpace FDynamicRHI::RHIGetColorSpace(FViewportRHIParamRef Viewport)
+{
+	return EColorSpace::ERec709sRGB;
+}
diff --git a/Engine/Source/Runtime/RHI/Public/DynamicRHI.h b/Engine/Source/Runtime/RHI/Public/DynamicRHI.h
--- a/Engine/Source/Runtime/RHI/Public/DynamicRHI.h
+++ b/Engine/Source/Runtime/RHI/Public/DynamicRHI.h
@@ -925,6 +925,8 @@ public:
 		RHIResizeViewport(Viewport, SizeX, SizeY, bIsFullscreen);
 	}
 
+	virtual EColorSpace RHIGetColorSpace(FViewportRHIParamRef Viewport);
+
 	//  must be called from the main thread.
 	// FlushType: Thread safe
 	virtual void RHITick(float DeltaTime) = 0;
@@ -1316,6 +1318,11 @@ FORCEINLINE void RHIResizeViewport(FViewportRHIParamRef Viewport, uint32 SizeX,
 	GDynamicRHI->RHIResizeViewport(Viewport, SizeX, SizeY, bIsFullscreen, PreferredPixelFormat);
 }
 
+FORCEINLINE EColorSpace RHIGetColorSpace(FViewportRHIParamRef Viewport)
+{
+	return GDynamicRHI->RHIGetColorSpace(Viewport);
+}
+
 FORCEINLINE void RHITick(float DeltaTime)
 {
 	GDynamicRHI->RHITick(DeltaTime);
diff --git a/Engine/Source/Runtime/RHI/Public/RHI.h b/Engine/Source/Runtime/RHI/Public/RHI.h
--- a/Engine/Source/Runtime/RHI/Public/RHI.h
+++ b/Engine/Source/Runtime/RHI/Public/RHI.h
@@ -1068,6 +1068,12 @@ enum class EClearBinding
 	EDepthStencilBound, //target has a depthstencil value bound.  Clears will use the bound values and do hardware clears.
 };
 
+enum class EColorSpace
+{
+	ERec709sRGB,				// 709 primaries with sRGB transfer 
+	ERec2020PQ,					// 2020 primaries with PQ transfer 
+};
+
 struct FClearValueBinding
 {
 	struct DSVAlue
diff --git a/Engine/Source/Runtime/Renderer/Private/PostProcess/PostProcessCombineLUTs.cpp b/Engine/Source/Runtime/Renderer/Private/PostProcess/PostProcessCombineLUTs.cpp
--- a/Engine/Source/Runtime/Renderer/Private/PostProcess/PostProcessCombineLUTs.cpp
+++ b/Engine/Source/Runtime/Renderer/Private/PostProcess/PostProcessCombineLUTs.cpp
@@ -244,6 +244,12 @@ public:
 			int32 OutputDeviceValue = CVarOutputDevice->GetValueOnRenderThread();
 			float Gamma = CVarOutputGamma->GetValueOnRenderThread();
 
+			// If the Viewfamily is requesting HDR for editor usage then override the CVAR setting
+			if (ViewFamily.bIsHDR)
+			{
+				OutputDeviceValue = 3;
+			}
+
 			if (PLATFORM_APPLE && Gamma == 0.0f)
 			{
 				Gamma = 2.2f;
diff --git a/Engine/Source/Runtime/Renderer/Private/PostProcess/PostProcessTonemap.cpp b/Engine/Source/Runtime/Renderer/Private/PostProcess/PostProcessTonemap.cpp
--- a/Engine/Source/Runtime/Renderer/Private/PostProcess/PostProcessTonemap.cpp
+++ b/Engine/Source/Runtime/Renderer/Private/PostProcess/PostProcessTonemap.cpp
@@ -475,6 +475,7 @@ public:
 		ScreenPosToScenePixel.Bind(ParameterMap, TEXT("ScreenPosToScenePixel"));
 		SceneUVMinMax.Bind(ParameterMap, TEXT("SceneUVMinMax"));
 		SceneBloomUVMinMax.Bind(ParameterMap, TEXT("SceneBloomUVMinMax"));
+		HDRParams.Bind(ParameterMap, TEXT("HDRParams"));
 	}
 	
 	template <typename TRHICmdList, typename TRHIShader>
@@ -510,6 +511,20 @@ public:
 			SetShaderValue(RHICmdList, ShaderRHI, TonemapperParams, Value);
 		}
 
+		{
+			float EditorNITLevel = 160.0f;
+#if WITH_EDITOR
+			static auto CVarHDRNITLevel = IConsoleManager::Get().FindConsoleVariable(TEXT("Editor.HDRNITLevel"));
+			if (CVarHDRNITLevel)
+			{
+				EditorNITLevel = CVarHDRNITLevel->GetFloat();
+			}
+#endif
+			FVector4 Value(	ViewFamily.bIsHDR ? 1.0f : 0.0f, EditorNITLevel, 0.0f, 0.0f);
+			SetShaderValue(RHICmdList, ShaderRHI, HDRParams, Value);
+		}
+
+
 		FVector GrainValue;
 		GrainPostSettings(&GrainValue, &Settings);
 		SetShaderValue(RHICmdList, ShaderRHI, GrainScaleBiasJitter, GrainValue);
@@ -656,7 +671,7 @@ public:
 		Ar << P.ColorGradingLUT << P.ColorGradingLUTSampler;
 		Ar << P.SceneUVMinMax << P.SceneBloomUVMinMax;
 		Ar << P.ChromaticAberrationParams << P.ScreenPosToScenePixel;
-
+		Ar << P.HDRParams;
 		return Ar;
 	}
 
@@ -674,6 +689,7 @@ public:
 
 	FShaderParameter SceneUVMinMax;
 	FShaderParameter SceneBloomUVMinMax;
+	FShaderParameter HDRParams;
 };
 
 
@@ -989,7 +1005,15 @@ void FRCPassPostProcessTonemap::Process(FRenderingCompositePassContext& Context)
 			DesktopPermutationVector.Set<TonemapperPermutation::FTonemapperColorFringeDim>(View.FinalPostProcessSettings.SceneFringeIntensity > 0.01f);
 		}
 
-		DesktopPermutationVector.Set<TonemapperPermutation::FTonemapperOutputDeviceDim>(GetOutputDeviceValue());
+		FTonemapperOutputDevice deviceValue = GetOutputDeviceValue();
+
+		// Override based on the ViewFamily HDR field for editor HDR
+		if (ViewFamily.bIsHDR)
+		{
+			deviceValue = FTonemapperOutputDevice::ACES1000nitST2084;
+		}
+
+		DesktopPermutationVector.Set<TonemapperPermutation::FTonemapperOutputDeviceDim>(deviceValue);
 
 		DesktopPermutationVector = TonemapperPermutation::RemapPermutation(DesktopPermutationVector);
 	}
diff --git a/Engine/Source/Runtime/Renderer/Private/PostProcess/PostProcessing.cpp b/Engine/Source/Runtime/Renderer/Private/PostProcess/PostProcessing.cpp
--- a/Engine/Source/Runtime/Renderer/Private/PostProcess/PostProcessing.cpp
+++ b/Engine/Source/Runtime/Renderer/Private/PostProcess/PostProcessing.cpp
@@ -1473,7 +1473,7 @@ void FPostProcessing::Process(FRHICommandListImmediate& RHICmdList, const FViewI
 		const bool bHDROutputEnabled = GRHISupportsHDROutput && IsHDREnabled();
 
 		static const auto CVarDumpFramesAsHDR = IConsoleManager::Get().FindTConsoleVariableDataInt(TEXT("r.BufferVisualizationDumpFramesAsHDR"));
-		const bool bHDRTonemapperOutput = bAllowTonemapper && (GetHighResScreenshotConfig().bCaptureHDR || CVarDumpFramesAsHDR->GetValueOnRenderThread() || bHDROutputEnabled);
+		const bool bHDRTonemapperOutput = bAllowTonemapper && (GetHighResScreenshotConfig().bCaptureHDR || CVarDumpFramesAsHDR->GetValueOnRenderThread() || bHDROutputEnabled || View.Family->bIsHDR);
 
 		FRCPassPostProcessTonemap* Tonemapper = 0;
 
diff --git a/Engine/Source/Runtime/SlateCore/Public/Widgets/SWindow.h b/Engine/Source/Runtime/SlateCore/Public/Widgets/SWindow.h
--- a/Engine/Source/Runtime/SlateCore/Public/Widgets/SWindow.h
+++ b/Engine/Source/Runtime/SlateCore/Public/Widgets/SWindow.h
@@ -745,6 +745,13 @@ public:
 		bIsMirrorWindow = bSetMirrorWindow;
 	}
 	
+	void SetIsHDR(bool bHDR)
+	{
+		bIsHDR = bHDR;
+	}
+
+	bool GetIsHDR() const { return bIsHDR; }
+
 	bool IsVirtualWindow() const { return bVirtualWindow; }
 
 	bool IsMirrorWindow()
@@ -1011,6 +1018,9 @@ protected:
 	/** True if the window should preserve its aspect ratio when resized by user */
 	bool bShouldPreserveAspectRatio : 1;
 
+	/** True if the window is being displayed on a HDR capable monitor */
+	bool bIsHDR : 1;
+
 	bool bManualManageDPI : 1;
 
 	/** When should the window be activated upon being shown */
diff --git a/Engine/Source/Runtime/SlateRHIRenderer/Private/SlateRHIRenderer.cpp b/Engine/Source/Runtime/SlateRHIRenderer/Private/SlateRHIRenderer.cpp
--- a/Engine/Source/Runtime/SlateRHIRenderer/Private/SlateRHIRenderer.cpp
+++ b/Engine/Source/Runtime/SlateRHIRenderer/Private/SlateRHIRenderer.cpp
@@ -326,6 +326,10 @@ void FSlateRHIRenderer::CreateViewport(const TSharedRef<SWindow> Window)
 		NewInfo->ViewportRHI = RHICreateViewport( NewInfo->OSWindow, Width, Height, bFullscreen, NewInfo->PixelFormat );
 		NewInfo->bFullscreen = bFullscreen;
 
+		// Was the window created on a HDR compatable display?
+		NewInfo->bHDREnabled = RHIGetColorSpace(NewInfo->ViewportRHI) == EColorSpace::ERec2020PQ ;
+		Window->SetIsHDR(NewInfo->bHDREnabled);
+
 		WindowToViewportInfo.Add(&Window.Get(), NewInfo);
 
 		BeginInitResource(NewInfo);
@@ -652,6 +656,15 @@ static FAutoConsoleVariableRef CVarSlateWireframe(TEXT("Slate.ShowWireFrame"), S
 void FSlateRHIRenderer::DrawWindow_RenderThread(FRHICommandListImmediate& RHICmdList, FViewportInfo& ViewportInfo, FSlateWindowElementList& WindowElementList, const struct FSlateDrawWindowCommandParams& DrawCommandParams)
 {
 	LLM_SCOPE(ELLMTag::SceneRender);
+	
+	bool bRenderOffscreen = false;	// Render to an offscreen texture which can then be finally color converted at the end.
+	
+#if WITH_EDITOR
+	if (RHIGetColorSpace(ViewportInfo.ViewportRHI)== EColorSpace::ERec2020PQ)
+	{
+		bRenderOffscreen = true;
+	}
+#endif
 
 	static uint32 LastTimestamp = FPlatformTime::Cycles();
 	{
@@ -759,6 +772,23 @@ void FSlateRHIRenderer::DrawWindow_RenderThread(FRHICommandListImmediate& RHICmd
 				bLUTStale = true;
 			}
 
+			TRefCountPtr<IPooledRenderTarget> HDRRenderRT;
+
+			if (bRenderOffscreen )
+			{
+				FPooledRenderTargetDesc Desc(FPooledRenderTargetDesc::Create2DDesc(FIntPoint(ViewportWidth, ViewportHeight),
+					PF_FloatRGBA,
+					FClearValueBinding::Transparent,
+					TexCreate_None,
+					TexCreate_ShaderResource | TexCreate_RenderTargetable,
+					false,
+					1,
+					true,
+					true));
+
+				GRenderTargetPool.FindFreeElement(RHICmdList, Desc, HDRRenderRT, TEXT("HDRTargetRT"));
+			}
+
 			FTexture2DRHIRef FinalBuffer = BackBuffer;
 
 			bool bClear = DrawCommandParams.bClear;
@@ -774,6 +804,13 @@ void FSlateRHIRenderer::DrawWindow_RenderThread(FRHICommandListImmediate& RHICmd
 				BackBuffer = ViewportInfo.UITargetRT->GetRenderTargetItem().TargetableTexture->GetTexture2D();
 			}
 
+			if (bRenderOffscreen)
+			{
+				FResolveParams ResolveParams;
+				RHICmdList.CopyToResolveTarget(FinalBuffer, FinalBuffer, ResolveParams);
+				BackBuffer = HDRRenderRT->GetRenderTargetItem().TargetableTexture->GetTexture2D();
+			}
+
 			if (SlateWireFrame)
 			{
 				bClear = true;
@@ -970,6 +1007,53 @@ void FSlateRHIRenderer::DrawWindow_RenderThread(FRHICommandListImmediate& RHICmd
 				BackBuffer = FinalBuffer;
 			} //bCompositeUI
 
+
+#if WITH_EDITOR
+			if (bRenderOffscreen)
+			{
+				const auto FeatureLevel = GMaxRHIFeatureLevel;
+				auto ShaderMap = GetGlobalShaderMap(FeatureLevel);
+
+				FRHIRenderPassInfo RPInfo(FinalBuffer, ERenderTargetActions::Load_Store);
+				RHICmdList.BeginRenderPass(RPInfo, TEXT("SlateComposite"));
+
+				FGraphicsPipelineStateInitializer GraphicsPSOInit;
+				RHICmdList.ApplyCachedRenderTargets(GraphicsPSOInit);
+				GraphicsPSOInit.BlendState			= TStaticBlendState<>::GetRHI();
+				GraphicsPSOInit.RasterizerState		= TStaticRasterizerState<>::GetRHI();
+				GraphicsPSOInit.DepthStencilState	= TStaticDepthStencilState<false, CF_Always>::GetRHI();
+
+				// ST2084 (PQ) encoding
+				TShaderMapRef<FHDREditorConvertPS> PixelShader(ShaderMap);
+				TShaderMapRef<FScreenVS> VertexShader(ShaderMap);
+
+				GraphicsPSOInit.BoundShaderState.VertexDeclarationRHI = GFilterVertexDeclaration.VertexDeclarationRHI;
+				GraphicsPSOInit.BoundShaderState.VertexShaderRHI = GETSAFERHISHADER_VERTEX(*VertexShader);
+				GraphicsPSOInit.BoundShaderState.PixelShaderRHI = GETSAFERHISHADER_PIXEL(*PixelShader);
+				GraphicsPSOInit.PrimitiveType = PT_TriangleList;
+
+				SetGraphicsPipelineState(RHICmdList, GraphicsPSOInit);
+
+				PixelShader->SetParameters(RHICmdList, HDRRenderRT->GetRenderTargetItem().TargetableTexture );
+
+				static const FName RendererModuleName("Renderer");
+				IRendererModule& RendererModule = FModuleManager::GetModuleChecked<IRendererModule>(RendererModuleName);
+
+				RendererModule.DrawRectangle(
+					RHICmdList,
+					0, 0,
+					ViewportWidth, ViewportHeight,
+					0, 0,
+					ViewportWidth, ViewportHeight,
+					FIntPoint(ViewportWidth, ViewportHeight),
+					FIntPoint(ViewportWidth, ViewportHeight),
+					*VertexShader,
+					EDRF_UseTriangleOptimization);
+				
+				RHICmdList.EndRenderPass();
+				BackBuffer = FinalBuffer;
+			}
+#endif
 			if (!bRenderedStereo && GEngine && IsValidRef(ViewportInfo.GetRenderTargetTexture()) && GEngine->StereoRenderingDevice.IsValid())
 			{
 				const FVector2D WindowSize = WindowElementList.GetWindowSize();
@@ -1129,6 +1213,10 @@ void FSlateRHIRenderer::DrawWindows_Private(FSlateDrawBuffer& WindowDrawBuffer)
 				// The viewport had better exist at this point  
 				FViewportInfo* ViewInfo = WindowToViewportInfo.FindChecked(Window);
 
+				// Cache off the HDR status
+				ViewInfo->bHDREnabled = RHIGetColorSpace(ViewInfo->ViewportRHI) == EColorSpace::ERec2020PQ;
+				Window->SetIsHDR(ViewInfo->bHDREnabled);
+
 				if (Window->IsViewportSizeDrivenByWindow())
 				{
 					// Resize the viewport if needed
diff --git a/Engine/Source/Runtime/SlateRHIRenderer/Private/SlateRHIRenderer.h b/Engine/Source/Runtime/SlateRHIRenderer/Private/SlateRHIRenderer.h
--- a/Engine/Source/Runtime/SlateRHIRenderer/Private/SlateRHIRenderer.h
+++ b/Engine/Source/Runtime/SlateRHIRenderer/Private/SlateRHIRenderer.h
@@ -95,6 +95,8 @@ struct FViewportInfo : public FRenderResource
 	int32 HDRColorGamut;
 	/** Device format for output to HDR display */
 	int32 HDROutputDevice;
+	/** Whether is in a HDR Color Space */
+	bool	bHDREnabled;
 
 	IViewportRenderTargetProvider* RTProvider;
 	
@@ -114,6 +116,7 @@ struct FViewportInfo : public FRenderResource
 			bFullscreen(false),
 			PixelFormat(EPixelFormat::PF_Unknown),
 			SDRPixelFormat(EPixelFormat::PF_Unknown),
+			bHDREnabled(false),
 			RTProvider(nullptr)
 	{
 	}
diff --git a/Engine/Source/Runtime/SlateRHIRenderer/Private/SlateShaders.cpp b/Engine/Source/Runtime/SlateRHIRenderer/Private/SlateShaders.cpp
--- a/Engine/Source/Runtime/SlateRHIRenderer/Private/SlateShaders.cpp
+++ b/Engine/Source/Runtime/SlateRHIRenderer/Private/SlateShaders.cpp
@@ -27,6 +27,9 @@ IMPLEMENT_SHADER_TYPE(, FSlateDebugBatchingPS, TEXT("/Engine/Private/SlateElemen
 	typedef TSlateElementPS<ESlateShader::ShaderType,bDrawDisabledEffect,bUseTextureAlpha> TSlateElementPS##ShaderType##bDrawDisabledEffect##bUseTextureAlpha##A; \
 	IMPLEMENT_SHADER_TYPE(template<>,TSlateElementPS##ShaderType##bDrawDisabledEffect##bUseTextureAlpha##A,TEXT("/Engine/Private/SlateElementPixelShader.usf"),TEXT("Main"),SF_Pixel);
 
+#if WITH_EDITOR
+IMPLEMENT_SHADER_TYPE(, FHDREditorConvertPS, TEXT("/Engine/Private/CompositeUIPixelShader.usf"), TEXT("HDREditorConvert"), SF_Pixel);
+#endif
 /**
 * All the different permutations of shaders used by slate. Uses #defines to avoid dynamic branches
 */
diff --git a/Engine/Source/Runtime/SlateRHIRenderer/Private/SlateShaders.h b/Engine/Source/Runtime/SlateRHIRenderer/Private/SlateShaders.h
--- a/Engine/Source/Runtime/SlateRHIRenderer/Private/SlateShaders.h
+++ b/Engine/Source/Runtime/SlateRHIRenderer/Private/SlateShaders.h
@@ -524,6 +524,66 @@ public:
 };
 
 
+#if WITH_EDITOR
+// Pixel shader to convert UI from linear rec709 to PQ 2020 for HDR monitors
+class FHDREditorConvertPS : public FGlobalShader
+{
+	DECLARE_SHADER_TYPE(FHDREditorConvertPS, Global);
+public:
+
+	static bool ShouldCompilePermutation(const FGlobalShaderPermutationParameters& Parameters)
+	{
+		return true;
+	}
+
+	FHDREditorConvertPS(const ShaderMetaType::CompiledShaderInitializerType& Initializer) :
+		FGlobalShader(Initializer)
+	{
+		SceneTexture.Bind(Initializer.ParameterMap, TEXT("SceneTexture"));
+		SceneSampler.Bind(Initializer.ParameterMap, TEXT("SceneSampler"));
+		UILevel.Bind(Initializer.ParameterMap, TEXT("UILevel"));
+	}
+	FHDREditorConvertPS() {}
+
+	void SetParameters(FRHICommandList& RHICmdList, FTextureRHIParamRef SceneTextureRHI)
+	{
+		static const auto CVarOutputDevice = IConsoleManager::Get().FindTConsoleVariableDataInt(TEXT("r.HDR.Display.OutputDevice"));
+
+		SetTextureParameter(RHICmdList, GetPixelShader(), SceneTexture, SceneSampler, TStaticSamplerState<SF_Point>::GetRHI(), SceneTextureRHI);
+		
+		static auto CVarHDRNITLevel = IConsoleManager::Get().FindConsoleVariable(TEXT("Editor.HDRNITLevel"));
+		SetShaderValue(RHICmdList, GetPixelShader(), UILevel, CVarHDRNITLevel->GetFloat());
+	}
+
+	static void ModifyCompilationEnvironment(const FGlobalShaderPermutationParameters& Parameters, FShaderCompilerEnvironment& OutEnvironment)
+	{
+		FGlobalShader::ModifyCompilationEnvironment(Parameters, OutEnvironment);
+	}
+
+	virtual bool Serialize(FArchive& Ar) override
+	{
+		bool bShaderHasOutdatedParameters = FGlobalShader::Serialize(Ar);
+		Ar << SceneTexture << SceneSampler << UILevel;
+		return bShaderHasOutdatedParameters;
+	}
+
+	static const TCHAR* GetSourceFilename()
+	{
+		return TEXT("/Engine/Private/CompositeUIPixelShader.usf");
+	}
+
+	static const TCHAR* GetFunctionName()
+	{
+		return TEXT("HDREditorConvert");
+	}
+
+private:
+	FShaderResourceParameter SceneTexture;
+	FShaderResourceParameter SceneSampler;
+	FShaderParameter UILevel;
+};
+#endif
+
 /** The simple element vertex declaration. */
 extern TGlobalResource<FSlateVertexDeclaration> GSlateVertexDeclaration;
 
diff --git a/Engine/Source/Runtime/Windows/D3D11RHI/Private/D3D11Device.cpp b/Engine/Source/Runtime/Windows/D3D11RHI/Private/D3D11Device.cpp
--- a/Engine/Source/Runtime/Windows/D3D11RHI/Private/D3D11Device.cpp
+++ b/Engine/Source/Runtime/Windows/D3D11RHI/Private/D3D11Device.cpp
@@ -100,6 +100,33 @@ FD3D11DynamicRHI::FD3D11DynamicRHI(IDXGIFactory1* InDXGIFactory1,D3D_FEATURE_LEV
 		GSupportsDepthFetchDuringDepthTest = false;
 	}
 
+   	
+	bAllowHDR		= false;
+	bAllowFlip		= true;
+	bAllowTearing	= true;
+
+	TRefCountPtr<IDXGIFactory5> Factory5;
+	HRESULT HResult = DXGIFactory1->QueryInterface(IID_PPV_ARGS(Factory5.GetInitReference()));
+	if (SUCCEEDED(HResult))
+	{
+		HResult = Factory5->CheckFeatureSupport(DXGI_FEATURE_PRESENT_ALLOW_TEARING, &bAllowTearing, sizeof(bAllowTearing));
+		bAllowHDR = true;
+	}
+
+	if (FAILED(HResult))
+	{
+		bAllowTearing = false;
+	}
+
+	// Disable FLIP if not on a supporting OS
+	TRefCountPtr<IDXGIFactory4> Factory4;
+	HResult = DXGIFactory1->QueryInterface(IID_PPV_ARGS(Factory4.GetInitReference()));
+	if (FAILED(HResult))
+	{
+		bAllowFlip =false;
+	}
+
+
 	ERHIFeatureLevel::Type PreviewFeatureLevel;
 	if (RHIGetPreviewFeatureLevel(PreviewFeatureLevel))
 	{
diff --git a/Engine/Source/Runtime/Windows/D3D11RHI/Private/D3D11RHIPrivate.h b/Engine/Source/Runtime/Windows/D3D11RHI/Private/D3D11RHIPrivate.h
--- a/Engine/Source/Runtime/Windows/D3D11RHI/Private/D3D11RHIPrivate.h
+++ b/Engine/Source/Runtime/Windows/D3D11RHI/Private/D3D11RHIPrivate.h
@@ -484,6 +484,7 @@ public:
 	virtual FViewportRHIRef RHICreateViewport(void* WindowHandle, uint32 SizeX, uint32 SizeY, bool bIsFullscreen, EPixelFormat PreferredPixelFormat) final override;
 	virtual void RHIResizeViewport(FViewportRHIParamRef Viewport, uint32 SizeX, uint32 SizeY, bool bIsFullscreen) final override;
 	virtual void RHIResizeViewport(FViewportRHIParamRef Viewport, uint32 SizeX, uint32 SizeY, bool bIsFullscreen, EPixelFormat PreferredPixelFormat) final override;
+	virtual EColorSpace RHIGetColorSpace(FViewportRHIParamRef Viewport) final override;
 	virtual void RHITick(float DeltaTime) final override;
 	virtual void RHISetStreamOutTargets(uint32 NumTargets,const FVertexBufferRHIParamRef* VertexBuffers,const uint32* Offsets) final override;
 	virtual void RHIBlockUntilGPUIdle() final override;
@@ -788,6 +789,11 @@ protected:
 	/** The global D3D interface. */
 	TRefCountPtr<IDXGIFactory1> DXGIFactory1;
 
+	// Flags determined from the particualt factories available
+	bool bAllowTearing;
+	bool bAllowHDR;
+	bool bAllowFlip;
+
 	/** The global D3D device's immediate context */
 	TRefCountPtr<FD3D11DeviceContext> Direct3DDeviceIMContext;
 
diff --git a/Engine/Source/Runtime/Windows/D3D11RHI/Private/D3D11Viewport.cpp b/Engine/Source/Runtime/Windows/D3D11RHI/Private/D3D11Viewport.cpp
--- a/Engine/Source/Runtime/Windows/D3D11RHI/Private/D3D11Viewport.cpp
+++ b/Engine/Source/Runtime/Windows/D3D11RHI/Private/D3D11Viewport.cpp
@@ -254,7 +254,7 @@ void FD3D11Viewport::Resize(uint32 InSizeX, uint32 InSizeY, bool bInIsFullscreen
 		{
 			// Resize the swap chain.
 			DXGI_FORMAT RenderTargetFormat = GetRenderTargetFormat(PixelFormat);
-			VERIFYD3D11RESIZEVIEWPORTRESULT(SwapChain->ResizeBuffers(1, SizeX, SizeY, RenderTargetFormat, DXGI_SWAP_CHAIN_FLAG_ALLOW_MODE_SWITCH), SizeX, SizeY, RenderTargetFormat, D3DRHI->GetDevice());
+			VERIFYD3D11RESIZEVIEWPORTRESULT(SwapChain->ResizeBuffers(BackBufferCount, SizeX, SizeY, RenderTargetFormat, DXGI_SWAP_CHAIN_FLAG_ALLOW_MODE_SWITCH), SizeX, SizeY, RenderTargetFormat, D3DRHI->GetDevice());
 
 			if (bInIsFullscreen)
 			{
@@ -538,6 +538,12 @@ bool FD3D11Viewport::Present(bool bLockToVsync)
 	return bNativelyPresented;
 }
 
+EColorSpace FD3D11DynamicRHI::RHIGetColorSpace(FViewportRHIParamRef ViewportRHI)
+{
+	FD3D11Viewport* Viewport = ResourceCast(ViewportRHI);
+	return Viewport->GetPixelColorSpace();
+}
+
 /*=============================================================================
  *	The following RHI functions must be called from the main thread.
  *=============================================================================*/
diff --git a/Engine/Source/Runtime/Windows/D3D11RHI/Private/Windows/WindowsD3D11Viewport.cpp b/Engine/Source/Runtime/Windows/D3D11RHI/Private/Windows/WindowsD3D11Viewport.cpp
--- a/Engine/Source/Runtime/Windows/D3D11RHI/Private/Windows/WindowsD3D11Viewport.cpp
+++ b/Engine/Source/Runtime/Windows/D3D11RHI/Private/Windows/WindowsD3D11Viewport.cpp
@@ -11,7 +11,12 @@
 #include "Windows/AllowWindowsPlatformTypes.h"
 #include <dwmapi.h>
 
-#include "dxgi1_2.h"
+#if defined(NTDDI_WIN10_RS2)
+#include <dxgi1_6.h>
+#else
+#include <dxgi1_5.h>
+#endif
+
 
 extern FD3D11Texture2D* GetSwapChainSurface(FD3D11DynamicRHI* D3DRHI, EPixelFormat PixelFormat, uint32 SizeX, uint32 SizeY, IDXGISwapChain* SwapChain);
 
@@ -28,6 +33,7 @@ FD3D11Viewport::FD3D11Viewport(FD3D11DynamicRHI* InD3DRHI,HWND InWindowHandle,ui
 	SizeY(InSizeY),
 	bIsFullscreen(bInIsFullscreen),
 	PixelFormat(InPreferredPixelFormat),
+	PixelColorSpace(EColorSpace::ERec709sRGB),
 	bIsValid(true),
 	FrameSyncEvent(InD3DRHI)
 {
@@ -112,24 +118,91 @@ FD3D11Viewport::FD3D11Viewport(FD3D11DynamicRHI* InD3DRHI,HWND InWindowHandle,ui
 		if (SwapChain == nullptr)
 		{
 			// Create the swapchain.
-			DXGI_SWAP_CHAIN_DESC SwapChainDesc;
-			FMemory::Memzero(&SwapChainDesc, sizeof(DXGI_SWAP_CHAIN_DESC));
-
-			SwapChainDesc.BufferDesc = SetupDXGI_MODE_DESC();
-			// MSAA Sample count
+			DXGI_SWAP_CHAIN_DESC1 SwapChainDesc;
+			FMemory::Memzero(&SwapChainDesc, sizeof(DXGI_SWAP_CHAIN_DESC1));
+			SwapChainDesc.Width = SizeX;
+			SwapChainDesc.Height = SizeY;
 			SwapChainDesc.SampleDesc.Count = 1;
 			SwapChainDesc.SampleDesc.Quality = 0;
+			SwapChainDesc.Format = DXGI_FORMAT_R10G10B10A2_UNORM;
 			SwapChainDesc.BufferUsage = DXGI_USAGE_RENDER_TARGET_OUTPUT | DXGI_USAGE_SHADER_INPUT;
-			// 1:single buffering, 2:double buffering, 3:triple buffering
-			SwapChainDesc.BufferCount = 1;
-			SwapChainDesc.OutputWindow = WindowHandle;
-			SwapChainDesc.Windowed = !bIsFullscreen;
-			// DXGI_SWAP_EFFECT_DISCARD / DXGI_SWAP_EFFECT_SEQUENTIAL
-			SwapChainDesc.SwapEffect = DXGI_SWAP_EFFECT_DISCARD;
+
+			DXGI_SWAP_CHAIN_FULLSCREEN_DESC fsSwapChainDesc = {};
+			fsSwapChainDesc.Windowed = !bIsFullscreen;
+
+			
+			if (InD3DRHI->bAllowTearing || InD3DRHI->bAllowHDR || InD3DRHI->bAllowFlip)
+			{
+				// Needed for HDR
+				BackBufferCount = 2;
+				SwapChainDesc.SwapEffect = DXGI_SWAP_EFFECT_FLIP_DISCARD;
+			}
+			else
+			{
+				// Optional for LDR
+				BackBufferCount = 1;
+				SwapChainDesc.SwapEffect = DXGI_SWAP_EFFECT_DISCARD;
+			}
+
+			SwapChainDesc.BufferCount = BackBufferCount;
 			SwapChainDesc.Flags = DXGI_SWAP_CHAIN_FLAG_ALLOW_MODE_SWITCH;
-			VERIFYD3D11RESULT_EX(D3DRHI->GetFactory()->CreateSwapChain(DXGIDevice, &SwapChainDesc, SwapChain.GetInitReference()), D3DRHI->GetDevice());
+
+			IDXGISwapChain1* SwapChain1 = nullptr;
+			IDXGIFactory2* Factory2 = (IDXGIFactory2*)D3DRHI->GetFactory();
+
+			VERIFYD3D11RESULT_EX((Factory2->CreateSwapChainForHwnd(D3DRHI->GetDevice(), WindowHandle, &SwapChainDesc, &fsSwapChainDesc, nullptr, &SwapChain1)), D3DRHI->GetDevice());
+			SwapChain1->QueryInterface(__uuidof(IDXGISwapChain1), (void**)SwapChain.GetInitReference());
+
+			// See if we are running on a HDR monitor
+			
+			EColorSpace ColorSpace = EColorSpace::ERec709sRGB;
+#if WITH_EDITOR
+
+			// 
+
+			static auto CVarHDREnable = IConsoleManager::Get().FindConsoleVariable(TEXT("Editor.HDRSupport"));
+
+			if (CVarHDREnable->GetInt() !=0 )
+			{
+				TRefCountPtr<IDXGIOutput> Output;
+				if (SUCCEEDED(SwapChain->GetContainingOutput(Output.GetInitReference())))
+				{
+					TRefCountPtr<IDXGIOutput6> Output6;
+					if (SUCCEEDED(Output->QueryInterface(IID_PPV_ARGS(Output6.GetInitReference()))))
+					{
+						DXGI_OUTPUT_DESC1 desc;
+						Output6->GetDesc1(&desc);
+
+						if (desc.ColorSpace == DXGI_COLOR_SPACE_RGB_FULL_G2084_NONE_P2020)
+						{
+							// Display output is HDR10.
+							ColorSpace = EColorSpace::ERec2020PQ;
+						}
+					}
+				}
+
+				if (ColorSpace == EColorSpace::ERec2020PQ)
+				{
+					TRefCountPtr<IDXGISwapChain3> swapChain3;
+					if (SUCCEEDED(SwapChain->QueryInterface(IID_PPV_ARGS(swapChain3.GetInitReference()))))
+					{
+						UINT colorSpaceSupport = 0;
+						DXGI_COLOR_SPACE_TYPE colorSpace = DXGI_COLOR_SPACE_RGB_FULL_G22_NONE_P709;
+
+						swapChain3->CheckColorSpaceSupport(colorSpace, &colorSpaceSupport);
+
+						if (colorSpaceSupport & DXGI_SWAP_CHAIN_COLOR_SPACE_SUPPORT_FLAG_PRESENT)
+						{
+							swapChain3->SetColorSpace1(colorSpace);
+						}
+					}
+				}
+			}
+#endif // WITH_EDITOR
+			PixelColorSpace = ColorSpace;
 		}
 
+
 		// Set the DXGI message hook to not change the window behind our back.
 		D3DRHI->GetFactory()->MakeWindowAssociation(WindowHandle,DXGI_MWA_NO_WINDOW_CHANGES);
 	}
@@ -169,6 +242,68 @@ void FD3D11Viewport::ConditionalResetSwapChain(bool bIgnoreFocus)
 			}
 		}
 	}
+
+	DXGI_COLOR_SPACE_TYPE colorSpace = DXGI_COLOR_SPACE_RGB_FULL_G22_NONE_P709;
+
+#if WITH_EDITOR
+#if defined(NTDDI_WIN10_RS2)
+
+	static auto CVarHDREnable = IConsoleManager::Get().FindConsoleVariable(TEXT("Editor.HDRSupport"));
+	if (CVarHDREnable->GetInt() != 0)
+	{
+		if (SwapChain)
+		{
+			TRefCountPtr<IDXGIOutput> Output;
+			if (SUCCEEDED(SwapChain->GetContainingOutput(Output.GetInitReference())))
+			{
+				TRefCountPtr<IDXGIOutput6> Output6;
+				if (SUCCEEDED(Output->QueryInterface(IID_PPV_ARGS(Output6.GetInitReference()))))
+				{
+					DXGI_OUTPUT_DESC1 desc;
+					Output6->GetDesc1(&desc);
+
+					if (desc.ColorSpace == DXGI_COLOR_SPACE_RGB_FULL_G2084_NONE_P2020)
+					{
+						// Display output is HDR10.
+						colorSpace = DXGI_COLOR_SPACE_RGB_FULL_G2084_NONE_P2020;
+					}
+				}
+			}
+		}
+	}
+
+	TRefCountPtr<IDXGISwapChain3> swapChain3;
+
+	if (SUCCEEDED(SwapChain->QueryInterface(IID_PPV_ARGS(swapChain3.GetInitReference()))))
+	{
+		UINT colorSpaceSupport = 0;
+		swapChain3->CheckColorSpaceSupport(colorSpace, &colorSpaceSupport);
+		DXGI_SWAP_CHAIN_DESC desc;
+		DXGI_SWAP_CHAIN_DESC1 desc1;
+
+		swapChain3->GetDesc(&desc);
+		swapChain3->GetDesc1(&desc1);
+
+
+		if (colorSpaceSupport & DXGI_SWAP_CHAIN_COLOR_SPACE_SUPPORT_FLAG_PRESENT)
+		{
+			swapChain3->SetColorSpace1(colorSpace);
+		}
+	}
+#endif
+#endif
+
+
+	if (colorSpace == DXGI_COLOR_SPACE_RGB_FULL_G2084_NONE_P2020)
+	{
+		PixelColorSpace = EColorSpace::ERec2020PQ;
+	}
+	else
+	{
+		PixelColorSpace = EColorSpace::ERec709sRGB;
+	}
+
+
 }
 
 #include "Windows/HideWindowsPlatformTypes.h"
diff --git a/Engine/Source/Runtime/Windows/D3D11RHI/Public/D3D11Viewport.h b/Engine/Source/Runtime/Windows/D3D11RHI/Public/D3D11Viewport.h
--- a/Engine/Source/Runtime/Windows/D3D11RHI/Public/D3D11Viewport.h
+++ b/Engine/Source/Runtime/Windows/D3D11RHI/Public/D3D11Viewport.h
@@ -9,6 +9,7 @@
 #include "CoreMinimal.h"
 #include "RenderResource.h"
 #include "RenderUtils.h"
+#include "DXGI1_5.h"
 
 /** A D3D event query resource. */
 class FD3D11EventQuery : public FRenderResource
@@ -75,6 +76,7 @@ public:
 	// Accessors.
 	FIntPoint GetSizeXY() const { return FIntPoint(SizeX, SizeY); }
 	FD3D11Texture2D* GetBackBuffer() const { return BackBuffer; }
+	EColorSpace GetPixelColorSpace() const { return PixelColorSpace; }
 
 	void WaitForFrameEventCompletion()
 	{
@@ -121,10 +123,12 @@ private:
 	uint32 MaximumFrameLatency;
 	uint32 SizeX;
 	uint32 SizeY;
+	uint32 BackBufferCount;
 	bool bIsFullscreen;
 	EPixelFormat PixelFormat;
+	EColorSpace PixelColorSpace;
 	bool bIsValid;
-	TRefCountPtr<IDXGISwapChain> SwapChain;
+	TRefCountPtr<IDXGISwapChain1> SwapChain;
 	TRefCountPtr<FD3D11Texture2D> BackBuffer;
 
 	// Support for selecting non-default output for display in fullscreen exclusive
