UE4.23 - Deprecate FUniformBufferRHIParamRef, FIndexBufferRHIParamRef, FVertexBufferRHIParamRef, FStructuredBufferRHIParamRef, FTextureRHIParamRef, FTexture2DRHIParamRef, FTexture2DArrayRHIParamRef, F

diff --git a/Engine/Source/Runtime/Windows/D3D11RHI/Private/D3D11IndexBuffer.cpp b/Engine/Source/Runtime/Windows/D3D11RHI/Private/D3D11IndexBuffer.cpp
--- a/Engine/Source/Runtime/Windows/D3D11RHI/Private/D3D11IndexBuffer.cpp
+++ b/Engine/Source/Runtime/Windows/D3D11RHI/Private/D3D11IndexBuffer.cpp
@@ -76,7 +76,7 @@ FIndexBufferRHIRef FD3D11DynamicRHI::CreateIndexBuffer_RenderThread(
 	return RHICreateIndexBuffer(Stride, Size, InUsage, CreateInfo);
 }
 
-void* FD3D11DynamicRHI::RHILockIndexBuffer(FIndexBufferRHIParamRef IndexBufferRHI,uint32 Offset,uint32 Size,EResourceLockMode LockMode)
+void* FD3D11DynamicRHI::RHILockIndexBuffer(FRHIIndexBuffer* IndexBufferRHI,uint32 Offset,uint32 Size,EResourceLockMode LockMode)
 {
 	FD3D11IndexBuffer* IndexBuffer = ResourceCast(IndexBufferRHI);
 	// If this resource is bound to the device, unbind it
@@ -140,7 +140,7 @@ void* FD3D11DynamicRHI::RHILockIndexBuffer(FIndexBufferRHIParamRef IndexBufferRH
 	return (void*)((uint8*)LockedData.GetData() + Offset);
 }
 
-void FD3D11DynamicRHI::RHIUnlockIndexBuffer(FIndexBufferRHIParamRef IndexBufferRHI)
+void FD3D11DynamicRHI::RHIUnlockIndexBuffer(FRHIIndexBuffer* IndexBufferRHI)
 {
 	FD3D11IndexBuffer* IndexBuffer = ResourceCast(IndexBufferRHI);
 
@@ -178,7 +178,7 @@ void FD3D11DynamicRHI::RHIUnlockIndexBuffer(FIndexBufferRHIParamRef IndexBufferR
 	}
 }
 
-void FD3D11DynamicRHI::RHITransferIndexBufferUnderlyingResource(FIndexBufferRHIParamRef DestIndexBuffer, FIndexBufferRHIParamRef SrcIndexBuffer)
+void FD3D11DynamicRHI::RHITransferIndexBufferUnderlyingResource(FRHIIndexBuffer* DestIndexBuffer, FRHIIndexBuffer* SrcIndexBuffer)
 {
 	check(DestIndexBuffer);
 	FD3D11IndexBuffer* Dest = ResourceCast(DestIndexBuffer);
diff --git a/Engine/Source/Runtime/Windows/D3D11RHI/Private/D3D11RHIPrivate.h b/Engine/Source/Runtime/Windows/D3D11RHI/Private/D3D11RHIPrivate.h
--- a/Engine/Source/Runtime/Windows/D3D11RHI/Private/D3D11RHIPrivate.h
+++ b/Engine/Source/Runtime/Windows/D3D11RHI/Private/D3D11RHIPrivate.h
@@ -388,8 +388,8 @@ public:
 	bool GetQueryData(ID3D11Query* Query, void* Data, SIZE_T DataSize, ERenderQueryType QueryType, bool bWait, bool bStallRHIThread);
 
 #if !PLATFORM_HOLOLENS	// allow the RHI default emptry implementations to fall through
-	virtual void RHIBeginUpdateMultiFrameResource(FTextureRHIParamRef Texture) override;
-	virtual void RHIEndUpdateMultiFrameResource(FTextureRHIParamRef Texture) override;
+	virtual void RHIBeginUpdateMultiFrameResource(FRHITexture* Texture) override;
+	virtual void RHIEndUpdateMultiFrameResource(FRHITexture* Texture) override;
 
 	virtual void RHIBeginUpdateMultiFrameResource(FRHIUnorderedAccessView* UAV) override;
 	virtual void RHIEndUpdateMultiFrameResource(FRHIUnorderedAccessView* UAV) override;
@@ -408,76 +408,76 @@ public:
 	virtual FGeometryShaderRHIRef RHICreateGeometryShaderWithStreamOutput(const TArray<uint8>& Code, const FStreamOutElementList& ElementList, uint32 NumStrides, const uint32* Strides, int32 RasterizedStream) final override;
 	virtual FComputeShaderRHIRef RHICreateComputeShader(const TArray<uint8>& Code) final override;
 	virtual FStagingBufferRHIRef RHICreateStagingBuffer() final override;
-	virtual void RHICopyToStagingBuffer(FVertexBufferRHIParamRef SourceBuffer, FRHIStagingBuffer* DestinationStagingBuffer, uint32 Offset, uint32 NumBytes) final override;
+	virtual void RHICopyToStagingBuffer(FRHIVertexBuffer* SourceBuffer, FRHIStagingBuffer* DestinationStagingBuffer, uint32 Offset, uint32 NumBytes) final override;
 	virtual void RHIWriteGPUFence(FRHIGPUFence* Fence) final override;
 	virtual void* RHILockStagingBuffer(FRHIStagingBuffer* StagingBuffer, uint32 Offset, uint32 SizeRHI) final override;
     virtual void RHIUnlockStagingBuffer(FRHIStagingBuffer* StagingBuffer) final override;
 	virtual FBoundShaderStateRHIRef RHICreateBoundShaderState(FRHIVertexDeclaration* VertexDeclaration, FRHIVertexShader* VertexShader, FRHIHullShader* HullShader, FRHIDomainShader* DomainShader, FRHIPixelShader* PixelShader, FRHIGeometryShader* GeometryShader) final override;
 	virtual FUniformBufferRHIRef RHICreateUniformBuffer(const void* Contents, const FRHIUniformBufferLayout& Layout, EUniformBufferUsage Usage, EUniformBufferValidation Validation) final override;
-	virtual void RHIUpdateUniformBuffer(FUniformBufferRHIParamRef UniformBufferRHI, const void* Contents) final override;
+	virtual void RHIUpdateUniformBuffer(FRHIUniformBuffer* UniformBufferRHI, const void* Contents) final override;
 	virtual FIndexBufferRHIRef RHICreateIndexBuffer(uint32 Stride, uint32 Size, uint32 InUsage, FRHIResourceCreateInfo& CreateInfo) final override;
-	virtual void* RHILockIndexBuffer(FIndexBufferRHIParamRef IndexBuffer, uint32 Offset, uint32 Size, EResourceLockMode LockMode) final override;
-	virtual void RHIUnlockIndexBuffer(FIndexBufferRHIParamRef IndexBuffer) final override;
-	virtual void RHITransferIndexBufferUnderlyingResource(FIndexBufferRHIParamRef DestIndexBuffer, FIndexBufferRHIParamRef SrcIndexBuffer) final override;
+	virtual void* RHILockIndexBuffer(FRHIIndexBuffer* IndexBuffer, uint32 Offset, uint32 Size, EResourceLockMode LockMode) final override;
+	virtual void RHIUnlockIndexBuffer(FRHIIndexBuffer* IndexBuffer) final override;
+	virtual void RHITransferIndexBufferUnderlyingResource(FRHIIndexBuffer* DestIndexBuffer, FRHIIndexBuffer* SrcIndexBuffer) final override;
 	virtual FVertexBufferRHIRef RHICreateVertexBuffer(uint32 Size, uint32 InUsage, FRHIResourceCreateInfo& CreateInfo) final override;
-	virtual void* RHILockVertexBuffer(FVertexBufferRHIParamRef VertexBuffer, uint32 Offset, uint32 SizeRHI, EResourceLockMode LockMode) final override;
-	virtual void RHIUnlockVertexBuffer(FVertexBufferRHIParamRef VertexBuffer) final override;
-	virtual void RHICopyVertexBuffer(FVertexBufferRHIParamRef SourceBuffer,FVertexBufferRHIParamRef DestBuffer) final override;
-	virtual void RHITransferVertexBufferUnderlyingResource(FVertexBufferRHIParamRef DestVertexBuffer, FVertexBufferRHIParamRef SrcVertexBuffer) final override;
+	virtual void* RHILockVertexBuffer(FRHIVertexBuffer* VertexBuffer, uint32 Offset, uint32 SizeRHI, EResourceLockMode LockMode) final override;
+	virtual void RHIUnlockVertexBuffer(FRHIVertexBuffer* VertexBuffer) final override;
+	virtual void RHICopyVertexBuffer(FRHIVertexBuffer* SourceBuffer, FRHIVertexBuffer* DestBuffer) final override;
+	virtual void RHITransferVertexBufferUnderlyingResource(FRHIVertexBuffer* DestVertexBuffer, FRHIVertexBuffer* SrcVertexBuffer) final override;
 	virtual FStructuredBufferRHIRef RHICreateStructuredBuffer(uint32 Stride, uint32 Size, uint32 InUsage, FRHIResourceCreateInfo& CreateInfo) final override;
-	virtual void* RHILockStructuredBuffer(FStructuredBufferRHIParamRef StructuredBuffer, uint32 Offset, uint32 SizeRHI, EResourceLockMode LockMode) final override;
-	virtual void RHIUnlockStructuredBuffer(FStructuredBufferRHIParamRef StructuredBuffer) final override;
-	virtual FUnorderedAccessViewRHIRef RHICreateUnorderedAccessView(FStructuredBufferRHIParamRef StructuredBuffer, bool bUseUAVCounter, bool bAppendBuffer) final override;
-	virtual FUnorderedAccessViewRHIRef RHICreateUnorderedAccessView(FTextureRHIParamRef Texture, uint32 MipLevel) final override;
-	virtual FUnorderedAccessViewRHIRef RHICreateUnorderedAccessView(FVertexBufferRHIParamRef VertexBuffer, uint8 Format) final override;
-	virtual FUnorderedAccessViewRHIRef RHICreateUnorderedAccessView(FIndexBufferRHIParamRef IndexBuffer, uint8 Format) final override;
-	virtual FShaderResourceViewRHIRef RHICreateShaderResourceView(FStructuredBufferRHIParamRef StructuredBuffer) final override;
-	virtual FShaderResourceViewRHIRef RHICreateShaderResourceView(FVertexBufferRHIParamRef VertexBuffer, uint32 Stride, uint8 Format) final override;
-	virtual FShaderResourceViewRHIRef RHICreateShaderResourceView(FIndexBufferRHIParamRef Buffer) final override;
-	virtual void RHIUpdateShaderResourceView(FRHIShaderResourceView* SRV, FVertexBufferRHIParamRef VertexBuffer, uint32 Stride, uint8 Format) final override;
-	virtual void RHIUpdateShaderResourceView(FRHIShaderResourceView* SRV, FIndexBufferRHIParamRef IndexBuffer) final override;
+	virtual void* RHILockStructuredBuffer(FRHIStructuredBuffer* StructuredBuffer, uint32 Offset, uint32 SizeRHI, EResourceLockMode LockMode) final override;
+	virtual void RHIUnlockStructuredBuffer(FRHIStructuredBuffer* StructuredBuffer) final override;
+	virtual FUnorderedAccessViewRHIRef RHICreateUnorderedAccessView(FRHIStructuredBuffer* StructuredBuffer, bool bUseUAVCounter, bool bAppendBuffer) final override;
+	virtual FUnorderedAccessViewRHIRef RHICreateUnorderedAccessView(FRHITexture* Texture, uint32 MipLevel) final override;
+	virtual FUnorderedAccessViewRHIRef RHICreateUnorderedAccessView(FRHIVertexBuffer* VertexBuffer, uint8 Format) final override;
+	virtual FUnorderedAccessViewRHIRef RHICreateUnorderedAccessView(FRHIIndexBuffer* IndexBuffer, uint8 Format) final override;
+	virtual FShaderResourceViewRHIRef RHICreateShaderResourceView(FRHIStructuredBuffer* StructuredBuffer) final override;
+	virtual FShaderResourceViewRHIRef RHICreateShaderResourceView(FRHIVertexBuffer* VertexBuffer, uint32 Stride, uint8 Format) final override;
+	virtual FShaderResourceViewRHIRef RHICreateShaderResourceView(FRHIIndexBuffer* Buffer) final override;
+	virtual void RHIUpdateShaderResourceView(FRHIShaderResourceView* SRV, FRHIVertexBuffer* VertexBuffer, uint32 Stride, uint8 Format) final override;
+	virtual void RHIUpdateShaderResourceView(FRHIShaderResourceView* SRV, FRHIIndexBuffer* IndexBuffer) final override;
 	virtual uint64 RHICalcTexture2DPlatformSize(uint32 SizeX, uint32 SizeY, uint8 Format, uint32 NumMips, uint32 NumSamples, uint32 Flags, uint32& OutAlign) final override;
 	virtual uint64 RHICalcTexture3DPlatformSize(uint32 SizeX, uint32 SizeY, uint32 SizeZ, uint8 Format, uint32 NumMips, uint32 Flags, uint32& OutAlign) final override;
 	virtual uint64 RHICalcTextureCubePlatformSize(uint32 Size, uint8 Format, uint32 NumMips, uint32 Flags, uint32& OutAlign) final override;
 	virtual void RHIGetTextureMemoryStats(FTextureMemoryStats& OutStats) final override;
 	virtual bool RHIGetTextureMemoryVisualizeData(FColor* TextureData,int32 SizeX,int32 SizeY,int32 Pitch,int32 PixelSize) final override;
 	virtual FTextureReferenceRHIRef RHICreateTextureReference(FLastRenderTimeContainer* LastRenderTime) final override;
-	virtual void RHIUpdateTextureReference(FRHITextureReference* TextureRef, FTextureRHIParamRef NewTexture) final override;
+	virtual void RHIUpdateTextureReference(FRHITextureReference* TextureRef, FRHITexture* NewTexture) final override;
 	virtual FTexture2DRHIRef RHICreateTexture2D(uint32 SizeX, uint32 SizeY, uint8 Format, uint32 NumMips, uint32 NumSamples, uint32 Flags, FRHIResourceCreateInfo& CreateInfo) final override;
 	virtual FTexture2DRHIRef RHIAsyncCreateTexture2D(uint32 SizeX, uint32 SizeY, uint8 Format, uint32 NumMips, uint32 Flags, void** InitialMipData, uint32 NumInitialMips) final override;
-	virtual void RHICopySharedMips(FTexture2DRHIParamRef DestTexture2D, FTexture2DRHIParamRef SrcTexture2D) final override;
+	virtual void RHICopySharedMips(FRHITexture2D* DestTexture2D, FRHITexture2D* SrcTexture2D) final override;
 	virtual FTexture2DArrayRHIRef RHICreateTexture2DArray(uint32 SizeX, uint32 SizeY, uint32 SizeZ, uint8 Format, uint32 NumMips, uint32 NumSamples, uint32 Flags, FRHIResourceCreateInfo& CreateInfo) final override;
 	virtual FTexture3DRHIRef RHICreateTexture3D(uint32 SizeX, uint32 SizeY, uint32 SizeZ, uint8 Format, uint32 NumMips, uint32 Flags, FRHIResourceCreateInfo& CreateInfo) final override;
-	virtual void RHIGetResourceInfo(FTextureRHIParamRef Ref, FRHIResourceInfo& OutInfo) final override;
-	virtual FShaderResourceViewRHIRef RHICreateShaderResourceView(FTextureRHIParamRef TextureRHI, const FRHITextureSRVCreateInfo& CreateInfo) final override;
-	virtual void RHIGenerateMips(FTextureRHIParamRef Texture) final override;
-	virtual uint32 RHIComputeMemorySize(FTextureRHIParamRef TextureRHI) final override;
-	virtual void RHIAsyncCopyTexture2DCopy(FTexture2DRHIParamRef NewTexture2DRHI, FTexture2DRHIParamRef Texture2DRHI, int32 NewMipCount, int32 NewSizeX, int32 NewSizeY, FThreadSafeCounter* RequestStatus);
-	virtual FTexture2DRHIRef RHIAsyncReallocateTexture2D(FTexture2DRHIParamRef Texture2D, int32 NewMipCount, int32 NewSizeX, int32 NewSizeY, FThreadSafeCounter* RequestStatus) final override;
-	virtual ETextureReallocationStatus RHIFinalizeAsyncReallocateTexture2D(FTexture2DRHIParamRef Texture2D, bool bBlockUntilCompleted) final override;
-	virtual ETextureReallocationStatus RHICancelAsyncReallocateTexture2D(FTexture2DRHIParamRef Texture2D, bool bBlockUntilCompleted) final override;
-	virtual void* RHILockTexture2D(FTexture2DRHIParamRef Texture, uint32 MipIndex, EResourceLockMode LockMode, uint32& DestStride, bool bLockWithinMiptail) final override;
-	virtual void RHIUnlockTexture2D(FTexture2DRHIParamRef Texture, uint32 MipIndex, bool bLockWithinMiptail) final override;
-	virtual void* RHILockTexture2DArray(FTexture2DArrayRHIParamRef Texture, uint32 TextureIndex, uint32 MipIndex, EResourceLockMode LockMode, uint32& DestStride, bool bLockWithinMiptail) final override;
-	virtual void RHIUnlockTexture2DArray(FTexture2DArrayRHIParamRef Texture, uint32 TextureIndex, uint32 MipIndex, bool bLockWithinMiptail) final override;
-	virtual void RHIUpdateTexture2D(FTexture2DRHIParamRef Texture, uint32 MipIndex, const struct FUpdateTextureRegion2D& UpdateRegion, uint32 SourcePitch, const uint8* SourceData) final override;
-	virtual void RHIUpdateTexture3D(FTexture3DRHIParamRef Texture, uint32 MipIndex, const struct FUpdateTextureRegion3D& UpdateRegion, uint32 SourceRowPitch, uint32 SourceDepthPitch, const uint8* SourceData) final override;
+	virtual void RHIGetResourceInfo(FRHITexture* Ref, FRHIResourceInfo& OutInfo) final override;
+	virtual FShaderResourceViewRHIRef RHICreateShaderResourceView(FRHITexture* TextureRHI, const FRHITextureSRVCreateInfo& CreateInfo) final override;
+	virtual void RHIGenerateMips(FRHITexture* Texture) final override;
+	virtual uint32 RHIComputeMemorySize(FRHITexture* TextureRHI) final override;
+	virtual void RHIAsyncCopyTexture2DCopy(FRHITexture2D* NewTexture2DRHI, FRHITexture2D* Texture2DRHI, int32 NewMipCount, int32 NewSizeX, int32 NewSizeY, FThreadSafeCounter* RequestStatus);
+	virtual FTexture2DRHIRef RHIAsyncReallocateTexture2D(FRHITexture2D* Texture2D, int32 NewMipCount, int32 NewSizeX, int32 NewSizeY, FThreadSafeCounter* RequestStatus) final override;
+	virtual ETextureReallocationStatus RHIFinalizeAsyncReallocateTexture2D(FRHITexture2D* Texture2D, bool bBlockUntilCompleted) final override;
+	virtual ETextureReallocationStatus RHICancelAsyncReallocateTexture2D(FRHITexture2D* Texture2D, bool bBlockUntilCompleted) final override;
+	virtual void* RHILockTexture2D(FRHITexture2D* Texture, uint32 MipIndex, EResourceLockMode LockMode, uint32& DestStride, bool bLockWithinMiptail) final override;
+	virtual void RHIUnlockTexture2D(FRHITexture2D* Texture, uint32 MipIndex, bool bLockWithinMiptail) final override;
+	virtual void* RHILockTexture2DArray(FRHITexture2DArray* Texture, uint32 TextureIndex, uint32 MipIndex, EResourceLockMode LockMode, uint32& DestStride, bool bLockWithinMiptail) final override;
+	virtual void RHIUnlockTexture2DArray(FRHITexture2DArray* Texture, uint32 TextureIndex, uint32 MipIndex, bool bLockWithinMiptail) final override;
+	virtual void RHIUpdateTexture2D(FRHITexture2D* Texture, uint32 MipIndex, const struct FUpdateTextureRegion2D& UpdateRegion, uint32 SourcePitch, const uint8* SourceData) final override;
+	virtual void RHIUpdateTexture3D(FRHITexture3D* Texture, uint32 MipIndex, const struct FUpdateTextureRegion3D& UpdateRegion, uint32 SourceRowPitch, uint32 SourceDepthPitch, const uint8* SourceData) final override;
 	virtual FTextureCubeRHIRef RHICreateTextureCube(uint32 Size, uint8 Format, uint32 NumMips, uint32 Flags, FRHIResourceCreateInfo& CreateInfo) final override;
 	virtual FTextureCubeRHIRef RHICreateTextureCubeArray(uint32 Size, uint32 ArraySize, uint8 Format, uint32 NumMips, uint32 Flags, FRHIResourceCreateInfo& CreateInfo) final override;
-	virtual void* RHILockTextureCubeFace(FTextureCubeRHIParamRef Texture, uint32 FaceIndex, uint32 ArrayIndex, uint32 MipIndex, EResourceLockMode LockMode, uint32& DestStride, bool bLockWithinMiptail) final override;
-	virtual void RHIUnlockTextureCubeFace(FTextureCubeRHIParamRef Texture, uint32 FaceIndex, uint32 ArrayIndex, uint32 MipIndex, bool bLockWithinMiptail) final override;
-	virtual void RHIBindDebugLabelName(FTextureRHIParamRef Texture, const TCHAR* Name) final override;
+	virtual void* RHILockTextureCubeFace(FRHITextureCube* Texture, uint32 FaceIndex, uint32 ArrayIndex, uint32 MipIndex, EResourceLockMode LockMode, uint32& DestStride, bool bLockWithinMiptail) final override;
+	virtual void RHIUnlockTextureCubeFace(FRHITextureCube* Texture, uint32 FaceIndex, uint32 ArrayIndex, uint32 MipIndex, bool bLockWithinMiptail) final override;
+	virtual void RHIBindDebugLabelName(FRHITexture* Texture, const TCHAR* Name) final override;
 	virtual void RHIBindDebugLabelName(FRHIUnorderedAccessView* UnorderedAccessViewRHI, const TCHAR* Name) final override;
-	virtual void RHIReadSurfaceData(FTextureRHIParamRef Texture,FIntRect Rect,TArray<FColor>& OutData,FReadSurfaceDataFlags InFlags) final override;
-	virtual void RHIReadSurfaceData(FTextureRHIParamRef TextureRHI, FIntRect InRect, TArray<FLinearColor>& OutData, FReadSurfaceDataFlags InFlags) final override;
-	virtual void RHIMapStagingSurface(FTextureRHIParamRef Texture,void*& OutData,int32& OutWidth,int32& OutHeight) final override;
-	virtual void RHIUnmapStagingSurface(FTextureRHIParamRef Texture) final override;
-	virtual void RHIReadSurfaceFloatData(FTextureRHIParamRef Texture,FIntRect Rect,TArray<FFloat16Color>& OutData,ECubeFace CubeFace,int32 ArrayIndex,int32 MipIndex) final override;
-	virtual void RHIRead3DSurfaceFloatData(FTextureRHIParamRef Texture,FIntRect Rect,FIntPoint ZMinMax,TArray<FFloat16Color>& OutData) final override;
+	virtual void RHIReadSurfaceData(FRHITexture* Texture,FIntRect Rect,TArray<FColor>& OutData,FReadSurfaceDataFlags InFlags) final override;
+	virtual void RHIReadSurfaceData(FRHITexture* TextureRHI, FIntRect InRect, TArray<FLinearColor>& OutData, FReadSurfaceDataFlags InFlags) final override;
+	virtual void RHIMapStagingSurface(FRHITexture* Texture,void*& OutData,int32& OutWidth,int32& OutHeight) final override;
+	virtual void RHIUnmapStagingSurface(FRHITexture* Texture) final override;
+	virtual void RHIReadSurfaceFloatData(FRHITexture* Texture,FIntRect Rect,TArray<FFloat16Color>& OutData,ECubeFace CubeFace,int32 ArrayIndex,int32 MipIndex) final override;
+	virtual void RHIRead3DSurfaceFloatData(FRHITexture* Texture,FIntRect Rect,FIntPoint ZMinMax,TArray<FFloat16Color>& OutData) final override;
 	virtual FRenderQueryRHIRef RHICreateRenderQuery(ERenderQueryType QueryType) final override;
 	virtual bool RHIGetRenderQueryResult(FRHIRenderQuery* RenderQuery, uint64& OutResult, bool bWait) final override;
 	virtual FTexture2DRHIRef RHIGetViewportBackBuffer(FRHIViewport* Viewport) final override;
-	virtual void RHIAliasTextureResources(FTextureRHIParamRef DestTexture, FTextureRHIParamRef SrcTexture) final override;
+	virtual void RHIAliasTextureResources(FRHITexture* DestTexture, FRHITexture* SrcTexture) final override;
 	virtual void RHIAdvanceFrameForGetViewportBackBuffer(FRHIViewport* Viewport) final override;
 	virtual void RHIAcquireThreadOwnership() final override;
 	virtual void RHIReleaseThreadOwnership() final override;
@@ -487,12 +487,12 @@ public:
 	virtual void RHIResizeViewport(FRHIViewport* Viewport, uint32 SizeX, uint32 SizeY, bool bIsFullscreen) final override;
 	virtual void RHIResizeViewport(FRHIViewport* Viewport, uint32 SizeX, uint32 SizeY, bool bIsFullscreen, EPixelFormat PreferredPixelFormat) final override;
 	virtual void RHITick(float DeltaTime) final override;
-	virtual void RHISetStreamOutTargets(uint32 NumTargets,const FVertexBufferRHIParamRef* VertexBuffers,const uint32* Offsets) final override;
+	virtual void RHISetStreamOutTargets(uint32 NumTargets, FRHIVertexBuffer* const* VertexBuffers,const uint32* Offsets) final override;
 	virtual void RHIBlockUntilGPUIdle() final override;
 	virtual bool RHIGetAvailableResolutions(FScreenResolutionArray& Resolutions, bool bIgnoreRefreshRate) final override;
 	virtual void RHIGetSupportedResolution(uint32& Width, uint32& Height) final override;
-	virtual void RHIVirtualTextureSetFirstMipInMemory(FTexture2DRHIParamRef Texture, uint32 FirstMip) final override;
-	virtual void RHIVirtualTextureSetFirstMipVisible(FTexture2DRHIParamRef Texture, uint32 FirstMip) final override;
+	virtual void RHIVirtualTextureSetFirstMipInMemory(FRHITexture2D* Texture, uint32 FirstMip) final override;
+	virtual void RHIVirtualTextureSetFirstMipVisible(FRHITexture2D* Texture, uint32 FirstMip) final override;
 	virtual void RHIExecuteCommandList(FRHICommandList* CmdList) final override;
 	virtual void* RHIGetNativeDevice() final override;
 	virtual class IRHICommandContext* RHIGetDefaultContext() final override;
@@ -500,27 +500,27 @@ public:
 
 	virtual void RHISetComputeShader(FRHIComputeShader* ComputeShader) final override;
 	virtual void RHIDispatchComputeShader(uint32 ThreadGroupCountX, uint32 ThreadGroupCountY, uint32 ThreadGroupCountZ) final override;
-	virtual void RHIDispatchIndirectComputeShader(FVertexBufferRHIParamRef ArgumentBuffer, uint32 ArgumentOffset) final override;
+	virtual void RHIDispatchIndirectComputeShader(FRHIVertexBuffer* ArgumentBuffer, uint32 ArgumentOffset) final override;
 	virtual void RHIAutomaticCacheFlushAfterComputeShader(bool bEnable) final override;
 	virtual void RHIFlushComputeShaderCache() final override;
 	virtual void RHISetMultipleViewports(uint32 Count, const FViewportBounds* Data) final override;
 	virtual void RHIClearTinyUAV(FRHIUnorderedAccessView* UnorderedAccessViewRHI, const uint32* Values) final override;
-	virtual void RHICopyToResolveTarget(FTextureRHIParamRef SourceTexture, FTextureRHIParamRef DestTexture, const FResolveParams& ResolveParams) final override;
-	virtual void RHICopyTexture(FTextureRHIParamRef SourceTexture, FTextureRHIParamRef DestTexture, const FRHICopyTextureInfo& CopyInfo) final override;
-	virtual void RHITransitionResources(EResourceTransitionAccess TransitionType, FTextureRHIParamRef* InTextures, int32 NumTextures) final override;
+	virtual void RHICopyToResolveTarget(FRHITexture* SourceTexture, FRHITexture* DestTexture, const FResolveParams& ResolveParams) final override;
+	virtual void RHICopyTexture(FRHITexture* SourceTexture, FRHITexture* DestTexture, const FRHICopyTextureInfo& CopyInfo) final override;
+	virtual void RHITransitionResources(EResourceTransitionAccess TransitionType, FRHITexture** InTextures, int32 NumTextures) final override;
 	virtual void RHITransitionResources(EResourceTransitionAccess TransitionType, EResourceTransitionPipeline TransitionPipeline, FRHIUnorderedAccessView** InUAVs, int32 NumUAVs, FRHIComputeFence* WriteFence) final override;
 	virtual void RHIBeginRenderQuery(FRHIRenderQuery* RenderQuery) final override;
 	virtual void RHIEndRenderQuery(FRHIRenderQuery* RenderQuery) final override;
 	void RHIBeginOcclusionQueryBatch(uint32 NumQueriesInBatch);
 	void RHIEndOcclusionQueryBatch();
 	virtual void RHISubmitCommandsHint() final override;
-	virtual void RHIBeginDrawingViewport(FRHIViewport* Viewport, FTextureRHIParamRef RenderTargetRHI) final override;
+	virtual void RHIBeginDrawingViewport(FRHIViewport* Viewport, FRHITexture* RenderTargetRHI) final override;
 	virtual void RHIEndDrawingViewport(FRHIViewport* Viewport, bool bPresent, bool bLockToVsync) final override;
 	virtual void RHIBeginFrame() override;
 	virtual void RHIEndFrame() override;
 	virtual void RHIBeginScene() final override;
 	virtual void RHIEndScene() final override;
-	virtual void RHISetStreamSource(uint32 StreamIndex, FVertexBufferRHIParamRef VertexBuffer, uint32 Offset) final override;
+	virtual void RHISetStreamSource(uint32 StreamIndex, FRHIVertexBuffer* VertexBuffer, uint32 Offset) final override;
 	virtual void RHISetRasterizerState(FRHIRasterizerState* NewState) final override;
 	virtual void RHISetViewport(uint32 MinX, uint32 MinY, float MinZ, uint32 MaxX, uint32 MaxY, float MaxZ) final override;
 	virtual void RHISetScissorRect(bool bEnable, uint32 MinX, uint32 MinY, uint32 MaxX, uint32 MaxY) final override;
@@ -533,12 +533,12 @@ public:
 		PrimitiveType = FallbackGraphicsState->Initializer.PrimitiveType;
 	}
 
-	virtual void RHISetShaderTexture(FRHIVertexShader* VertexShader, uint32 TextureIndex, FTextureRHIParamRef NewTexture) final override;
-	virtual void RHISetShaderTexture(FRHIHullShader* HullShader, uint32 TextureIndex, FTextureRHIParamRef NewTexture) final override;
-	virtual void RHISetShaderTexture(FRHIDomainShader* DomainShader, uint32 TextureIndex, FTextureRHIParamRef NewTexture) final override;
-	virtual void RHISetShaderTexture(FRHIGeometryShader* GeometryShader, uint32 TextureIndex, FTextureRHIParamRef NewTexture) final override;
-	virtual void RHISetShaderTexture(FRHIPixelShader* PixelShader, uint32 TextureIndex, FTextureRHIParamRef NewTexture) final override;
-	virtual void RHISetShaderTexture(FRHIComputeShader* PixelShader, uint32 TextureIndex, FTextureRHIParamRef NewTexture) final override;
+	virtual void RHISetShaderTexture(FRHIVertexShader* VertexShader, uint32 TextureIndex, FRHITexture* NewTexture) final override;
+	virtual void RHISetShaderTexture(FRHIHullShader* HullShader, uint32 TextureIndex, FRHITexture* NewTexture) final override;
+	virtual void RHISetShaderTexture(FRHIDomainShader* DomainShader, uint32 TextureIndex, FRHITexture* NewTexture) final override;
+	virtual void RHISetShaderTexture(FRHIGeometryShader* GeometryShader, uint32 TextureIndex, FRHITexture* NewTexture) final override;
+	virtual void RHISetShaderTexture(FRHIPixelShader* PixelShader, uint32 TextureIndex, FRHITexture* NewTexture) final override;
+	virtual void RHISetShaderTexture(FRHIComputeShader* PixelShader, uint32 TextureIndex, FRHITexture* NewTexture) final override;
 	virtual void RHISetShaderSampler(FRHIComputeShader* ComputeShader, uint32 SamplerIndex, FRHISamplerState* NewState) final override;
 	virtual void RHISetShaderSampler(FRHIVertexShader* VertexShader, uint32 SamplerIndex, FRHISamplerState* NewState) final override;
 	virtual void RHISetShaderSampler(FRHIGeometryShader* GeometryShader, uint32 SamplerIndex, FRHISamplerState* NewState) final override;
@@ -553,12 +553,12 @@ public:
 	virtual void RHISetShaderResourceViewParameter(FRHIHullShader* HullShader, uint32 SamplerIndex, FRHIShaderResourceView* SRV) final override;
 	virtual void RHISetShaderResourceViewParameter(FRHIDomainShader* DomainShader, uint32 SamplerIndex, FRHIShaderResourceView* SRV) final override;
 	virtual void RHISetShaderResourceViewParameter(FRHIGeometryShader* GeometryShader, uint32 SamplerIndex, FRHIShaderResourceView* SRV) final override;
-	virtual void RHISetShaderUniformBuffer(FRHIVertexShader* VertexShader, uint32 BufferIndex, FUniformBufferRHIParamRef Buffer) final override;
-	virtual void RHISetShaderUniformBuffer(FRHIHullShader* HullShader, uint32 BufferIndex, FUniformBufferRHIParamRef Buffer) final override;
-	virtual void RHISetShaderUniformBuffer(FRHIDomainShader* DomainShader, uint32 BufferIndex, FUniformBufferRHIParamRef Buffer) final override;
-	virtual void RHISetShaderUniformBuffer(FRHIGeometryShader* GeometryShader, uint32 BufferIndex, FUniformBufferRHIParamRef Buffer) final override;
-	virtual void RHISetShaderUniformBuffer(FRHIPixelShader* PixelShader, uint32 BufferIndex, FUniformBufferRHIParamRef Buffer) final override;
-	virtual void RHISetShaderUniformBuffer(FRHIComputeShader* ComputeShader, uint32 BufferIndex, FUniformBufferRHIParamRef Buffer) final override;
+	virtual void RHISetShaderUniformBuffer(FRHIVertexShader* VertexShader, uint32 BufferIndex, FRHIUniformBuffer* Buffer) final override;
+	virtual void RHISetShaderUniformBuffer(FRHIHullShader* HullShader, uint32 BufferIndex, FRHIUniformBuffer* Buffer) final override;
+	virtual void RHISetShaderUniformBuffer(FRHIDomainShader* DomainShader, uint32 BufferIndex, FRHIUniformBuffer* Buffer) final override;
+	virtual void RHISetShaderUniformBuffer(FRHIGeometryShader* GeometryShader, uint32 BufferIndex, FRHIUniformBuffer* Buffer) final override;
+	virtual void RHISetShaderUniformBuffer(FRHIPixelShader* PixelShader, uint32 BufferIndex, FRHIUniformBuffer* Buffer) final override;
+	virtual void RHISetShaderUniformBuffer(FRHIComputeShader* ComputeShader, uint32 BufferIndex, FRHIUniformBuffer* Buffer) final override;
 	virtual void RHISetShaderParameter(FRHIVertexShader* VertexShader, uint32 BufferIndex, uint32 BaseIndex, uint32 NumBytes, const void* NewValue) final override;
 	virtual void RHISetShaderParameter(FRHIPixelShader* PixelShader, uint32 BufferIndex, uint32 BaseIndex, uint32 NumBytes, const void* NewValue) final override;
 	virtual void RHISetShaderParameter(FRHIHullShader* HullShader, uint32 BufferIndex, uint32 BaseIndex, uint32 NumBytes, const void* NewValue) final override;
@@ -573,10 +573,10 @@ public:
 	virtual void RHISetRenderTargetsAndClear(const FRHISetRenderTargetsInfo& RenderTargetsInfo) final override;
 	virtual void RHIBindClearMRTValues(bool bClearColor, bool bClearDepth, bool bClearStencil) final override;
 	virtual void RHIDrawPrimitive(uint32 BaseVertexIndex, uint32 NumPrimitives, uint32 NumInstances) final override;
-	virtual void RHIDrawPrimitiveIndirect(FVertexBufferRHIParamRef ArgumentBuffer, uint32 ArgumentOffset) final override;
-	virtual void RHIDrawIndexedIndirect(FIndexBufferRHIParamRef IndexBufferRHI, FStructuredBufferRHIParamRef ArgumentsBufferRHI, int32 DrawArgumentsIndex, uint32 NumInstances) final override;
-	virtual void RHIDrawIndexedPrimitive(FIndexBufferRHIParamRef IndexBuffer, int32 BaseVertexIndex, uint32 FirstInstance, uint32 NumVertices, uint32 StartIndex, uint32 NumPrimitives, uint32 NumInstances) final override;
-	virtual void RHIDrawIndexedPrimitiveIndirect(FIndexBufferRHIParamRef IndexBuffer, FVertexBufferRHIParamRef ArgumentBuffer, uint32 ArgumentOffset) final override;
+	virtual void RHIDrawPrimitiveIndirect(FRHIVertexBuffer* ArgumentBuffer, uint32 ArgumentOffset) final override;
+	virtual void RHIDrawIndexedIndirect(FRHIIndexBuffer* IndexBufferRHI, FRHIStructuredBuffer* ArgumentsBufferRHI, int32 DrawArgumentsIndex, uint32 NumInstances) final override;
+	virtual void RHIDrawIndexedPrimitive(FRHIIndexBuffer* IndexBuffer, int32 BaseVertexIndex, uint32 FirstInstance, uint32 NumVertices, uint32 StartIndex, uint32 NumPrimitives, uint32 NumInstances) final override;
+	virtual void RHIDrawIndexedPrimitiveIndirect(FRHIIndexBuffer* IndexBuffer, FRHIVertexBuffer* ArgumentBuffer, uint32 ArgumentOffset) final override;
 	virtual void RHIEnableDepthBoundsTest(bool bEnable) final override
 	{
 		if (GSupportsDepthBoundsTest && StateCache.bDepthBoundsEnabled != bEnable)
@@ -594,7 +594,7 @@ public:
 	virtual void RHIPushEvent(const TCHAR* Name, FColor Color) final override;
 	virtual void RHIPopEvent() final override;
 
-	virtual void RHICopySubTextureRegion(FTexture2DRHIParamRef SourceTexture, FTexture2DRHIParamRef DestinationTexture, FBox2D SourceBox, FBox2D DestinationBox) final override;
+	virtual void RHICopySubTextureRegion(FRHITexture2D* SourceTexture, FRHITexture2D* DestinationTexture, FBox2D SourceBox, FBox2D DestinationBox) final override;
 
 	virtual FTexture2DRHIRef RHICreateTexture2DFromResource(EPixelFormat Format, uint32 TexCreateFlags, const FClearValueBinding& ClearValueBinding, ID3D11Texture2D* Resource);
 	virtual FTextureCubeRHIRef RHICreateTextureCubeFromResource(EPixelFormat Format, uint32 TexCreateFlags, const FClearValueBinding& ClearValueBinding, ID3D11Texture2D* Resource);
@@ -607,11 +607,11 @@ public:
 	// For D3D11, these functions mainly just remove RHIT stalls because ID3D11Device is thread safe.
 	virtual FVertexBufferRHIRef CreateVertexBuffer_RenderThread(class FRHICommandListImmediate& RHICmdList, uint32 Size, uint32 InUsage, FRHIResourceCreateInfo& CreateInfo) final override;
 	virtual FStructuredBufferRHIRef CreateStructuredBuffer_RenderThread(class FRHICommandListImmediate& RHICmdList, uint32 Stride, uint32 Size, uint32 InUsage, FRHIResourceCreateInfo& CreateInfo) final override;
-	virtual FShaderResourceViewRHIRef CreateShaderResourceView_RenderThread(class FRHICommandListImmediate& RHICmdList, FVertexBufferRHIParamRef VertexBuffer, uint32 Stride, uint8 Format) final override;
-	virtual FShaderResourceViewRHIRef CreateShaderResourceView_RenderThread(class FRHICommandListImmediate& RHICmdList, FIndexBufferRHIParamRef Buffer) final override;
-	virtual FTexture2DRHIRef AsyncReallocateTexture2D_RenderThread(class FRHICommandListImmediate& RHICmdList, FTexture2DRHIParamRef Texture2D, int32 NewMipCount, int32 NewSizeX, int32 NewSizeY, FThreadSafeCounter* RequestStatus) final override;
-	virtual ETextureReallocationStatus FinalizeAsyncReallocateTexture2D_RenderThread(class FRHICommandListImmediate& RHICmdList, FTexture2DRHIParamRef Texture2D, bool bBlockUntilCompleted) final override;
-	virtual ETextureReallocationStatus CancelAsyncReallocateTexture2D_RenderThread(class FRHICommandListImmediate& RHICmdList, FTexture2DRHIParamRef Texture2D, bool bBlockUntilCompleted) final override;
+	virtual FShaderResourceViewRHIRef CreateShaderResourceView_RenderThread(class FRHICommandListImmediate& RHICmdList, FRHIVertexBuffer* VertexBuffer, uint32 Stride, uint8 Format) final override;
+	virtual FShaderResourceViewRHIRef CreateShaderResourceView_RenderThread(class FRHICommandListImmediate& RHICmdList, FRHIIndexBuffer* Buffer) final override;
+	virtual FTexture2DRHIRef AsyncReallocateTexture2D_RenderThread(class FRHICommandListImmediate& RHICmdList, FRHITexture2D* Texture2D, int32 NewMipCount, int32 NewSizeX, int32 NewSizeY, FThreadSafeCounter* RequestStatus) final override;
+	virtual ETextureReallocationStatus FinalizeAsyncReallocateTexture2D_RenderThread(class FRHICommandListImmediate& RHICmdList, FRHITexture2D* Texture2D, bool bBlockUntilCompleted) final override;
+	virtual ETextureReallocationStatus CancelAsyncReallocateTexture2D_RenderThread(class FRHICommandListImmediate& RHICmdList, FRHITexture2D* Texture2D, bool bBlockUntilCompleted) final override;
 	virtual FIndexBufferRHIRef CreateIndexBuffer_RenderThread(class FRHICommandListImmediate& RHICmdList, uint32 Stride, uint32 Size, uint32 InUsage, FRHIResourceCreateInfo& CreateInfo) final override;
 	virtual FVertexShaderRHIRef CreateVertexShader_RenderThread(class FRHICommandListImmediate& RHICmdList, const TArray<uint8>& Code) final override;
 	virtual FPixelShaderRHIRef CreatePixelShader_RenderThread(class FRHICommandListImmediate& RHICmdList, const TArray<uint8>& Code) final override;
@@ -620,27 +620,27 @@ public:
 	virtual FComputeShaderRHIRef CreateComputeShader_RenderThread(class FRHICommandListImmediate& RHICmdList, const TArray<uint8>& Code) final override;
 	virtual FHullShaderRHIRef CreateHullShader_RenderThread(class FRHICommandListImmediate& RHICmdList, const TArray<uint8>& Code) final override;
 	virtual FDomainShaderRHIRef CreateDomainShader_RenderThread(class FRHICommandListImmediate& RHICmdList, const TArray<uint8>& Code);
-	virtual void* LockTexture2D_RenderThread(class FRHICommandListImmediate& RHICmdList, FTexture2DRHIParamRef Texture, uint32 MipIndex, EResourceLockMode LockMode, uint32& DestStride, bool bLockWithinMiptail, bool bNeedsDefaultRHIFlush = true) final override;
-	virtual void UnlockTexture2D_RenderThread(class FRHICommandListImmediate& RHICmdList, FTexture2DRHIParamRef Texture, uint32 MipIndex, bool bLockWithinMiptail, bool bNeedsDefaultRHIFlush = true) final override;
-	virtual void UpdateTexture2D_RenderThread(class FRHICommandListImmediate& RHICmdList, FTexture2DRHIParamRef Texture, uint32 MipIndex, const struct FUpdateTextureRegion2D& UpdateRegion, uint32 SourcePitch, const uint8* SourceData) final override;
+	virtual void* LockTexture2D_RenderThread(class FRHICommandListImmediate& RHICmdList, FRHITexture2D* Texture, uint32 MipIndex, EResourceLockMode LockMode, uint32& DestStride, bool bLockWithinMiptail, bool bNeedsDefaultRHIFlush = true) final override;
+	virtual void UnlockTexture2D_RenderThread(class FRHICommandListImmediate& RHICmdList, FRHITexture2D* Texture, uint32 MipIndex, bool bLockWithinMiptail, bool bNeedsDefaultRHIFlush = true) final override;
+	virtual void UpdateTexture2D_RenderThread(class FRHICommandListImmediate& RHICmdList, FRHITexture2D* Texture, uint32 MipIndex, const struct FUpdateTextureRegion2D& UpdateRegion, uint32 SourcePitch, const uint8* SourceData) final override;
 	virtual void EndUpdateTexture3D_RenderThread(class FRHICommandListImmediate& RHICmdList, FUpdateTexture3DData& UpdateData) final override;
-	virtual void UpdateTexture3D_RenderThread(class FRHICommandListImmediate& RHICmdList, FTexture3DRHIParamRef Texture, uint32 MipIndex, const struct FUpdateTextureRegion3D& UpdateRegion, uint32 SourceRowPitch, uint32 SourceDepthPitch, const uint8* SourceData) final override;
+	virtual void UpdateTexture3D_RenderThread(class FRHICommandListImmediate& RHICmdList, FRHITexture3D* Texture, uint32 MipIndex, const struct FUpdateTextureRegion3D& UpdateRegion, uint32 SourceRowPitch, uint32 SourceDepthPitch, const uint8* SourceData) final override;
 	virtual FTexture2DRHIRef RHICreateTexture2D_RenderThread(class FRHICommandListImmediate& RHICmdList, uint32 SizeX, uint32 SizeY, uint8 Format, uint32 NumMips, uint32 NumSamples, uint32 Flags, FRHIResourceCreateInfo& CreateInfo) final override;
 	virtual FTexture2DArrayRHIRef RHICreateTexture2DArray_RenderThread(class FRHICommandListImmediate& RHICmdList, uint32 SizeX, uint32 SizeY, uint32 SizeZ, uint8 Format, uint32 NumMips, uint32 NumSamples, uint32 Flags, FRHIResourceCreateInfo& CreateInfo) final override;
 	virtual FTexture3DRHIRef RHICreateTexture3D_RenderThread(class FRHICommandListImmediate& RHICmdList, uint32 SizeX, uint32 SizeY, uint32 SizeZ, uint8 Format, uint32 NumMips, uint32 Flags, FRHIResourceCreateInfo& CreateInfo) final override;
-	virtual FUnorderedAccessViewRHIRef RHICreateUnorderedAccessView_RenderThread(class FRHICommandListImmediate& RHICmdList, FStructuredBufferRHIParamRef StructuredBuffer, bool bUseUAVCounter, bool bAppendBuffer) final override;
-	virtual FUnorderedAccessViewRHIRef RHICreateUnorderedAccessView_RenderThread(class FRHICommandListImmediate& RHICmdList, FTextureRHIParamRef Texture, uint32 MipLevel) final override;
-	virtual FUnorderedAccessViewRHIRef RHICreateUnorderedAccessView_RenderThread(class FRHICommandListImmediate& RHICmdList, FVertexBufferRHIParamRef VertexBuffer, uint8 Format) final override;
-	virtual FUnorderedAccessViewRHIRef RHICreateUnorderedAccessView_RenderThread(class FRHICommandListImmediate& RHICmdList, FIndexBufferRHIParamRef IndexBuffer, uint8 Format) final override;
-	virtual FShaderResourceViewRHIRef RHICreateShaderResourceView_RenderThread(class FRHICommandListImmediate& RHICmdList, FTextureRHIParamRef Texture, const FRHITextureSRVCreateInfo& CreateInfo) final override;
-	virtual FShaderResourceViewRHIRef RHICreateShaderResourceView_RenderThread(class FRHICommandListImmediate& RHICmdList, FStructuredBufferRHIParamRef StructuredBuffer) final override;
+	virtual FUnorderedAccessViewRHIRef RHICreateUnorderedAccessView_RenderThread(class FRHICommandListImmediate& RHICmdList, FRHIStructuredBuffer* StructuredBuffer, bool bUseUAVCounter, bool bAppendBuffer) final override;
+	virtual FUnorderedAccessViewRHIRef RHICreateUnorderedAccessView_RenderThread(class FRHICommandListImmediate& RHICmdList, FRHITexture* Texture, uint32 MipLevel) final override;
+	virtual FUnorderedAccessViewRHIRef RHICreateUnorderedAccessView_RenderThread(class FRHICommandListImmediate& RHICmdList, FRHIVertexBuffer* VertexBuffer, uint8 Format) final override;
+	virtual FUnorderedAccessViewRHIRef RHICreateUnorderedAccessView_RenderThread(class FRHICommandListImmediate& RHICmdList, FRHIIndexBuffer* IndexBuffer, uint8 Format) final override;
+	virtual FShaderResourceViewRHIRef RHICreateShaderResourceView_RenderThread(class FRHICommandListImmediate& RHICmdList, FRHITexture* Texture, const FRHITextureSRVCreateInfo& CreateInfo) final override;
+	virtual FShaderResourceViewRHIRef RHICreateShaderResourceView_RenderThread(class FRHICommandListImmediate& RHICmdList, FRHIStructuredBuffer* StructuredBuffer) final override;
 	virtual FTextureCubeRHIRef RHICreateTextureCube_RenderThread(class FRHICommandListImmediate& RHICmdList, uint32 Size, uint8 Format, uint32 NumMips, uint32 Flags, FRHIResourceCreateInfo& CreateInfo) final override;
 	virtual FTextureCubeRHIRef RHICreateTextureCubeArray_RenderThread(class FRHICommandListImmediate& RHICmdList, uint32 Size, uint32 ArraySize, uint8 Format, uint32 NumMips, uint32 Flags, FRHIResourceCreateInfo& CreateInfo) final override;
 	virtual FRenderQueryRHIRef RHICreateRenderQuery_RenderThread(class FRHICommandListImmediate& RHICmdList, ERenderQueryType QueryType) final override;
 	virtual FTextureReferenceRHIRef RHICreateTextureReference_RenderThread(class FRHICommandListImmediate& RHICmdList, FLastRenderTimeContainer* LastRenderTime) final override;
-	virtual void* RHILockTextureCubeFace_RenderThread(class FRHICommandListImmediate& RHICmdList, FTextureCubeRHIParamRef Texture, uint32 FaceIndex, uint32 ArrayIndex, uint32 MipIndex, EResourceLockMode LockMode, uint32& DestStride, bool bLockWithinMiptail) final override;
-	virtual void RHIUnlockTextureCubeFace_RenderThread(class FRHICommandListImmediate& RHICmdList, FTextureCubeRHIParamRef Texture, uint32 FaceIndex, uint32 ArrayIndex, uint32 MipIndex, bool bLockWithinMiptail) final override;
-	virtual void RHICopySubTextureRegion_RenderThread(class FRHICommandListImmediate& RHICmdList, FTexture2DRHIParamRef SourceTexture, FTexture2DRHIParamRef DestinationTexture, FBox2D SourceBox, FBox2D DestinationBox) final override;
+	virtual void* RHILockTextureCubeFace_RenderThread(class FRHICommandListImmediate& RHICmdList, FRHITextureCube* Texture, uint32 FaceIndex, uint32 ArrayIndex, uint32 MipIndex, EResourceLockMode LockMode, uint32& DestStride, bool bLockWithinMiptail) final override;
+	virtual void RHIUnlockTextureCubeFace_RenderThread(class FRHICommandListImmediate& RHICmdList, FRHITextureCube* Texture, uint32 FaceIndex, uint32 ArrayIndex, uint32 MipIndex, bool bLockWithinMiptail) final override;
+	virtual void RHICopySubTextureRegion_RenderThread(class FRHICommandListImmediate& RHICmdList, FRHITexture2D* SourceTexture, FRHITexture2D* DestinationTexture, FBox2D SourceBox, FBox2D DestinationBox) final override;
 
 	virtual void RHIBeginRenderPass(const FRHIRenderPassInfo& InInfo, const TCHAR* InName) final override
 	{
@@ -1014,11 +1014,11 @@ protected:
 	* @param StagingRectOUT - parameter is filled with the rectangle to read from the returned texture.
 	* @return The CPU readable Texture object.
 	*/
-	TRefCountPtr<ID3D11Texture2D> GetStagingTexture(FTextureRHIParamRef TextureRHI,FIntRect InRect, FIntRect& OutRect, FReadSurfaceDataFlags InFlags);
+	TRefCountPtr<ID3D11Texture2D> GetStagingTexture(FRHITexture* TextureRHI,FIntRect InRect, FIntRect& OutRect, FReadSurfaceDataFlags InFlags);
 
-	void ReadSurfaceDataNoMSAARaw(FTextureRHIParamRef TextureRHI,FIntRect Rect,TArray<uint8>& OutData, FReadSurfaceDataFlags InFlags);
+	void ReadSurfaceDataNoMSAARaw(FRHITexture* TextureRHI,FIntRect Rect,TArray<uint8>& OutData, FReadSurfaceDataFlags InFlags);
 
-	void ReadSurfaceDataMSAARaw(FRHICommandList_RecursiveHazardous& RHICmdList, FTextureRHIParamRef TextureRHI, FIntRect Rect, TArray<uint8>& OutData, FReadSurfaceDataFlags InFlags);
+	void ReadSurfaceDataMSAARaw(FRHICommandList_RecursiveHazardous& RHICmdList, FRHITexture* TextureRHI, FIntRect Rect, TArray<uint8>& OutData, FReadSurfaceDataFlags InFlags);
 
 #if NV_AFTERMATH
 	void StartNVAftermath();
diff --git a/Engine/Source/Runtime/Windows/D3D11RHI/Private/D3D11RenderTarget.cpp b/Engine/Source/Runtime/Windows/D3D11RHI/Private/D3D11RenderTarget.cpp
--- a/Engine/Source/Runtime/Windows/D3D11RHI/Private/D3D11RenderTarget.cpp
+++ b/Engine/Source/Runtime/Windows/D3D11RHI/Private/D3D11RenderTarget.cpp
@@ -217,7 +217,7 @@ void FD3D11DynamicRHI::ResolveTextureUsingShader(
 * @param bKeepOriginalSurface - true if the original surface will still be used after this function so must remain valid
 * @param ResolveParams - optional resolve params
 */
-void FD3D11DynamicRHI::RHICopyToResolveTarget(FTextureRHIParamRef SourceTextureRHI, FTextureRHIParamRef DestTextureRHI, const FResolveParams& ResolveParams)
+void FD3D11DynamicRHI::RHICopyToResolveTarget(FRHITexture* SourceTextureRHI, FRHITexture* DestTextureRHI, const FResolveParams& ResolveParams)
 {
 	if (!SourceTextureRHI || !DestTextureRHI)
 	{
@@ -486,7 +486,7 @@ static uint32 ComputeBytesPerPixel(DXGI_FORMAT Format)
 	return BytesPerPixel;
 }
 
-TRefCountPtr<ID3D11Texture2D> FD3D11DynamicRHI::GetStagingTexture(FTextureRHIParamRef TextureRHI,FIntRect InRect, FIntRect& StagingRectOUT, FReadSurfaceDataFlags InFlags)
+TRefCountPtr<ID3D11Texture2D> FD3D11DynamicRHI::GetStagingTexture(FRHITexture* TextureRHI,FIntRect InRect, FIntRect& StagingRectOUT, FReadSurfaceDataFlags InFlags)
 {
 	FD3D11TextureBase* Texture = GetD3D11TextureFromRHITexture(TextureRHI);
 	D3D11_TEXTURE2D_DESC SourceDesc; 
@@ -555,7 +555,7 @@ TRefCountPtr<ID3D11Texture2D> FD3D11DynamicRHI::GetStagingTexture(FTextureRHIPar
 	return TempTexture2D;
 }
 
-void FD3D11DynamicRHI::ReadSurfaceDataNoMSAARaw(FTextureRHIParamRef TextureRHI,FIntRect InRect,TArray<uint8>& OutData, FReadSurfaceDataFlags InFlags)
+void FD3D11DynamicRHI::ReadSurfaceDataNoMSAARaw(FRHITexture* TextureRHI,FIntRect InRect,TArray<uint8>& OutData, FReadSurfaceDataFlags InFlags)
 {
 	checkf(InRect.Width() <= TextureRHI->GetSizeXYZ().X >> InFlags.GetMip(), TEXT("Provided rect width (%d), must be smaller or equal to the texture size requested Mip (%d)"), InRect.Width(), TextureRHI->GetSizeXYZ().X >> InFlags.GetMip());
 	checkf(InRect.Height() <= TextureRHI->GetSizeXYZ().Y >> InFlags.GetMip(), TEXT("Provided rect height (%d), must be smaller or equal to the texture size requested Mip (%d)"), InRect.Height(), TextureRHI->GetSizeXYZ().Y >> InFlags.GetMip());
@@ -921,7 +921,7 @@ static void ConvertRAWSurfaceDataToFColor(DXGI_FORMAT Format, uint32 Width, uint
 	}
 }
 
-void FD3D11DynamicRHI::RHIReadSurfaceData(FTextureRHIParamRef TextureRHI,FIntRect InRect,TArray<FColor>& OutData, FReadSurfaceDataFlags InFlags)
+void FD3D11DynamicRHI::RHIReadSurfaceData(FRHITexture* TextureRHI,FIntRect InRect,TArray<FColor>& OutData, FReadSurfaceDataFlags InFlags)
 {
 	if (!ensure(TextureRHI))
 	{
@@ -964,7 +964,7 @@ void FD3D11DynamicRHI::RHIReadSurfaceData(FTextureRHIParamRef TextureRHI,FIntRec
 	ConvertRAWSurfaceDataToFColor(TextureDesc.Format, SizeX, SizeY, OutDataRaw.GetData(), SrcPitch, OutData.GetData(), InFlags);
 }
 
-void FD3D11DynamicRHI::ReadSurfaceDataMSAARaw(FRHICommandList_RecursiveHazardous& RHICmdList, FTextureRHIParamRef TextureRHI,FIntRect InRect,TArray<uint8>& OutData, FReadSurfaceDataFlags InFlags)
+void FD3D11DynamicRHI::ReadSurfaceDataMSAARaw(FRHICommandList_RecursiveHazardous& RHICmdList, FRHITexture* TextureRHI,FIntRect InRect,TArray<uint8>& OutData, FReadSurfaceDataFlags InFlags)
 {
 	FD3D11TextureBase* Texture = GetD3D11TextureFromRHITexture(TextureRHI);
 
@@ -1090,7 +1090,7 @@ void FD3D11DynamicRHI::ReadSurfaceDataMSAARaw(FRHICommandList_RecursiveHazardous
 	}
 }
 
-void FD3D11DynamicRHI::RHIMapStagingSurface(FTextureRHIParamRef TextureRHI,void*& OutData,int32& OutWidth,int32& OutHeight)
+void FD3D11DynamicRHI::RHIMapStagingSurface(FRHITexture* TextureRHI,void*& OutData,int32& OutWidth,int32& OutHeight)
 {
 	ID3D11Texture2D* Texture = (ID3D11Texture2D*)(GetD3D11TextureFromRHITexture(TextureRHI)->GetResource());
 	
@@ -1108,14 +1108,14 @@ void FD3D11DynamicRHI::RHIMapStagingSurface(FTextureRHIParamRef TextureRHI,void*
 	check(OutData);
 }
 
-void FD3D11DynamicRHI::RHIUnmapStagingSurface(FTextureRHIParamRef TextureRHI)
+void FD3D11DynamicRHI::RHIUnmapStagingSurface(FRHITexture* TextureRHI)
 {
 	ID3D11Texture2D* Texture = (ID3D11Texture2D*)(GetD3D11TextureFromRHITexture(TextureRHI)->GetResource());
 
 	Direct3DDeviceIMContext->Unmap(Texture,0);
 }
 
-void FD3D11DynamicRHI::RHIReadSurfaceFloatData(FTextureRHIParamRef TextureRHI,FIntRect InRect,TArray<FFloat16Color>& OutData,ECubeFace CubeFace,int32 ArrayIndex,int32 MipIndex)
+void FD3D11DynamicRHI::RHIReadSurfaceFloatData(FRHITexture* TextureRHI,FIntRect InRect,TArray<FFloat16Color>& OutData,ECubeFace CubeFace,int32 ArrayIndex,int32 MipIndex)
 {
 	FD3D11TextureBase* Texture = GetD3D11TextureFromRHITexture(TextureRHI);
 
@@ -1479,7 +1479,7 @@ static void ConvertRAWSurfaceDataToFLinearColor(EPixelFormat Format, uint32 Widt
 	}
 }
 
-void FD3D11DynamicRHI::RHIReadSurfaceData(FTextureRHIParamRef TextureRHI, FIntRect InRect, TArray<FLinearColor>& OutData, FReadSurfaceDataFlags InFlags)
+void FD3D11DynamicRHI::RHIReadSurfaceData(FRHITexture* TextureRHI, FIntRect InRect, TArray<FLinearColor>& OutData, FReadSurfaceDataFlags InFlags)
 {
 	TArray<uint8> OutDataRaw;
 
@@ -1518,7 +1518,7 @@ void FD3D11DynamicRHI::RHIReadSurfaceData(FTextureRHIParamRef TextureRHI, FIntRe
 	}
 }
 
-void FD3D11DynamicRHI::RHIRead3DSurfaceFloatData(FTextureRHIParamRef TextureRHI,FIntRect InRect,FIntPoint ZMinMax,TArray<FFloat16Color>& OutData)
+void FD3D11DynamicRHI::RHIRead3DSurfaceFloatData(FRHITexture* TextureRHI,FIntRect InRect,FIntPoint ZMinMax,TArray<FFloat16Color>& OutData)
 {
 	FD3D11TextureBase* Texture = GetD3D11TextureFromRHITexture(TextureRHI);
 
diff --git a/Engine/Source/Runtime/Windows/D3D11RHI/Private/D3D11StructuredBuffer.cpp b/Engine/Source/Runtime/Windows/D3D11RHI/Private/D3D11StructuredBuffer.cpp
--- a/Engine/Source/Runtime/Windows/D3D11RHI/Private/D3D11StructuredBuffer.cpp
+++ b/Engine/Source/Runtime/Windows/D3D11RHI/Private/D3D11StructuredBuffer.cpp
@@ -102,7 +102,7 @@ FStructuredBufferRHIRef FD3D11DynamicRHI::CreateStructuredBuffer_RenderThread(
 	return RHICreateStructuredBuffer(Stride, Size, InUsage, CreateInfo);
 }
 
-void* FD3D11DynamicRHI::RHILockStructuredBuffer(FStructuredBufferRHIParamRef StructuredBufferRHI,uint32 Offset,uint32 Size,EResourceLockMode LockMode)
+void* FD3D11DynamicRHI::RHILockStructuredBuffer(FRHIStructuredBuffer* StructuredBufferRHI,uint32 Offset,uint32 Size,EResourceLockMode LockMode)
 {
 	FD3D11StructuredBuffer* StructuredBuffer = ResourceCast(StructuredBufferRHI);
 	
@@ -167,7 +167,7 @@ void* FD3D11DynamicRHI::RHILockStructuredBuffer(FStructuredBufferRHIParamRef Str
 	return (void*)((uint8*)LockedData.GetData() + Offset);
 }
 
-void FD3D11DynamicRHI::RHIUnlockStructuredBuffer(FStructuredBufferRHIParamRef StructuredBufferRHI)
+void FD3D11DynamicRHI::RHIUnlockStructuredBuffer(FRHIStructuredBuffer* StructuredBufferRHI)
 {
 	FD3D11StructuredBuffer* StructuredBuffer = ResourceCast(StructuredBufferRHI);
 
diff --git a/Engine/Source/Runtime/Windows/D3D11RHI/Private/D3D11Texture.cpp b/Engine/Source/Runtime/Windows/D3D11RHI/Private/D3D11Texture.cpp
--- a/Engine/Source/Runtime/Windows/D3D11RHI/Private/D3D11Texture.cpp
+++ b/Engine/Source/Runtime/Windows/D3D11RHI/Private/D3D11Texture.cpp
@@ -1233,7 +1233,7 @@ FTexture2DRHIRef FD3D11DynamicRHI::RHIAsyncCreateTexture2D(uint32 SizeX,uint32 S
 	return NewTexture;
 }
 
-void FD3D11DynamicRHI::RHICopySharedMips(FTexture2DRHIParamRef DestTexture2DRHI,FTexture2DRHIParamRef SrcTexture2DRHI)
+void FD3D11DynamicRHI::RHICopySharedMips(FRHITexture2D* DestTexture2DRHI, FRHITexture2D* SrcTexture2DRHI)
 {
 	FD3D11Texture2D* DestTexture2D = ResourceCast(DestTexture2DRHI);
 	FD3D11Texture2D* SrcTexture2D = ResourceCast(SrcTexture2DRHI);
@@ -1297,7 +1297,7 @@ FTexture3DRHIRef FD3D11DynamicRHI::RHICreateTexture3D_RenderThread(
 	return RHICreateTexture3D(SizeX, SizeY, SizeZ, Format, NumMips, Flags, CreateInfo);
 }
 
-void FD3D11DynamicRHI::RHIGetResourceInfo(FTextureRHIParamRef Ref, FRHIResourceInfo& OutInfo)
+void FD3D11DynamicRHI::RHIGetResourceInfo(FRHITexture* Ref, FRHIResourceInfo& OutInfo)
 {
 	if(Ref)
 	{
@@ -1305,7 +1305,7 @@ void FD3D11DynamicRHI::RHIGetResourceInfo(FTextureRHIParamRef Ref, FRHIResourceI
 	}
 }
 
-FShaderResourceViewRHIRef FD3D11DynamicRHI::RHICreateShaderResourceView(FTextureRHIParamRef TextureRHI, const FRHITextureSRVCreateInfo& CreateInfo)
+FShaderResourceViewRHIRef FD3D11DynamicRHI::RHICreateShaderResourceView(FRHITexture* TextureRHI, const FRHITextureSRVCreateInfo& CreateInfo)
 {
 	FD3D11TextureBase* Texture = GetD3D11TextureFromRHITexture(TextureRHI);
 
@@ -1388,13 +1388,13 @@ FShaderResourceViewRHIRef FD3D11DynamicRHI::RHICreateShaderResourceView(FTexture
 	return new FD3D11ShaderResourceView(ShaderResourceView, Texture);
 }
 
-FShaderResourceViewRHIRef FD3D11DynamicRHI::RHICreateShaderResourceView_RenderThread(class FRHICommandListImmediate& RHICmdList, FTextureRHIParamRef Texture, const FRHITextureSRVCreateInfo& CreateInfo)
+FShaderResourceViewRHIRef FD3D11DynamicRHI::RHICreateShaderResourceView_RenderThread(class FRHICommandListImmediate& RHICmdList, FRHITexture* Texture, const FRHITextureSRVCreateInfo& CreateInfo)
 {
 	return RHICreateShaderResourceView(Texture, CreateInfo);
 }
 
 /** Generates mip maps for the surface. */
-void FD3D11DynamicRHI::RHIGenerateMips(FTextureRHIParamRef TextureRHI)
+void FD3D11DynamicRHI::RHIGenerateMips(FRHITexture* TextureRHI)
 {
 	FD3D11TextureBase* Texture = GetD3D11TextureFromRHITexture(TextureRHI);
 	// Surface must have been created with D3D11_BIND_RENDER_TARGET for GenerateMips to work
@@ -1410,7 +1410,7 @@ void FD3D11DynamicRHI::RHIGenerateMips(FTextureRHIParamRef TextureRHI)
  * @param	TextureRHI		- Texture we want to know the size of
  * @return					- Size in Bytes
  */
-uint32 FD3D11DynamicRHI::RHIComputeMemorySize(FTextureRHIParamRef TextureRHI)
+uint32 FD3D11DynamicRHI::RHIComputeMemorySize(FRHITexture* TextureRHI)
 {
 	if(!TextureRHI)
 	{
@@ -1432,7 +1432,7 @@ uint32 FD3D11DynamicRHI::RHIComputeMemorySize(FTextureRHIParamRef TextureRHI)
  * @param RequestStatus			- Will be decremented by 1 when the reallocation is complete (success or failure).
  * @return						- New reference to the texture, or an invalid reference upon failure
  */
-void FD3D11DynamicRHI::RHIAsyncCopyTexture2DCopy(FTexture2DRHIParamRef NewTexture2DRHI, FTexture2DRHIParamRef Texture2DRHI, int32 NewMipCount, int32 NewSizeX, int32 NewSizeY, FThreadSafeCounter* RequestStatus)
+void FD3D11DynamicRHI::RHIAsyncCopyTexture2DCopy(FRHITexture2D* NewTexture2DRHI, FRHITexture2D* Texture2DRHI, int32 NewMipCount, int32 NewSizeX, int32 NewSizeY, FThreadSafeCounter* RequestStatus)
 {
 	FD3D11Texture2D* Texture2D = ResourceCast(Texture2DRHI);
 	FD3D11Texture2D* NewTexture2D = ResourceCast(NewTexture2DRHI);
@@ -1478,7 +1478,7 @@ void FD3D11DynamicRHI::RHIAsyncCopyTexture2DCopy(FTexture2DRHIParamRef NewTextur
  * @param RequestStatus	- Will be decremented by 1 when the reallocation is complete (success or failure).
  * @return				- New reference to the texture, or an invalid reference upon failure
  */
-FTexture2DRHIRef FD3D11DynamicRHI::RHIAsyncReallocateTexture2D(FTexture2DRHIParamRef Texture2DRHI, int32 NewMipCount, int32 NewSizeX, int32 NewSizeY, FThreadSafeCounter* RequestStatus)
+FTexture2DRHIRef FD3D11DynamicRHI::RHIAsyncReallocateTexture2D(FRHITexture2D* Texture2DRHI, int32 NewMipCount, int32 NewSizeX, int32 NewSizeY, FThreadSafeCounter* RequestStatus)
 {
 	FD3D11Texture2D* Texture2D = ResourceCast(Texture2DRHI);
 
@@ -1493,7 +1493,7 @@ FTexture2DRHIRef FD3D11DynamicRHI::RHIAsyncReallocateTexture2D(FTexture2DRHIPara
 
 FTexture2DRHIRef FD3D11DynamicRHI::AsyncReallocateTexture2D_RenderThread(
 	class FRHICommandListImmediate& RHICmdList,
-	FTexture2DRHIParamRef Texture2D,
+	FRHITexture2D* Texture2D,
 	int32 NewMipCount,
 	int32 NewSizeX,
 	int32 NewSizeY,
@@ -1529,14 +1529,14 @@ FTexture2DRHIRef FD3D11DynamicRHI::AsyncReallocateTexture2D_RenderThread(
  * @param Texture2D		- Texture to check the reallocation status for
  * @return				- Current reallocation status
  */
-ETextureReallocationStatus FD3D11DynamicRHI::RHIFinalizeAsyncReallocateTexture2D( FTexture2DRHIParamRef Texture2D, bool bBlockUntilCompleted )
+ETextureReallocationStatus FD3D11DynamicRHI::RHIFinalizeAsyncReallocateTexture2D(FRHITexture2D* Texture2D, bool bBlockUntilCompleted )
 {
 	return TexRealloc_Succeeded;
 }
 
 ETextureReallocationStatus FD3D11DynamicRHI::FinalizeAsyncReallocateTexture2D_RenderThread(
 	class FRHICommandListImmediate& RHICmdList,
-	FTexture2DRHIParamRef Texture2D,
+	FRHITexture2D* Texture2D,
 	bool bBlockUntilCompleted)
 {
 	return RHIFinalizeAsyncReallocateTexture2D(Texture2D, bBlockUntilCompleted);
@@ -1550,14 +1550,14 @@ ETextureReallocationStatus FD3D11DynamicRHI::FinalizeAsyncReallocateTexture2D_Re
  * @param bBlockUntilCompleted	If true, blocks until the cancellation is fully completed
  * @return						Reallocation status
  */
-ETextureReallocationStatus FD3D11DynamicRHI::RHICancelAsyncReallocateTexture2D( FTexture2DRHIParamRef Texture2D, bool bBlockUntilCompleted )
+ETextureReallocationStatus FD3D11DynamicRHI::RHICancelAsyncReallocateTexture2D(FRHITexture2D* Texture2D, bool bBlockUntilCompleted )
 {
 	return TexRealloc_Succeeded;
 }
 
 ETextureReallocationStatus FD3D11DynamicRHI::CancelAsyncReallocateTexture2D_RenderThread(
 	class FRHICommandListImmediate& RHICmdList,
-	FTexture2DRHIParamRef Texture2D,
+	FRHITexture2D* Texture2D,
 	bool bBlockUntilCompleted)
 {
 	return RHICancelAsyncReallocateTexture2D(Texture2D, bBlockUntilCompleted);
@@ -1704,7 +1704,7 @@ void TD3D11Texture2D<RHIResourceType>::Unlock(uint32 MipIndex,uint32 ArrayIndex)
 	}
 }
 
-void* FD3D11DynamicRHI::RHILockTexture2D(FTexture2DRHIParamRef TextureRHI,uint32 MipIndex,EResourceLockMode LockMode,uint32& DestStride,bool bLockWithinMiptail)
+void* FD3D11DynamicRHI::RHILockTexture2D(FRHITexture2D* TextureRHI,uint32 MipIndex,EResourceLockMode LockMode,uint32& DestStride,bool bLockWithinMiptail)
 {
 	check(TextureRHI);
 	FD3D11Texture2D* Texture = ResourceCast(TextureRHI);
@@ -1714,7 +1714,7 @@ void* FD3D11DynamicRHI::RHILockTexture2D(FTexture2DRHIParamRef TextureRHI,uint32
 
 void* FD3D11DynamicRHI::LockTexture2D_RenderThread(
 	class FRHICommandListImmediate& RHICmdList,
-	FTexture2DRHIParamRef Texture,
+	FRHITexture2D* Texture,
 	uint32 MipIndex,
 	EResourceLockMode LockMode,
 	uint32& DestStride,
@@ -1740,7 +1740,7 @@ void* FD3D11DynamicRHI::LockTexture2D_RenderThread(
 	return LockedTexture;
 }
 
-void FD3D11DynamicRHI::RHIUnlockTexture2D(FTexture2DRHIParamRef TextureRHI,uint32 MipIndex,bool bLockWithinMiptail)
+void FD3D11DynamicRHI::RHIUnlockTexture2D(FRHITexture2D* TextureRHI,uint32 MipIndex,bool bLockWithinMiptail)
 {
 	check(TextureRHI);
 	FD3D11Texture2D* Texture = ResourceCast(TextureRHI);
@@ -1749,7 +1749,7 @@ void FD3D11DynamicRHI::RHIUnlockTexture2D(FTexture2DRHIParamRef TextureRHI,uint3
 
 void FD3D11DynamicRHI::UnlockTexture2D_RenderThread(
 	class FRHICommandListImmediate& RHICmdList,
-	FTexture2DRHIParamRef Texture,
+	FRHITexture2D* Texture,
 	uint32 MipIndex,
 	bool bLockWithinMiptail,
 	bool bNeedsDefaultRHIFlush)
@@ -1767,20 +1767,20 @@ void FD3D11DynamicRHI::UnlockTexture2D_RenderThread(
 	}
 }
 
-void* FD3D11DynamicRHI::RHILockTexture2DArray(FTexture2DArrayRHIParamRef TextureRHI,uint32 TextureIndex,uint32 MipIndex,EResourceLockMode LockMode,uint32& DestStride,bool bLockWithinMiptail)
+void* FD3D11DynamicRHI::RHILockTexture2DArray(FRHITexture2DArray* TextureRHI,uint32 TextureIndex,uint32 MipIndex,EResourceLockMode LockMode,uint32& DestStride,bool bLockWithinMiptail)
 {
 	FD3D11Texture2DArray* Texture = ResourceCast(TextureRHI);
 	ConditionalClearShaderResource(Texture, false);
 	return Texture->Lock(MipIndex,TextureIndex,LockMode,DestStride);
 }
 
-void FD3D11DynamicRHI::RHIUnlockTexture2DArray(FTexture2DArrayRHIParamRef TextureRHI,uint32 TextureIndex,uint32 MipIndex,bool bLockWithinMiptail)
+void FD3D11DynamicRHI::RHIUnlockTexture2DArray(FRHITexture2DArray* TextureRHI,uint32 TextureIndex,uint32 MipIndex,bool bLockWithinMiptail)
 {
 	FD3D11Texture2DArray* Texture = ResourceCast(TextureRHI);
 	Texture->Unlock(MipIndex,TextureIndex);
 }
 
-void FD3D11DynamicRHI::RHIUpdateTexture2D(FTexture2DRHIParamRef TextureRHI,uint32 MipIndex,const FUpdateTextureRegion2D& UpdateRegion,uint32 SourcePitch,const uint8* SourceData)
+void FD3D11DynamicRHI::RHIUpdateTexture2D(FRHITexture2D* TextureRHI,uint32 MipIndex,const FUpdateTextureRegion2D& UpdateRegion,uint32 SourcePitch,const uint8* SourceData)
 {
 	FD3D11Texture2D* Texture = ResourceCast(TextureRHI);
 
@@ -1802,7 +1802,7 @@ void FD3D11DynamicRHI::RHIUpdateTexture2D(FTexture2DRHIParamRef TextureRHI,uint3
 
 void FD3D11DynamicRHI::UpdateTexture2D_RenderThread(
 	class FRHICommandListImmediate& RHICmdList,
-	FTexture2DRHIParamRef Texture,
+	FRHITexture2D* Texture,
 	uint32 MipIndex,
 	const struct FUpdateTextureRegion2D& UpdateRegion,
 	uint32 SourcePitch,
@@ -1825,7 +1825,7 @@ void FD3D11DynamicRHI::UpdateTexture2D_RenderThread(
 	}
 }
 
-void FD3D11DynamicRHI::RHIUpdateTexture3D(FTexture3DRHIParamRef TextureRHI,uint32 MipIndex,const FUpdateTextureRegion3D& UpdateRegion,uint32 SourceRowPitch,uint32 SourceDepthPitch,const uint8* SourceData)
+void FD3D11DynamicRHI::RHIUpdateTexture3D(FRHITexture3D* TextureRHI,uint32 MipIndex,const FUpdateTextureRegion3D& UpdateRegion,uint32 SourceRowPitch,uint32 SourceDepthPitch,const uint8* SourceData)
 {
 	FD3D11Texture3D* Texture = ResourceCast(TextureRHI);
 
@@ -1875,7 +1875,7 @@ void FD3D11DynamicRHI::EndUpdateTexture3D_RenderThread(class FRHICommandListImme
 
 void FD3D11DynamicRHI::UpdateTexture3D_RenderThread(
 	class FRHICommandListImmediate& RHICmdList,
-	FTexture3DRHIParamRef Texture,
+	FRHITexture3D* Texture,
 	uint32 MipIndex,
 	const struct FUpdateTextureRegion3D& UpdateRegion,
 	uint32 SourceRowPitch,
@@ -1935,7 +1935,7 @@ FTextureCubeRHIRef FD3D11DynamicRHI::RHICreateTextureCubeArray_RenderThread(
 	return RHICreateTextureCubeArray(Size, ArraySize, Format, NumMips, Flags, CreateInfo);
 }
 
-void* FD3D11DynamicRHI::RHILockTextureCubeFace(FTextureCubeRHIParamRef TextureCubeRHI,uint32 FaceIndex,uint32 ArrayIndex,uint32 MipIndex,EResourceLockMode LockMode,uint32& DestStride,bool bLockWithinMiptail)
+void* FD3D11DynamicRHI::RHILockTextureCubeFace(FRHITextureCube* TextureCubeRHI,uint32 FaceIndex,uint32 ArrayIndex,uint32 MipIndex,EResourceLockMode LockMode,uint32& DestStride,bool bLockWithinMiptail)
 {
 	FD3D11TextureCube* TextureCube = ResourceCast(TextureCubeRHI);
 	ConditionalClearShaderResource(TextureCube, false);
@@ -1945,7 +1945,7 @@ void* FD3D11DynamicRHI::RHILockTextureCubeFace(FTextureCubeRHIParamRef TextureCu
 
 void* FD3D11DynamicRHI::RHILockTextureCubeFace_RenderThread(
 	class FRHICommandListImmediate& RHICmdList,
-	FTextureCubeRHIParamRef Texture,
+	FRHITextureCube* Texture,
 	uint32 FaceIndex,
 	uint32 ArrayIndex,
 	uint32 MipIndex,
@@ -1973,7 +1973,7 @@ void* FD3D11DynamicRHI::RHILockTextureCubeFace_RenderThread(
 	return LockedTexture;
 }
 
-void FD3D11DynamicRHI::RHIUnlockTextureCubeFace(FTextureCubeRHIParamRef TextureCubeRHI,uint32 FaceIndex,uint32 ArrayIndex,uint32 MipIndex,bool bLockWithinMiptail)
+void FD3D11DynamicRHI::RHIUnlockTextureCubeFace(FRHITextureCube* TextureCubeRHI,uint32 FaceIndex,uint32 ArrayIndex,uint32 MipIndex,bool bLockWithinMiptail)
 {
 	FD3D11TextureCube* TextureCube = ResourceCast(TextureCubeRHI);
 	uint32 D3DFace = GetD3D11CubeFace((ECubeFace)FaceIndex);
@@ -1982,7 +1982,7 @@ void FD3D11DynamicRHI::RHIUnlockTextureCubeFace(FTextureCubeRHIParamRef TextureC
 
 void FD3D11DynamicRHI::RHIUnlockTextureCubeFace_RenderThread(
 	class FRHICommandListImmediate& RHICmdList,
-	FTextureCubeRHIParamRef Texture,
+	FRHITextureCube* Texture,
 	uint32 FaceIndex,
 	uint32 ArrayIndex,
 	uint32 MipIndex,
@@ -2001,7 +2001,7 @@ void FD3D11DynamicRHI::RHIUnlockTextureCubeFace_RenderThread(
 	}
 }
 
-void FD3D11DynamicRHI::RHIBindDebugLabelName(FTextureRHIParamRef TextureRHI, const TCHAR* Name)
+void FD3D11DynamicRHI::RHIBindDebugLabelName(FRHITexture* TextureRHI, const TCHAR* Name)
 {
 	//todo: require names at texture creation time.
 	FName DebugName(Name);
@@ -2022,11 +2022,11 @@ void FD3D11DynamicRHI::RHIBindDebugLabelName(FTextureRHIParamRef TextureRHI, con
 #endif
 }
 
-void FD3D11DynamicRHI::RHIVirtualTextureSetFirstMipInMemory(FTexture2DRHIParamRef TextureRHI, uint32 FirstMip)
+void FD3D11DynamicRHI::RHIVirtualTextureSetFirstMipInMemory(FRHITexture2D* TextureRHI, uint32 FirstMip)
 {
 }
 
-void FD3D11DynamicRHI::RHIVirtualTextureSetFirstMipVisible(FTexture2DRHIParamRef TextureRHI, uint32 FirstMip)
+void FD3D11DynamicRHI::RHIVirtualTextureSetFirstMipVisible(FRHITexture2D* TextureRHI, uint32 FirstMip)
 {
 }
 
@@ -2042,7 +2042,7 @@ FTextureReferenceRHIRef FD3D11DynamicRHI::RHICreateTextureReference_RenderThread
 	return RHICreateTextureReference(LastRenderTime);
 }
 
-void FD3D11DynamicRHI::RHICopySubTextureRegion(FTexture2DRHIParamRef SourceTextureRHI, FTexture2DRHIParamRef DestinationTextureRHI, FBox2D SourceBox, FBox2D DestinationBox)
+void FD3D11DynamicRHI::RHICopySubTextureRegion(FRHITexture2D* SourceTextureRHI, FRHITexture2D* DestinationTextureRHI, FBox2D SourceBox, FBox2D DestinationBox)
 {
 	FD3D11Texture2D* SourceTexture = ResourceCast(SourceTextureRHI);
 	FD3D11Texture2D* DestinationTexture = ResourceCast(DestinationTextureRHI);
@@ -2125,8 +2125,8 @@ void FD3D11DynamicRHI::RHICopySubTextureRegion(FTexture2DRHIParamRef SourceTextu
 
 void FD3D11DynamicRHI::RHICopySubTextureRegion_RenderThread(
 	class FRHICommandListImmediate& RHICmdList,
-	FTexture2DRHIParamRef SourceTexture,
-	FTexture2DRHIParamRef DestinationTexture,
+	FRHITexture2D* SourceTexture,
+	FRHITexture2D* DestinationTexture,
 	FBox2D SourceBox,
 	FBox2D DestinationBox)
 {
@@ -2143,7 +2143,7 @@ void FD3D11DynamicRHI::RHICopySubTextureRegion_RenderThread(
 	}
 }
 
-void FD3D11DynamicRHI::RHIUpdateTextureReference(FRHITextureReference* TextureRefRHI, FTextureRHIParamRef NewTextureRHI)
+void FD3D11DynamicRHI::RHIUpdateTextureReference(FRHITextureReference* TextureRefRHI, FRHITexture* NewTextureRHI)
 {
 	// Updating texture references is disallowed while the RHI could be caching them in referenced resource tables.
 	check(ResourceTableFrameCounter == INDEX_NONE);
@@ -2362,7 +2362,7 @@ FTextureCubeRHIRef FD3D11DynamicRHI::RHICreateTextureCubeFromResource(EPixelForm
 	return CreateTextureFromResource<FD3D11BaseTextureCube>(false, true, Format, TexCreateFlags, ClearValueBinding, TextureResource);
 }
 
-void FD3D11DynamicRHI::RHIAliasTextureResources(FTextureRHIParamRef DestTextureRHI, FTextureRHIParamRef SrcTextureRHI)
+void FD3D11DynamicRHI::RHIAliasTextureResources(FRHITexture* DestTextureRHI, FRHITexture* SrcTextureRHI)
 {
 	FD3D11TextureBase* DestTexture = GetD3D11TextureFromRHITexture(DestTextureRHI);
 	FD3D11TextureBase* SrcTexture = GetD3D11TextureFromRHITexture(SrcTextureRHI);
@@ -2373,7 +2373,7 @@ void FD3D11DynamicRHI::RHIAliasTextureResources(FTextureRHIParamRef DestTextureR
 	}
 }
 
-void FD3D11DynamicRHI::RHICopyTexture(FTextureRHIParamRef SourceTextureRHI, FTextureRHIParamRef DestTextureRHI, const FRHICopyTextureInfo& CopyInfo)
+void FD3D11DynamicRHI::RHICopyTexture(FRHITexture* SourceTextureRHI, FRHITexture* DestTextureRHI, const FRHICopyTextureInfo& CopyInfo)
 {
 	if (!SourceTextureRHI || !DestTextureRHI || SourceTextureRHI == DestTextureRHI)
 	{
diff --git a/Engine/Source/Runtime/Windows/D3D11RHI/Private/D3D11UAV.cpp b/Engine/Source/Runtime/Windows/D3D11RHI/Private/D3D11UAV.cpp
--- a/Engine/Source/Runtime/Windows/D3D11RHI/Private/D3D11UAV.cpp
+++ b/Engine/Source/Runtime/Windows/D3D11RHI/Private/D3D11UAV.cpp
@@ -2,7 +2,7 @@
 
 #include "D3D11RHIPrivate.h"
 
-FUnorderedAccessViewRHIRef FD3D11DynamicRHI::RHICreateUnorderedAccessView(FStructuredBufferRHIParamRef StructuredBufferRHI, bool bUseUAVCounter, bool bAppendBuffer)
+FUnorderedAccessViewRHIRef FD3D11DynamicRHI::RHICreateUnorderedAccessView(FRHIStructuredBuffer* StructuredBufferRHI, bool bUseUAVCounter, bool bAppendBuffer)
 {
 	FD3D11StructuredBuffer* StructuredBuffer = ResourceCast(StructuredBufferRHI);
 
@@ -54,14 +54,14 @@ FUnorderedAccessViewRHIRef FD3D11DynamicRHI::RHICreateUnorderedAccessView(FStruc
 
 FUnorderedAccessViewRHIRef FD3D11DynamicRHI::RHICreateUnorderedAccessView_RenderThread(
 	class FRHICommandListImmediate& RHICmdList,
-	FStructuredBufferRHIParamRef StructuredBuffer,
+	FRHIStructuredBuffer* StructuredBuffer,
 	bool bUseUAVCounter,
 	bool bAppendBuffer)
 {
 	return RHICreateUnorderedAccessView(StructuredBuffer, bUseUAVCounter, bAppendBuffer);
 }
 
-FUnorderedAccessViewRHIRef FD3D11DynamicRHI::RHICreateUnorderedAccessView(FTextureRHIParamRef TextureRHI, uint32 MipLevel)
+FUnorderedAccessViewRHIRef FD3D11DynamicRHI::RHICreateUnorderedAccessView(FRHITexture* TextureRHI, uint32 MipLevel)
 {
 	FD3D11TextureBase* Texture = GetD3D11TextureFromRHITexture(TextureRHI);
 	
@@ -107,13 +107,13 @@ FUnorderedAccessViewRHIRef FD3D11DynamicRHI::RHICreateUnorderedAccessView(FTextu
 
 FUnorderedAccessViewRHIRef FD3D11DynamicRHI::RHICreateUnorderedAccessView_RenderThread(
 	class FRHICommandListImmediate& RHICmdList,
-	FTextureRHIParamRef Texture,
+	FRHITexture* Texture,
 	uint32 MipLevel)
 {
 	return RHICreateUnorderedAccessView(Texture, MipLevel);
 }
 
-FUnorderedAccessViewRHIRef FD3D11DynamicRHI::RHICreateUnorderedAccessView(FVertexBufferRHIParamRef VertexBufferRHI, uint8 Format)
+FUnorderedAccessViewRHIRef FD3D11DynamicRHI::RHICreateUnorderedAccessView(FRHIVertexBuffer* VertexBufferRHI, uint8 Format)
 {
 	FD3D11VertexBuffer* VertexBuffer = ResourceCast(VertexBufferRHI);
 
@@ -144,13 +144,13 @@ FUnorderedAccessViewRHIRef FD3D11DynamicRHI::RHICreateUnorderedAccessView(FVerte
 
 FUnorderedAccessViewRHIRef FD3D11DynamicRHI::RHICreateUnorderedAccessView_RenderThread(
 	class FRHICommandListImmediate& RHICmdList,
-	FVertexBufferRHIParamRef VertexBuffer,
+	FRHIVertexBuffer* VertexBuffer,
 	uint8 Format)
 {
 	return RHICreateUnorderedAccessView(VertexBuffer, Format);
 }
 
-FUnorderedAccessViewRHIRef FD3D11DynamicRHI::RHICreateUnorderedAccessView(FIndexBufferRHIParamRef IndexBufferRHI, uint8 Format)
+FUnorderedAccessViewRHIRef FD3D11DynamicRHI::RHICreateUnorderedAccessView(FRHIIndexBuffer* IndexBufferRHI, uint8 Format)
 {
 	FD3D11IndexBuffer* IndexBuffer = ResourceCast(IndexBufferRHI);
 
@@ -181,13 +181,13 @@ FUnorderedAccessViewRHIRef FD3D11DynamicRHI::RHICreateUnorderedAccessView(FIndex
 
 FUnorderedAccessViewRHIRef FD3D11DynamicRHI::RHICreateUnorderedAccessView_RenderThread(
 	class FRHICommandListImmediate& RHICmdList,
-	FIndexBufferRHIParamRef IndexBuffer,
+	FRHIIndexBuffer* IndexBuffer,
 	uint8 Format)
 {
 	return RHICreateUnorderedAccessView(IndexBuffer, Format);
 }
 
-FShaderResourceViewRHIRef FD3D11DynamicRHI::RHICreateShaderResourceView(FStructuredBufferRHIParamRef StructuredBufferRHI)
+FShaderResourceViewRHIRef FD3D11DynamicRHI::RHICreateShaderResourceView(FRHIStructuredBuffer* StructuredBufferRHI)
 {
 	FD3D11StructuredBuffer* StructuredBuffer = ResourceCast(StructuredBufferRHI);
 
@@ -223,7 +223,7 @@ FShaderResourceViewRHIRef FD3D11DynamicRHI::RHICreateShaderResourceView(FStructu
 
 FShaderResourceViewRHIRef FD3D11DynamicRHI::RHICreateShaderResourceView_RenderThread(
 	class FRHICommandListImmediate& RHICmdList,
-	FStructuredBufferRHIParamRef StructuredBuffer)
+	FRHIStructuredBuffer* StructuredBuffer)
 {
 	return RHICreateShaderResourceView(StructuredBuffer);
 }
@@ -258,7 +258,7 @@ static void CreateD3D11ShaderResourceViewOnBuffer(ID3D11Device* Direct3DDevice,
 	}
 }
 
-FShaderResourceViewRHIRef FD3D11DynamicRHI::RHICreateShaderResourceView(FVertexBufferRHIParamRef VertexBufferRHI, uint32 Stride, uint8 Format)
+FShaderResourceViewRHIRef FD3D11DynamicRHI::RHICreateShaderResourceView(FRHIVertexBuffer* VertexBufferRHI, uint32 Stride, uint8 Format)
 {
 	if (!VertexBufferRHI)
 	{
@@ -276,14 +276,14 @@ FShaderResourceViewRHIRef FD3D11DynamicRHI::RHICreateShaderResourceView(FVertexB
 
 FShaderResourceViewRHIRef FD3D11DynamicRHI::CreateShaderResourceView_RenderThread(
 	class FRHICommandListImmediate& RHICmdList,
-	FVertexBufferRHIParamRef VertexBuffer,
+	FRHIVertexBuffer* VertexBuffer,
 	uint32 Stride,
 	uint8 Format)
 {
 	return RHICreateShaderResourceView(VertexBuffer, Stride, Format);
 }
 
-void FD3D11DynamicRHI::RHIUpdateShaderResourceView(FRHIShaderResourceView* SRV, FVertexBufferRHIParamRef VertexBufferRHI, uint32 Stride, uint8 Format)
+void FD3D11DynamicRHI::RHIUpdateShaderResourceView(FRHIShaderResourceView* SRV, FRHIVertexBuffer* VertexBufferRHI, uint32 Stride, uint8 Format)
 {
 	check(SRV);
 	FD3D11ShaderResourceView* SRVD3D11 = ResourceCast(SRV);
@@ -303,7 +303,7 @@ void FD3D11DynamicRHI::RHIUpdateShaderResourceView(FRHIShaderResourceView* SRV,
 	}
 }
 
-void FD3D11DynamicRHI::RHIUpdateShaderResourceView(FRHIShaderResourceView* SRV, FIndexBufferRHIParamRef IndexBufferRHI)
+void FD3D11DynamicRHI::RHIUpdateShaderResourceView(FRHIShaderResourceView* SRV, FRHIIndexBuffer* IndexBufferRHI)
 {
 	check(SRV);
 	FD3D11ShaderResourceView* SRVD3D11 = ResourceCast(SRV);
@@ -325,7 +325,7 @@ void FD3D11DynamicRHI::RHIUpdateShaderResourceView(FRHIShaderResourceView* SRV,
 	}
 }
 
-FShaderResourceViewRHIRef FD3D11DynamicRHI::RHICreateShaderResourceView(FIndexBufferRHIParamRef BufferRHI)
+FShaderResourceViewRHIRef FD3D11DynamicRHI::RHICreateShaderResourceView(FRHIIndexBuffer* BufferRHI)
 {
 	if (!BufferRHI)
 	{
@@ -347,7 +347,7 @@ FShaderResourceViewRHIRef FD3D11DynamicRHI::RHICreateShaderResourceView(FIndexBu
 
 FShaderResourceViewRHIRef FD3D11DynamicRHI::CreateShaderResourceView_RenderThread(
 	class FRHICommandListImmediate& RHICmdList,
-	FIndexBufferRHIParamRef Buffer)
+	FRHIIndexBuffer* Buffer)
 {
 	return RHICreateShaderResourceView(Buffer);
 }
diff --git a/Engine/Source/Runtime/Windows/D3D11RHI/Private/D3D11UniformBuffer.cpp b/Engine/Source/Runtime/Windows/D3D11RHI/Private/D3D11UniformBuffer.cpp
--- a/Engine/Source/Runtime/Windows/D3D11RHI/Private/D3D11UniformBuffer.cpp
+++ b/Engine/Source/Runtime/Windows/D3D11RHI/Private/D3D11UniformBuffer.cpp
@@ -288,7 +288,7 @@ void UpdateUniformBufferContents(FD3D11Device* Direct3DDevice, FD3D11DeviceConte
 	}
 }
 
-void FD3D11DynamicRHI::RHIUpdateUniformBuffer(FUniformBufferRHIParamRef UniformBufferRHI, const void* Contents)
+void FD3D11DynamicRHI::RHIUpdateUniformBuffer(FRHIUniformBuffer* UniformBufferRHI, const void* Contents)
 {
 	check(IsInRenderingThread());
 	check(UniformBufferRHI);
diff --git a/Engine/Source/Runtime/Windows/D3D11RHI/Private/D3D11VertexBuffer.cpp b/Engine/Source/Runtime/Windows/D3D11RHI/Private/D3D11VertexBuffer.cpp
--- a/Engine/Source/Runtime/Windows/D3D11RHI/Private/D3D11VertexBuffer.cpp
+++ b/Engine/Source/Runtime/Windows/D3D11RHI/Private/D3D11VertexBuffer.cpp
@@ -108,7 +108,7 @@ FVertexBufferRHIRef FD3D11DynamicRHI::CreateVertexBuffer_RenderThread(
 	return RHICreateVertexBuffer(Size, InUsage, CreateInfo);
 }
 
-void* FD3D11DynamicRHI::RHILockVertexBuffer(FVertexBufferRHIParamRef VertexBufferRHI,uint32 Offset,uint32 Size,EResourceLockMode LockMode)
+void* FD3D11DynamicRHI::RHILockVertexBuffer(FRHIVertexBuffer* VertexBufferRHI,uint32 Offset,uint32 Size,EResourceLockMode LockMode)
 {
 	check(Size > 0);
 
@@ -189,7 +189,7 @@ void* FD3D11DynamicRHI::RHILockVertexBuffer(FVertexBufferRHIParamRef VertexBuffe
 	return (void*)((uint8*)LockedData.GetData() + Offset);
 }
 
-void FD3D11DynamicRHI::RHIUnlockVertexBuffer(FVertexBufferRHIParamRef VertexBufferRHI)
+void FD3D11DynamicRHI::RHIUnlockVertexBuffer(FRHIVertexBuffer* VertexBufferRHI)
 {
 	FD3D11VertexBuffer* VertexBuffer = ResourceCast(VertexBufferRHI);
 
@@ -227,7 +227,7 @@ void FD3D11DynamicRHI::RHIUnlockVertexBuffer(FVertexBufferRHIParamRef VertexBuff
 	}
 }
 
-void FD3D11DynamicRHI::RHICopyVertexBuffer(FVertexBufferRHIParamRef SourceBufferRHI,FVertexBufferRHIParamRef DestBufferRHI)
+void FD3D11DynamicRHI::RHICopyVertexBuffer(FRHIVertexBuffer* SourceBufferRHI, FRHIVertexBuffer* DestBufferRHI)
 {
 	FD3D11VertexBuffer* SourceBuffer = ResourceCast(SourceBufferRHI);
 	FD3D11VertexBuffer* DestBuffer = ResourceCast(DestBufferRHI);
@@ -245,7 +245,7 @@ void FD3D11DynamicRHI::RHICopyVertexBuffer(FVertexBufferRHIParamRef SourceBuffer
 	GPUProfilingData.RegisterGPUWork(1);
 }
 
-void FD3D11DynamicRHI::RHITransferVertexBufferUnderlyingResource(FVertexBufferRHIParamRef DestVertexBuffer, FVertexBufferRHIParamRef SrcVertexBuffer)
+void FD3D11DynamicRHI::RHITransferVertexBufferUnderlyingResource(FRHIVertexBuffer* DestVertexBuffer, FRHIVertexBuffer* SrcVertexBuffer)
 {
 	check(DestVertexBuffer);
 	FD3D11VertexBuffer* Dest = ResourceCast(DestVertexBuffer);
diff --git a/Engine/Source/Runtime/Windows/D3D11RHI/Private/D3D11Viewport.cpp b/Engine/Source/Runtime/Windows/D3D11RHI/Private/D3D11Viewport.cpp
--- a/Engine/Source/Runtime/Windows/D3D11RHI/Private/D3D11Viewport.cpp
+++ b/Engine/Source/Runtime/Windows/D3D11RHI/Private/D3D11Viewport.cpp
@@ -604,7 +604,7 @@ void FD3D11DynamicRHI::RHITick( float DeltaTime )
  *	Viewport functions.
  *=============================================================================*/
 
-void FD3D11DynamicRHI::RHIBeginDrawingViewport(FRHIViewport* ViewportRHI, FTextureRHIParamRef RenderTarget)
+void FD3D11DynamicRHI::RHIBeginDrawingViewport(FRHIViewport* ViewportRHI, FRHITexture* RenderTarget)
 {
 	FD3D11Viewport* Viewport = ResourceCast(ViewportRHI);
 
