UAV overlap was accidentally never enabled
fix renderdoc check happening too late for certain SLI settings
Both changes proposed by ikrima/kite&lighting.

diff --git a/Engine/Source/Runtime/Windows/D3D11RHI/Private/D3D11Commands.cpp b/Engine/Source/Runtime/Windows/D3D11RHI/Private/D3D11Commands.cpp
--- a/Engine/Source/Runtime/Windows/D3D11RHI/Private/D3D11Commands.cpp
+++ b/Engine/Source/Runtime/Windows/D3D11RHI/Private/D3D11Commands.cpp
@@ -1960,8 +1960,8 @@ static bool GOverlapUAVOBegin = false;
 static bool IsUAVOverlapSupported()
 {
 	if (!GAllowUAVFlushExt ||
-		!IsRHIDeviceNVIDIA() ||
-		!IsRHIDeviceAMD())
+		(!IsRHIDeviceNVIDIA() &&
+		!IsRHIDeviceAMD()))
 	{
 		return false;
 	}
diff --git a/Engine/Source/Runtime/Windows/D3D11RHI/Private/D3D11Device.cpp b/Engine/Source/Runtime/Windows/D3D11RHI/Private/D3D11Device.cpp
--- a/Engine/Source/Runtime/Windows/D3D11RHI/Private/D3D11Device.cpp
+++ b/Engine/Source/Runtime/Windows/D3D11RHI/Private/D3D11Device.cpp
@@ -417,30 +417,6 @@ void FD3D11DynamicRHI::SetupAfterDeviceCreation()
 			D3D11RHI_ShouldAllowAsyncResourceCreation() ? TEXT("no driver support") : TEXT("disabled by user"));
 	}
 
-#if PLATFORM_WINDOWS
-	IUnknown* RenderDoc;
-	IID RenderDocID;
-	if (SUCCEEDED(IIDFromString(L"{A7AA6116-9C8D-4BBA-9083-B4D816B71B78}", &RenderDocID)))
-	{
-		if (SUCCEEDED(Direct3DDevice->QueryInterface(RenderDocID, (void**)(&RenderDoc))))
-		{
-			bRenderDoc = true;
-
-			// Running under RenderDoc, so enable capturing mode
-			GDynamicRHI->EnableIdealGPUCaptureOptions(true);
-		}
-	}
-
-	IUnknown* IntelGPA;
-	static const IID IntelGPAID = { 0xCCFFEF16, 0x7B69, 0x468F, {0xBC, 0xE3, 0xCD, 0x95, 0x33, 0x69, 0xA3, 0x9A} };
-
-	if (SUCCEEDED(Direct3DDevice->QueryInterface(IntelGPAID, (void**)(&IntelGPA))))
-	{
-		// Running under Intel GPA, so enable capturing mode
-		GDynamicRHI->EnableIdealGPUCaptureOptions(true);
-	}
-#endif
-
 	// Check for typed UAV load support
 	for (uint32 PF = 0; PF < PF_MAX; ++PF)
 	{
diff --git a/Engine/Source/Runtime/Windows/D3D11RHI/Private/Windows/WindowsD3D11Device.cpp b/Engine/Source/Runtime/Windows/D3D11RHI/Private/Windows/WindowsD3D11Device.cpp
--- a/Engine/Source/Runtime/Windows/D3D11RHI/Private/Windows/WindowsD3D11Device.cpp
+++ b/Engine/Source/Runtime/Windows/D3D11RHI/Private/Windows/WindowsD3D11Device.cpp
@@ -1630,7 +1630,32 @@ void FD3D11DynamicRHI::InitD3DDevice()
 #endif //NVAPI_INTERFACE
 
 		CACHE_NV_AFTERMATH_ENABLED();
-		
+
+#if PLATFORM_WINDOWS
+		IUnknown* RenderDoc;
+		IID RenderDocID;
+		if (SUCCEEDED(IIDFromString(L"{A7AA6116-9C8D-4BBA-9083-B4D816B71B78}", &RenderDocID)))
+		{
+			if (SUCCEEDED(Direct3DDevice->QueryInterface(RenderDocID, (void**)(&RenderDoc))))
+			{
+				bRenderDoc = true;
+
+				// Running under RenderDoc, so enable capturing mode
+				GDynamicRHI->EnableIdealGPUCaptureOptions(true);
+			}
+		}
+
+		IUnknown* IntelGPA;
+		static const IID IntelGPAID = { 0xCCFFEF16, 0x7B69, 0x468F, {0xBC, 0xE3, 0xCD, 0x95, 0x33, 0x69, 0xA3, 0x9A} };
+
+		if (SUCCEEDED(Direct3DDevice->QueryInterface(IntelGPAID, (void**)(&IntelGPA))))
+		{
+			// Running under Intel GPA, so enable capturing mode
+			GDynamicRHI->EnableIdealGPUCaptureOptions(true);
+		}
+#endif
+
+
 #if WITH_SLI
 
 		GNumAlternateFrameRenderingGroups = 1;
