flush logs, so the aftermath call stack is showing up in the log.

diff --git a/Engine/Source/Runtime/Windows/D3D11RHI/Private/D3D11RHI.cpp b/Engine/Source/Runtime/Windows/D3D11RHI/Private/D3D11RHI.cpp
--- a/Engine/Source/Runtime/Windows/D3D11RHI/Private/D3D11RHI.cpp
+++ b/Engine/Source/Runtime/Windows/D3D11RHI/Private/D3D11RHI.cpp
@@ -693,7 +693,7 @@ bool FD3DGPUProfiler::CheckGpuHeartbeat() const
 					UE_LOG(LogRHI, Error, TEXT("[Aftermath] Invalid context handle"));
 					NVAFTERMATH_ON_ERROR();
 				}
-
+				GLog->PanicFlushThreadedLogs();
 				return false;
 			}
 		}
