Use DYNAMIC constant buffer usage when not pooling uniform buffers to allow for a map discard.

diff --git a/Engine/Source/Runtime/Windows/D3D11RHI/Private/D3D11UniformBuffer.cpp b/Engine/Source/Runtime/Windows/D3D11RHI/Private/D3D11UniformBuffer.cpp
--- a/Engine/Source/Runtime/Windows/D3D11RHI/Private/D3D11UniformBuffer.cpp
+++ b/Engine/Source/Runtime/Windows/D3D11RHI/Private/D3D11UniformBuffer.cpp
@@ -223,9 +223,9 @@ FUniformBufferRHIRef FD3D11DynamicRHI::RHICreateUniformBuffer(const void* Conten
 			// No pooling
 			D3D11_BUFFER_DESC Desc;
 			Desc.ByteWidth = NumBytes;
-			Desc.Usage = D3D11_USAGE_IMMUTABLE;
+			Desc.Usage = D3D11_USAGE_DYNAMIC;
 			Desc.BindFlags = D3D11_BIND_CONSTANT_BUFFER;
-			Desc.CPUAccessFlags = 0;
+			Desc.CPUAccessFlags = D3D11_CPU_ACCESS_WRITE;
 			Desc.MiscFlags = 0;
 			Desc.StructureByteStride = 0;
 
