Implemented d3d11 staging buffer. This fixes slow GPU read data time with GPU particles in d3d11.

diff --git a/Engine/Source/Runtime/Windows/D3D11RHI/Private/D3D11Commands.cpp b/Engine/Source/Runtime/Windows/D3D11RHI/Private/D3D11Commands.cpp
--- a/Engine/Source/Runtime/Windows/D3D11RHI/Private/D3D11Commands.cpp
+++ b/Engine/Source/Runtime/Windows/D3D11RHI/Private/D3D11Commands.cpp
@@ -2151,3 +2151,71 @@ void FD3D11DynamicRHI::RHIFlushComputeShaderCache()
 
 	EndUAVOverlap();
 }
+
+//*********************** StagingBuffer Implementation ***********************//
+
+FStagingBufferRHIRef FD3D11DynamicRHI::RHICreateStagingBuffer(FVertexBufferRHIParamRef VertexBufferRHI)
+{
+	return new FD3D11StagingBuffer(VertexBufferRHI);
+}
+
+void FD3D11DynamicRHI::RHIEnqueueStagedRead(FStagingBufferRHIParamRef StagingBufferRHI, FGPUFenceRHIParamRef Fence, uint32 Offset, uint32 NumBytes)
+{
+	FD3D11StagingBuffer* StagingBuffer = ResourceCast(StagingBufferRHI);
+	FD3D11VertexBuffer* VertexBuffer = ResourceCast(StagingBuffer->GetBackingBuffer());
+	if (StagingBuffer && VertexBuffer)
+	{
+		// Free previously allocated buffer.
+		StagingBuffer->StagedRead.SafeRelease();
+
+		// If the static buffer is being locked for reading, create a staging buffer.
+		D3D11_BUFFER_DESC StagedReadDesc;
+		ZeroMemory( &StagedReadDesc, sizeof( D3D11_BUFFER_DESC ) );
+		StagedReadDesc.ByteWidth = NumBytes;
+		StagedReadDesc.Usage = D3D11_USAGE_STAGING;
+		StagedReadDesc.BindFlags = 0;
+		StagedReadDesc.CPUAccessFlags = D3D11_CPU_ACCESS_READ;
+		StagedReadDesc.MiscFlags = 0;
+		TRefCountPtr<ID3D11Buffer> StagingVertexBuffer;
+		VERIFYD3D11RESULT_EX(Direct3DDevice->CreateBuffer(&StagedReadDesc,NULL,StagingBuffer->StagedRead.GetInitReference()), Direct3DDevice);
+
+		// Copy the contents of the vertex buffer to the staging buffer.
+		D3D11_BOX SourceBox;
+		SourceBox.left = Offset;
+		SourceBox.right = NumBytes;
+		SourceBox.top = SourceBox.front = 0;
+		SourceBox.bottom = SourceBox.back = 1;
+		Direct3DDeviceIMContext->CopySubresourceRegion(StagingBuffer->StagedRead,0,0,0,0,VertexBuffer->Resource,0,&SourceBox);
+	}
+
+	if (Fence)
+	{
+		Fence->Write();
+	}
+}
+
+void* FD3D11DynamicRHI::RHILockStagingBuffer(FStagingBufferRHIParamRef StagingBufferRHI, uint32 Offset, uint32 SizeRHI)
+{
+	FD3D11StagingBuffer* StagingBuffer = ResourceCast(StagingBufferRHI);
+	if (StagingBuffer && StagingBuffer->StagedRead)
+	{
+		// Map the staging buffer's memory for reading.
+		D3D11_MAPPED_SUBRESOURCE MappedSubresource;
+		VERIFYD3D11RESULT_EX(Direct3DDeviceIMContext->Map(StagingBuffer->StagedRead ,0,D3D11_MAP_READ,0,&MappedSubresource), Direct3DDevice);
+
+		return (void*)((uint8*)MappedSubresource.pData + Offset);
+	}
+	else
+	{
+		return nullptr;
+	}
+}
+
+void FD3D11DynamicRHI::RHIUnlockStagingBuffer(FStagingBufferRHIParamRef StagingBufferRHI)
+{
+	FD3D11StagingBuffer* StagingBuffer = ResourceCast(StagingBufferRHI);
+	if (StagingBuffer && StagingBuffer->StagedRead)
+	{
+		Direct3DDeviceIMContext->Unmap(StagingBuffer->StagedRead,0);
+	}
+}
diff --git a/Engine/Source/Runtime/Windows/D3D11RHI/Private/D3D11RHIPrivate.h b/Engine/Source/Runtime/Windows/D3D11RHI/Private/D3D11RHIPrivate.h
--- a/Engine/Source/Runtime/Windows/D3D11RHI/Private/D3D11RHIPrivate.h
+++ b/Engine/Source/Runtime/Windows/D3D11RHI/Private/D3D11RHIPrivate.h
@@ -395,6 +395,10 @@ public:
 	virtual FGeometryShaderRHIRef RHICreateGeometryShader(const TArray<uint8>& Code) final override;
 	virtual FGeometryShaderRHIRef RHICreateGeometryShaderWithStreamOutput(const TArray<uint8>& Code, const FStreamOutElementList& ElementList, uint32 NumStrides, const uint32* Strides, int32 RasterizedStream) final override;
 	virtual FComputeShaderRHIRef RHICreateComputeShader(const TArray<uint8>& Code) final override;
+	virtual FStagingBufferRHIRef RHICreateStagingBuffer(FVertexBufferRHIParamRef VertexBuffer) final override;
+    virtual void RHIEnqueueStagedRead(FStagingBufferRHIParamRef StagingBuffer, FGPUFenceRHIParamRef Fence, uint32 Offset, uint32 NumBytes) final override;
+	virtual void* RHILockStagingBuffer(FStagingBufferRHIParamRef StagingBuffer, uint32 Offset, uint32 SizeRHI) final override;
+    virtual void RHIUnlockStagingBuffer(FStagingBufferRHIParamRef StagingBuffer) final override;
 	virtual FBoundShaderStateRHIRef RHICreateBoundShaderState(FVertexDeclarationRHIParamRef VertexDeclaration, FVertexShaderRHIParamRef VertexShader, FHullShaderRHIParamRef HullShader, FDomainShaderRHIParamRef DomainShader, FPixelShaderRHIParamRef PixelShader, FGeometryShaderRHIParamRef GeometryShader) final override;
 	virtual FUniformBufferRHIRef RHICreateUniformBuffer(const void* Contents, const FRHIUniformBufferLayout& Layout, EUniformBufferUsage Usage) final override;
 	virtual FIndexBufferRHIRef RHICreateIndexBuffer(uint32 Stride, uint32 Size, uint32 InUsage, FRHIResourceCreateInfo& CreateInfo) final override;
diff --git a/Engine/Source/Runtime/Windows/D3D11RHI/Public/D3D11Resources.h b/Engine/Source/Runtime/Windows/D3D11RHI/Public/D3D11Resources.h
--- a/Engine/Source/Runtime/Windows/D3D11RHI/Public/D3D11Resources.h
+++ b/Engine/Source/Runtime/Windows/D3D11RHI/Public/D3D11Resources.h
@@ -782,6 +782,16 @@ public:
 	}
 };
 
+class FD3D11StagingBuffer : public FRHIStagingBuffer
+{
+public:
+	FD3D11StagingBuffer(FVertexBufferRHIRef InBuffer)
+		: FRHIStagingBuffer(InBuffer)
+	{}
+
+	TRefCountPtr<ID3D11Buffer> StagedRead;
+};
+
 /** Shader resource view class. */
 class FD3D11ShaderResourceView : public FRHIShaderResourceView
 {
@@ -910,6 +920,11 @@ struct TD3D11ResourceTraits<FRHIVertexBuffer>
 	typedef FD3D11VertexBuffer TConcreteType;
 };
 template<>
+struct TD3D11ResourceTraits<FRHIStagingBuffer>
+{
+	typedef FD3D11StagingBuffer TConcreteType;
+};
+template<>
 struct TD3D11ResourceTraits<FRHIShaderResourceView>
 {
 	typedef FD3D11ShaderResourceView TConcreteType;
