UE4.21 - Static analysis fixed
#rnx

diff --git a/Engine/Source/Runtime/Windows/D3D11RHI/Private/D3D11Commands.cpp b/Engine/Source/Runtime/Windows/D3D11RHI/Private/D3D11Commands.cpp
--- a/Engine/Source/Runtime/Windows/D3D11RHI/Private/D3D11Commands.cpp
+++ b/Engine/Source/Runtime/Windows/D3D11RHI/Private/D3D11Commands.cpp
@@ -2162,30 +2162,33 @@ FStagingBufferRHIRef FD3D11DynamicRHI::RHICreateStagingBuffer(FVertexBufferRHIPa
 void FD3D11DynamicRHI::RHIEnqueueStagedRead(FStagingBufferRHIParamRef StagingBufferRHI, FGPUFenceRHIParamRef Fence, uint32 Offset, uint32 NumBytes)
 {
 	FD3D11StagingBuffer* StagingBuffer = ResourceCast(StagingBufferRHI);
-	FD3D11VertexBuffer* VertexBuffer = ResourceCast(StagingBuffer->GetBackingBuffer());
-	if (StagingBuffer && VertexBuffer)
-	{
-		// Free previously allocated buffer.
-		StagingBuffer->StagedRead.SafeRelease();
-
-		// If the static buffer is being locked for reading, create a staging buffer.
-		D3D11_BUFFER_DESC StagedReadDesc;
-		ZeroMemory( &StagedReadDesc, sizeof( D3D11_BUFFER_DESC ) );
-		StagedReadDesc.ByteWidth = NumBytes;
-		StagedReadDesc.Usage = D3D11_USAGE_STAGING;
-		StagedReadDesc.BindFlags = 0;
-		StagedReadDesc.CPUAccessFlags = D3D11_CPU_ACCESS_READ;
-		StagedReadDesc.MiscFlags = 0;
-		TRefCountPtr<ID3D11Buffer> StagingVertexBuffer;
-		VERIFYD3D11RESULT_EX(Direct3DDevice->CreateBuffer(&StagedReadDesc,NULL,StagingBuffer->StagedRead.GetInitReference()), Direct3DDevice);
-
-		// Copy the contents of the vertex buffer to the staging buffer.
-		D3D11_BOX SourceBox;
-		SourceBox.left = Offset;
-		SourceBox.right = NumBytes;
-		SourceBox.top = SourceBox.front = 0;
-		SourceBox.bottom = SourceBox.back = 1;
-		Direct3DDeviceIMContext->CopySubresourceRegion(StagingBuffer->StagedRead,0,0,0,0,VertexBuffer->Resource,0,&SourceBox);
+	if (StagingBuffer)
+	{
+		FD3D11VertexBuffer* VertexBuffer = ResourceCast(StagingBuffer->GetBackingBuffer());
+		if (VertexBuffer)
+		{
+			// Free previously allocated buffer.
+			StagingBuffer->StagedRead.SafeRelease();
+
+			// If the static buffer is being locked for reading, create a staging buffer.
+			D3D11_BUFFER_DESC StagedReadDesc;
+			ZeroMemory(&StagedReadDesc, sizeof(D3D11_BUFFER_DESC));
+			StagedReadDesc.ByteWidth = NumBytes;
+			StagedReadDesc.Usage = D3D11_USAGE_STAGING;
+			StagedReadDesc.BindFlags = 0;
+			StagedReadDesc.CPUAccessFlags = D3D11_CPU_ACCESS_READ;
+			StagedReadDesc.MiscFlags = 0;
+			TRefCountPtr<ID3D11Buffer> StagingVertexBuffer;
+			VERIFYD3D11RESULT_EX(Direct3DDevice->CreateBuffer(&StagedReadDesc, NULL, StagingBuffer->StagedRead.GetInitReference()), Direct3DDevice);
+
+			// Copy the contents of the vertex buffer to the staging buffer.
+			D3D11_BOX SourceBox;
+			SourceBox.left = Offset;
+			SourceBox.right = NumBytes;
+			SourceBox.top = SourceBox.front = 0;
+			SourceBox.bottom = SourceBox.back = 1;
+			Direct3DDeviceIMContext->CopySubresourceRegion(StagingBuffer->StagedRead, 0, 0, 0, 0, VertexBuffer->Resource, 0, &SourceBox);
+		}
 	}
 
 	if (Fence)
