Merging CL 4531739 from //UE4/Dev-VR
Workaround crash when using the Google Daydream VR preview

diff --git a/Engine/Source/Runtime/Windows/D3D11RHI/Private/D3D11Query.cpp b/Engine/Source/Runtime/Windows/D3D11RHI/Private/D3D11Query.cpp
--- a/Engine/Source/Runtime/Windows/D3D11RHI/Private/D3D11Query.cpp
+++ b/Engine/Source/Runtime/Windows/D3D11RHI/Private/D3D11Query.cpp
@@ -130,9 +130,10 @@ bool FD3D11DynamicRHI::GetQueryData(ID3D11Query* Query,void* Data,SIZE_T DataSiz
 	{
 		return true;
 	}
-	else if(Result == S_FALSE && !bWait)
+	else if((Result == S_FALSE || Result == DXGI_ERROR_INVALID_CALL) && !bWait)
 	{
 		// Return failure if the query isn't complete, and waiting wasn't requested.
+		// HACK: Sometimes DXGI_ERROR_INVALID_CALL is returned instead of S_FALSE (UE-65122).
 		return false;
 	}
 	else
