DR - Prep for deprecation RHI's Draw*Primitive* calls to not use the PrimitiveType on D3D11
#rnx

diff --git a/Engine/Source/Runtime/Windows/D3D11RHI/Private/D3D11Commands.cpp b/Engine/Source/Runtime/Windows/D3D11RHI/Private/D3D11Commands.cpp
--- a/Engine/Source/Runtime/Windows/D3D11RHI/Private/D3D11Commands.cpp
+++ b/Engine/Source/Runtime/Windows/D3D11RHI/Private/D3D11Commands.cpp
@@ -279,9 +279,6 @@ void FD3D11DynamicRHI::RHISetScissorRect(bool bEnable,uint32 MinX,uint32 MinY,ui
 */
 void FD3D11DynamicRHI::RHISetBoundShaderState( FBoundShaderStateRHIParamRef BoundShaderStateRHI)
 {
-	// Non-PSO
-	PSOPrimitiveType = PT_Num;
-
 	FD3D11BoundShaderState* BoundShaderState = ResourceCast(BoundShaderStateRHI);
 
 	StateCache.SetStreamStrides(BoundShaderState->StreamStrides);
@@ -1205,7 +1202,7 @@ void FD3D11DynamicRHI::RHISetRenderTargetsAndClear(const FRHISetRenderTargetsInf
 
 // Primitive drawing.
 
-static D3D11_PRIMITIVE_TOPOLOGY GetD3D11PrimitiveType(uint32 PrimitiveType, bool bUsingTessellation)
+static D3D11_PRIMITIVE_TOPOLOGY GetD3D11PrimitiveType(EPrimitiveType PrimitiveType, bool bUsingTessellation)
 {
 	if(bUsingTessellation)
 	{
@@ -1278,9 +1275,9 @@ static D3D11_PRIMITIVE_TOPOLOGY GetD3D11PrimitiveType(uint32 PrimitiveType, bool
 	return D3D11_PRIMITIVE_TOPOLOGY_TRIANGLELIST;
 }
 
-static inline void VerifyPrimitiveType(EPrimitiveType PSOPrimitiveType, uint32 PrimitiveType)
+static inline void VerifyPrimitiveType(uint32 DrawCallPrimitiveType, EPrimitiveType PSOPrimitiveType)
 {
-	ensureAlwaysMsgf(PSOPrimitiveType == PT_Num || (EPrimitiveType)PrimitiveType == PSOPrimitiveType, TEXT("PSO was created using PrimitiveType %d, but the Draw call is using %d! This will break D3D12, Metal and Vulkan"), (uint32)PSOPrimitiveType, (uint32)PrimitiveType);
+	ensureAlwaysMsgf((EPrimitiveType)DrawCallPrimitiveType == PSOPrimitiveType, TEXT("PSO was created using PrimitiveType %d, but the Draw call is using %d! PrimitiveTypes must match as deprecation is imminent..."), (uint32)PSOPrimitiveType, (uint32)DrawCallPrimitiveType);
 }
 
 void FD3D11DynamicRHI::CommitNonComputeShaderConstants()
@@ -1585,7 +1582,7 @@ void FD3D11DynamicRHI::CommitComputeResourceTables(FD3D11ComputeShader* InComput
 	SetResourcesFromTables(ComputeShader);
 }
 
-void FD3D11DynamicRHI::RHIDrawPrimitive(uint32 PrimitiveType,uint32 BaseVertexIndex,uint32 NumPrimitives,uint32 NumInstances)
+void FD3D11DynamicRHI::RHIDrawPrimitive(uint32 PrimitiveType_DEPRECATED,uint32 BaseVertexIndex,uint32 NumPrimitives,uint32 NumInstances)
 {
 	RHI_DRAW_CALL_STATS(PrimitiveType,NumInstances*NumPrimitives);
 
@@ -1595,7 +1592,7 @@ void FD3D11DynamicRHI::RHIDrawPrimitive(uint32 PrimitiveType,uint32 BaseVertexIn
 	uint32 VertexCount = GetVertexCountForPrimitiveCount(NumPrimitives,PrimitiveType);
 
 	GPUProfilingData.RegisterGPUWork(NumPrimitives * NumInstances, VertexCount * NumInstances);
-	VerifyPrimitiveType(PSOPrimitiveType, PrimitiveType);
+	VerifyPrimitiveType(PrimitiveType_DEPRECATED, PrimitiveType);
 	StateCache.SetPrimitiveTopology(GetD3D11PrimitiveType(PrimitiveType,bUsingTessellation));
 	if(NumInstances > 1)
 	{
@@ -1607,7 +1604,7 @@ void FD3D11DynamicRHI::RHIDrawPrimitive(uint32 PrimitiveType,uint32 BaseVertexIn
 	}
 }
 
-void FD3D11DynamicRHI::RHIDrawPrimitiveIndirect(uint32 PrimitiveType,FVertexBufferRHIParamRef ArgumentBufferRHI,uint32 ArgumentOffset)
+void FD3D11DynamicRHI::RHIDrawPrimitiveIndirect(uint32 PrimitiveType_DEPRECATED,FVertexBufferRHIParamRef ArgumentBufferRHI,uint32 ArgumentOffset)
 {
 	FD3D11VertexBuffer* ArgumentBuffer = ResourceCast(ArgumentBufferRHI);
 
@@ -1618,12 +1615,12 @@ void FD3D11DynamicRHI::RHIDrawPrimitiveIndirect(uint32 PrimitiveType,FVertexBuff
 	CommitGraphicsResourceTables();
 	CommitNonComputeShaderConstants();
 
-	VerifyPrimitiveType(PSOPrimitiveType, PrimitiveType);
+	VerifyPrimitiveType(PrimitiveType_DEPRECATED, PrimitiveType);
 	StateCache.SetPrimitiveTopology(GetD3D11PrimitiveType(PrimitiveType,bUsingTessellation));
 	Direct3DDeviceIMContext->DrawInstancedIndirect(ArgumentBuffer->Resource,ArgumentOffset);
 }
 
-void FD3D11DynamicRHI::RHIDrawIndexedIndirect(FIndexBufferRHIParamRef IndexBufferRHI, uint32 PrimitiveType, FStructuredBufferRHIParamRef ArgumentsBufferRHI, int32 DrawArgumentsIndex, uint32 NumInstances)
+void FD3D11DynamicRHI::RHIDrawIndexedIndirect(FIndexBufferRHIParamRef IndexBufferRHI, uint32 PrimitiveType_DEPRECATED, FStructuredBufferRHIParamRef ArgumentsBufferRHI, int32 DrawArgumentsIndex, uint32 NumInstances)
 {
 	FD3D11IndexBuffer* IndexBuffer = ResourceCast(IndexBufferRHI);
 	FD3D11StructuredBuffer* ArgumentsBuffer = ResourceCast(ArgumentsBufferRHI);
@@ -1640,7 +1637,7 @@ void FD3D11DynamicRHI::RHIDrawIndexedIndirect(FIndexBufferRHIParamRef IndexBuffe
 	const DXGI_FORMAT Format = (IndexBuffer->GetStride() == sizeof(uint16) ? DXGI_FORMAT_R16_UINT : DXGI_FORMAT_R32_UINT);
 
 	StateCache.SetIndexBuffer(IndexBuffer->Resource, Format, 0);
-	VerifyPrimitiveType(PSOPrimitiveType, PrimitiveType);
+	VerifyPrimitiveType(PrimitiveType_DEPRECATED, PrimitiveType);
 	StateCache.SetPrimitiveTopology(GetD3D11PrimitiveType(PrimitiveType,bUsingTessellation));
 
 	if(NumInstances > 1)
@@ -1653,7 +1650,7 @@ void FD3D11DynamicRHI::RHIDrawIndexedIndirect(FIndexBufferRHIParamRef IndexBuffe
 	}
 }
 
-void FD3D11DynamicRHI::RHIDrawIndexedPrimitive(FIndexBufferRHIParamRef IndexBufferRHI,uint32 PrimitiveType,int32 BaseVertexIndex,uint32 FirstInstance,uint32 NumVertices,uint32 StartIndex,uint32 NumPrimitives,uint32 NumInstances)
+void FD3D11DynamicRHI::RHIDrawIndexedPrimitive(FIndexBufferRHIParamRef IndexBufferRHI,uint32 PrimitiveType_DEPRECATED,int32 BaseVertexIndex,uint32 FirstInstance,uint32 NumVertices,uint32 StartIndex,uint32 NumPrimitives,uint32 NumInstances)
 {
 	FD3D11IndexBuffer* IndexBuffer = ResourceCast(IndexBufferRHI);
 
@@ -1679,7 +1676,7 @@ void FD3D11DynamicRHI::RHIDrawIndexedPrimitive(FIndexBufferRHIParamRef IndexBuff
 		TEXT("Start %u, Count %u, Type %u, Buffer Size %u, Buffer stride %u"), StartIndex, IndexCount, PrimitiveType, IndexBuffer->GetSize(), IndexBuffer->GetStride());
 
 	StateCache.SetIndexBuffer(IndexBuffer->Resource, Format, 0);
-	VerifyPrimitiveType(PSOPrimitiveType, PrimitiveType);
+	VerifyPrimitiveType(PrimitiveType_DEPRECATED, PrimitiveType);
 	StateCache.SetPrimitiveTopology(GetD3D11PrimitiveType(PrimitiveType,bUsingTessellation));
 
 	if (NumInstances > 1 || FirstInstance != 0)
@@ -1694,7 +1691,7 @@ void FD3D11DynamicRHI::RHIDrawIndexedPrimitive(FIndexBufferRHIParamRef IndexBuff
 	}
 }
 
-void FD3D11DynamicRHI::RHIDrawIndexedPrimitiveIndirect(uint32 PrimitiveType,FIndexBufferRHIParamRef IndexBufferRHI,FVertexBufferRHIParamRef ArgumentBufferRHI,uint32 ArgumentOffset)
+void FD3D11DynamicRHI::RHIDrawIndexedPrimitiveIndirect(uint32 PrimitiveType_DEPRECATED,FIndexBufferRHIParamRef IndexBufferRHI,FVertexBufferRHIParamRef ArgumentBufferRHI,uint32 ArgumentOffset)
 {
 	FD3D11IndexBuffer* IndexBuffer = ResourceCast(IndexBufferRHI);
 	FD3D11VertexBuffer* ArgumentBuffer = ResourceCast(ArgumentBufferRHI);
@@ -1710,7 +1707,7 @@ void FD3D11DynamicRHI::RHIDrawIndexedPrimitiveIndirect(uint32 PrimitiveType,FInd
 	const uint32 SizeFormat = sizeof(DXGI_FORMAT);
 	const DXGI_FORMAT Format = (IndexBuffer->GetStride() == sizeof(uint16) ? DXGI_FORMAT_R16_UINT : DXGI_FORMAT_R32_UINT);
 	StateCache.SetIndexBuffer(IndexBuffer->Resource, Format, 0);
-	VerifyPrimitiveType(PSOPrimitiveType, PrimitiveType);
+	VerifyPrimitiveType(PrimitiveType_DEPRECATED, PrimitiveType);
 	StateCache.SetPrimitiveTopology(GetD3D11PrimitiveType(PrimitiveType,bUsingTessellation));
 	Direct3DDeviceIMContext->DrawIndexedInstancedIndirect(ArgumentBuffer->Resource,ArgumentOffset);
 }
@@ -1723,12 +1720,13 @@ void FD3D11DynamicRHI::RHIDrawIndexedPrimitiveIndirect(uint32 PrimitiveType,FInd
  * @param VertexDataStride Size of each vertex 
  * @param OutVertexData Reference to the allocated vertex memory
  */
-void FD3D11DynamicRHI::RHIBeginDrawPrimitiveUP( uint32 PrimitiveType, uint32 NumPrimitives, uint32 NumVertices, uint32 VertexDataStride, void*& OutVertexData)
+void FD3D11DynamicRHI::RHIBeginDrawPrimitiveUP( uint32 PrimitiveType_DEPRECATED, uint32 NumPrimitives, uint32 NumVertices, uint32 VertexDataStride, void*& OutVertexData)
 {
 	checkSlow( PendingNumVertices == 0 );
 
+	VerifyPrimitiveType(PrimitiveType_DEPRECATED, PrimitiveType);
+
 	// Remember the parameters for the draw call.
-	PendingPrimitiveType = PrimitiveType;
 	PendingNumPrimitives = NumPrimitives;
 	PendingNumVertices = NumVertices;
 	PendingVertexDataStride = VertexDataStride;
@@ -1742,9 +1740,9 @@ void FD3D11DynamicRHI::RHIBeginDrawPrimitiveUP( uint32 PrimitiveType, uint32 Num
  */
 void FD3D11DynamicRHI::RHIEndDrawPrimitiveUP()
 {
-	RHI_DRAW_CALL_STATS(PendingPrimitiveType,PendingNumPrimitives);
+	RHI_DRAW_CALL_STATS(PrimitiveType,PendingNumPrimitives);
 
-	checkSlow(!bUsingTessellation || PendingPrimitiveType == PT_TriangleList);
+	checkSlow(!bUsingTessellation || PrimitiveType == PT_TriangleList);
 
 	GPUProfilingData.RegisterGPUWork(PendingNumPrimitives,PendingNumVertices);
 
@@ -1756,12 +1754,10 @@ void FD3D11DynamicRHI::RHIEndDrawPrimitiveUP()
 	CommitGraphicsResourceTables();
 	CommitNonComputeShaderConstants();
 	StateCache.SetStreamSource(D3DBuffer, 0, PendingVertexDataStride, VBOffset);
-	VerifyPrimitiveType(PSOPrimitiveType, PendingPrimitiveType);
-	StateCache.SetPrimitiveTopology(GetD3D11PrimitiveType(PendingPrimitiveType,bUsingTessellation));
+	StateCache.SetPrimitiveTopology(GetD3D11PrimitiveType(PrimitiveType,bUsingTessellation));
 	Direct3DDeviceIMContext->Draw(PendingNumVertices,0);
 
 	// Clear these parameters.
-	PendingPrimitiveType = 0;
 	PendingNumPrimitives = 0;
 	PendingNumVertices = 0;
 	PendingVertexDataStride = 0;
@@ -1779,12 +1775,13 @@ void FD3D11DynamicRHI::RHIEndDrawPrimitiveUP()
  * @param IndexDataStride Size of each index (either 2 or 4 bytes)
  * @param OutIndexData Reference to the allocated index memory
  */
-void FD3D11DynamicRHI::RHIBeginDrawIndexedPrimitiveUP( uint32 PrimitiveType, uint32 NumPrimitives, uint32 NumVertices, uint32 VertexDataStride, void*& OutVertexData, uint32 MinVertexIndex, uint32 NumIndices, uint32 IndexDataStride, void*& OutIndexData)
+void FD3D11DynamicRHI::RHIBeginDrawIndexedPrimitiveUP( uint32 PrimitiveType_DEPRECATED, uint32 NumPrimitives, uint32 NumVertices, uint32 VertexDataStride, void*& OutVertexData, uint32 MinVertexIndex, uint32 NumIndices, uint32 IndexDataStride, void*& OutIndexData)
 {
 	checkSlow((sizeof(uint16) == IndexDataStride) || (sizeof(uint32) == IndexDataStride));
 
+	VerifyPrimitiveType(PrimitiveType_DEPRECATED, PrimitiveType);
+
 	// Store off information needed for the draw call.
-	PendingPrimitiveType = PrimitiveType;
 	PendingNumPrimitives = NumPrimitives;
 	PendingMinVertexIndex = MinVertexIndex;
 	PendingIndexDataStride = IndexDataStride;
@@ -1803,9 +1800,9 @@ void FD3D11DynamicRHI::RHIBeginDrawIndexedPrimitiveUP( uint32 PrimitiveType, uin
 void FD3D11DynamicRHI::RHIEndDrawIndexedPrimitiveUP()
 {
 	// tessellation only supports trilists
-	checkSlow(!bUsingTessellation || PendingPrimitiveType == PT_TriangleList);
+	checkSlow(!bUsingTessellation || PrimitiveType == PT_TriangleList);
 
-	RHI_DRAW_CALL_STATS(PendingPrimitiveType,PendingNumPrimitives);
+	RHI_DRAW_CALL_STATS(PrimitiveType,PendingNumPrimitives);
 
 	GPUProfilingData.RegisterGPUWork(PendingNumPrimitives,PendingNumVertices);
 
@@ -1819,12 +1816,10 @@ void FD3D11DynamicRHI::RHIEndDrawIndexedPrimitiveUP()
 	CommitNonComputeShaderConstants();
 	StateCache.SetStreamSource(VertexBuffer, 0, PendingVertexDataStride, VBOffset);
 	StateCache.SetIndexBuffer(IndexBuffer, PendingIndexDataStride == sizeof(uint16) ? DXGI_FORMAT_R16_UINT : DXGI_FORMAT_R32_UINT, 0);
-	VerifyPrimitiveType(PSOPrimitiveType, PendingPrimitiveType);
-	StateCache.SetPrimitiveTopology(GetD3D11PrimitiveType(PendingPrimitiveType, bUsingTessellation));
+	StateCache.SetPrimitiveTopology(GetD3D11PrimitiveType(PrimitiveType, bUsingTessellation));
 	Direct3DDeviceIMContext->DrawIndexed(PendingNumIndices,0,PendingMinVertexIndex);
 
 	// Clear these parameters.
-	PendingPrimitiveType = 0;
 	PendingNumPrimitives = 0;
 	PendingMinVertexIndex = 0;
 	PendingIndexDataStride = 0;
diff --git a/Engine/Source/Runtime/Windows/D3D11RHI/Private/D3D11Device.cpp b/Engine/Source/Runtime/Windows/D3D11RHI/Private/D3D11Device.cpp
--- a/Engine/Source/Runtime/Windows/D3D11RHI/Private/D3D11Device.cpp
+++ b/Engine/Source/Runtime/Windows/D3D11RHI/Private/D3D11Device.cpp
@@ -49,7 +49,6 @@ FD3D11DynamicRHI::FD3D11DynamicRHI(IDXGIFactory1* InDXGIFactory1,D3D_FEATURE_LEV
 	FeatureLevel(InFeatureLevel),
 	AmdAgsContext(NULL),
 	bCurrentDepthStencilStateIsReadOnly(false),
-	PSOPrimitiveType(PT_Num),
 	CurrentDepthTexture(NULL),
 	NumSimultaneousRenderTargets(0),
 	NumUAVs(0),
@@ -61,7 +60,6 @@ FD3D11DynamicRHI::FD3D11DynamicRHI(IDXGIFactory1* InDXGIFactory1,D3D_FEATURE_LEV
 	bUsingTessellation(false),
 	PendingNumVertices(0),
 	PendingVertexDataStride(0),
-	PendingPrimitiveType(0),
 	PendingNumPrimitives(0),
 	PendingMinVertexIndex(0),
 	PendingNumIndices(0),
diff --git a/Engine/Source/Runtime/Windows/D3D11RHI/Private/D3D11RHI.cpp b/Engine/Source/Runtime/Windows/D3D11RHI/Private/D3D11RHI.cpp
--- a/Engine/Source/Runtime/Windows/D3D11RHI/Private/D3D11RHI.cpp
+++ b/Engine/Source/Runtime/Windows/D3D11RHI/Private/D3D11RHI.cpp
@@ -42,8 +42,6 @@ void FD3D11DynamicRHI::RHIBeginFrame()
 		IntelMetricsDicoveryBeginFrame();
 	}
 #endif // INTEL_METRICSDISCOVERY
-
-	PSOPrimitiveType = PT_Num;
 }
 
 template <int32 Frequency>
diff --git a/Engine/Source/Runtime/Windows/D3D11RHI/Private/D3D11RHIPrivate.h b/Engine/Source/Runtime/Windows/D3D11RHI/Private/D3D11RHIPrivate.h
--- a/Engine/Source/Runtime/Windows/D3D11RHI/Private/D3D11RHIPrivate.h
+++ b/Engine/Source/Runtime/Windows/D3D11RHI/Private/D3D11RHIPrivate.h
@@ -509,7 +509,7 @@ public:
 		FRHIGraphicsPipelineStateFallBack* FallbackGraphicsState = static_cast<FRHIGraphicsPipelineStateFallBack*>(GraphicsState);
 		IRHICommandContext::RHISetGraphicsPipelineState(GraphicsState);
 		// Store the PSO's primitive (after since IRHICommandContext::RHISetGraphicsPipelineState sets the BSS)
-		PSOPrimitiveType = FallbackGraphicsState->Initializer.PrimitiveType;
+		PrimitiveType = FallbackGraphicsState->Initializer.PrimitiveType;
 	}
 
 	virtual void RHISetShaderTexture(FVertexShaderRHIParamRef VertexShader, uint32 TextureIndex, FTextureRHIParamRef NewTexture) final override;
@@ -751,7 +751,7 @@ protected:
 	bool bCurrentDepthStencilStateIsReadOnly;
 
 	// Current PSO Primitive Type
-	EPrimitiveType PSOPrimitiveType;
+	EPrimitiveType PrimitiveType;
 
 	TRefCountPtr<ID3D11RenderTargetView> CurrentRenderTargets[D3D11_SIMULTANEOUS_RENDER_TARGET_COUNT];
 	TRefCountPtr<ID3D11UnorderedAccessView> CurrentUAVs[D3D11_PS_CS_UAV_REGISTER_COUNT];
@@ -802,7 +802,6 @@ protected:
 	// State for begin/end draw primitive UP interface.
 	uint32 PendingNumVertices;
 	uint32 PendingVertexDataStride;
-	uint32 PendingPrimitiveType;
 	uint32 PendingNumPrimitives;
 	uint32 PendingMinVertexIndex;
 	uint32 PendingNumIndices;
