Allow depth textures without SRV on DX11. Unifies DX11/DX12 behavior.

diff --git a/Engine/Source/Runtime/RenderCore/Private/RenderUtils.cpp b/Engine/Source/Runtime/RenderCore/Private/RenderUtils.cpp
--- a/Engine/Source/Runtime/RenderCore/Private/RenderUtils.cpp
+++ b/Engine/Source/Runtime/RenderCore/Private/RenderUtils.cpp
@@ -465,7 +465,7 @@ public:
 	{
 		// Create the texture RHI.
 		FRHIResourceCreateInfo CreateInfo;
-		FTextureCubeRHIRef TextureCube = RHICreateTextureCube(1, PixelFormat, 1, 0, CreateInfo);
+		FTextureCubeRHIRef TextureCube = RHICreateTextureCube(1, PixelFormat, 1, TexCreate_ShaderResource, CreateInfo);
 		TextureRHI = TextureCube;
 
 		// Write the contents of the texture.
diff --git a/Engine/Source/Runtime/SlateRHIRenderer/Private/SlateRHIResourceManager.cpp b/Engine/Source/Runtime/SlateRHIRenderer/Private/SlateRHIResourceManager.cpp
--- a/Engine/Source/Runtime/SlateRHIRenderer/Private/SlateRHIResourceManager.cpp
+++ b/Engine/Source/Runtime/SlateRHIRenderer/Private/SlateRHIResourceManager.cpp
@@ -490,7 +490,7 @@ FSlateShaderResourceProxy* FSlateRHIResourceManager::GenerateTextureResource( co
 		NewProxy = new FSlateShaderResourceProxy;
 
 		// Create a new standalone texture because we can't atlas this one
-		FSlateTexture2DRHIRef* Texture = new FSlateTexture2DRHIRef( Width, Height, PF_B8G8R8A8, Info.TextureData, Info.bSrgb ? TexCreate_SRGB : TexCreate_None );
+		FSlateTexture2DRHIRef* Texture = new FSlateTexture2DRHIRef( Width, Height, PF_B8G8R8A8, Info.TextureData, (Info.bSrgb ? TexCreate_SRGB : TexCreate_None) | TexCreate_ShaderResource );
 		// Add it to the list of non atlased textures that we must clean up later
 		NonAtlasedTextures.Add( Texture );
 
diff --git a/Engine/Source/Runtime/Windows/D3D11RHI/Private/D3D11Texture.cpp b/Engine/Source/Runtime/Windows/D3D11RHI/Private/D3D11Texture.cpp
--- a/Engine/Source/Runtime/Windows/D3D11RHI/Private/D3D11Texture.cpp
+++ b/Engine/Source/Runtime/Windows/D3D11RHI/Private/D3D11Texture.cpp
@@ -545,7 +545,6 @@ TD3D11Texture2D<BaseResourceType>* FD3D11DynamicRHI::CreateD3D11Texture2D(uint32
 	D3D11_SRV_DIMENSION ShaderResourceViewDimension = D3D11_SRV_DIMENSION_TEXTURE2D;
 	uint32 CPUAccessFlags = 0;
 	D3D11_USAGE TextureUsage = D3D11_USAGE_DEFAULT;
-	uint32 BindFlags = D3D11_BIND_SHADER_RESOURCE;
 	bool bCreateShaderResource = true;
 
 	uint32 ActualMSAACount = NumSamples;
@@ -581,7 +580,6 @@ TD3D11Texture2D<BaseResourceType>* FD3D11DynamicRHI::CreateD3D11Texture2D(uint32
 
 		CPUAccessFlags = D3D11_CPU_ACCESS_READ;
 		TextureUsage = D3D11_USAGE_STAGING;
-		BindFlags = 0;
 		bCreateShaderResource = false;
 	}
 
@@ -589,7 +587,6 @@ TD3D11Texture2D<BaseResourceType>* FD3D11DynamicRHI::CreateD3D11Texture2D(uint32
 	{
 		CPUAccessFlags = D3D11_CPU_ACCESS_WRITE;
 		TextureUsage = D3D11_USAGE_STAGING;
-		BindFlags = 0;
 		bCreateShaderResource = false;
 	}
 
@@ -604,7 +601,7 @@ TD3D11Texture2D<BaseResourceType>* FD3D11DynamicRHI::CreateD3D11Texture2D(uint32
 	TextureDesc.SampleDesc.Count = ActualMSAACount;
 	TextureDesc.SampleDesc.Quality = ActualMSAAQuality;
 	TextureDesc.Usage = TextureUsage;
-	TextureDesc.BindFlags = BindFlags;
+	TextureDesc.BindFlags = bCreateShaderResource? D3D11_BIND_SHADER_RESOURCE : 0;
 	TextureDesc.CPUAccessFlags = CPUAccessFlags;
 	TextureDesc.MiscFlags = bCubeTexture ? D3D11_RESOURCE_MISC_TEXTURECUBE : 0;
 
@@ -661,6 +658,12 @@ TD3D11Texture2D<BaseResourceType>* FD3D11DynamicRHI::CreateD3D11Texture2D(uint32
 		bPooledTexture = false;
 	}
 
+	if (bCreateDSV && !(Flags & TexCreate_ShaderResource))
+	{
+		TextureDesc.BindFlags &= ~D3D11_BIND_SHADER_RESOURCE;
+		bCreateShaderResource = false;
+	}
+
 	if (bCreateDSV || bCreateRTV || bCubeTexture || bTextureArray)
 	{
 		bPooledTexture = false;
