Build fix for d3d11, by  forward declaring FD3D11DeviceContext

diff --git a/Engine/Source/Runtime/Windows/D3D11RHI/Public/D3D11Resources.h b/Engine/Source/Runtime/Windows/D3D11RHI/Public/D3D11Resources.h
--- a/Engine/Source/Runtime/Windows/D3D11RHI/Public/D3D11Resources.h
+++ b/Engine/Source/Runtime/Windows/D3D11RHI/Public/D3D11Resources.h
@@ -9,6 +9,9 @@
 #include "BoundShaderStateCache.h"
 #include "D3D11ShaderResources.h"
 
+interface ID3D11DeviceContext;
+typedef ID3D11DeviceContext FD3D11DeviceContext;
+
 template <>
 struct TTypeTraits<D3D11_INPUT_ELEMENT_DESC> : public TTypeTraitsBase<D3D11_INPUT_ELEMENT_DESC>
 {
