Added support in RHICopyTexture to copy multiple mips (implemented only for D3D11)

diff --git a/Engine/Source/Runtime/RHI/Public/RHI.h b/Engine/Source/Runtime/RHI/Public/RHI.h
--- a/Engine/Source/Runtime/RHI/Public/RHI.h
+++ b/Engine/Source/Runtime/RHI/Public/RHI.h
@@ -1271,6 +1271,7 @@ struct FRHICopyTextureInfo
 	// Mips to copy and destination mips
 	uint32 SourceMipIndex = 0;
 	uint32 DestMipIndex = 0;
+	uint32 NumMips = 1;
 };
 
 enum class EResourceTransitionAccess
diff --git a/Engine/Source/Runtime/Renderer/Private/ReflectionEnvironmentCapture.cpp b/Engine/Source/Runtime/Renderer/Private/ReflectionEnvironmentCapture.cpp
--- a/Engine/Source/Runtime/Renderer/Private/ReflectionEnvironmentCapture.cpp
+++ b/Engine/Source/Runtime/Renderer/Private/ReflectionEnvironmentCapture.cpp
@@ -1573,16 +1573,9 @@ void CopyToSkyTexture(FRHICommandList& RHICmdList, FScene* Scene, FTexture* Proc
 		FRHICopyTextureInfo CopyInfo;
 		CopyInfo.Size = FilteredCube.ShaderResourceTexture->GetSizeXYZ();
 		CopyInfo.NumSlices = 6;
+		CopyInfo.NumMips = NumMips;
 
-		for (int32 MipIndex = 0; MipIndex < NumMips; MipIndex++)
-		{
-			RHICmdList.CopyTexture(FilteredCube.ShaderResourceTexture, ProcessedTexture->TextureRHI, CopyInfo);
-
-			++CopyInfo.SourceMipIndex;
-			++CopyInfo.DestMipIndex;
-			CopyInfo.Size.X = FMath::Max(1, CopyInfo.Size.X / 2);
-			CopyInfo.Size.Y = FMath::Max(1, CopyInfo.Size.Y / 2);
-		}
+		RHICmdList.CopyTexture(FilteredCube.ShaderResourceTexture, ProcessedTexture->TextureRHI, CopyInfo);
 
 		RHICmdList.TransitionResource(EResourceTransitionAccess::EReadable, ProcessedTexture->TextureRHI);
 	}
diff --git a/Engine/Source/Runtime/Windows/D3D11RHI/Private/D3D11Texture.cpp b/Engine/Source/Runtime/Windows/D3D11RHI/Private/D3D11Texture.cpp
--- a/Engine/Source/Runtime/Windows/D3D11RHI/Private/D3D11Texture.cpp
+++ b/Engine/Source/Runtime/Windows/D3D11RHI/Private/D3D11Texture.cpp
@@ -2087,19 +2087,26 @@ void FD3D11DynamicRHI::RHICopyTexture(FTextureRHIParamRef SourceTextureRHI, FTex
 		SrcBox.left = CopyInfo.SourcePosition.X;
 		SrcBox.top = CopyInfo.SourcePosition.Y;
 		SrcBox.front = CopyInfo.SourcePosition.Z;
-		SrcBox.right = CopyInfo.SourcePosition.X + CopyInfo.Size.X;
-		SrcBox.bottom = CopyInfo.SourcePosition.Y + CopyInfo.Size.Y;
-		SrcBox.back = CopyInfo.SourcePosition.Z + CopyInfo.Size.Z;
 
-		for (uint32 i = 0; i < CopyInfo.NumSlices; ++i)
+		for (uint32 SliceIndex = 0; SliceIndex < CopyInfo.NumSlices; ++SliceIndex)
 		{
-			uint32 SourceSliceIndex = CopyInfo.SourceSliceIndex + i;
-			uint32 DestSliceIndex = CopyInfo.DestSliceIndex + i;
+			uint32 SourceSliceIndex = CopyInfo.SourceSliceIndex + SliceIndex;
+			uint32 DestSliceIndex = CopyInfo.DestSliceIndex + SliceIndex;
 
-			const uint32 SourceSubresource = D3D11CalcSubresource(CopyInfo.SourceMipIndex, SourceSliceIndex, SourceTextureRHI->GetNumMips());
-			const uint32 DestSubresource = D3D11CalcSubresource(CopyInfo.DestMipIndex, DestSliceIndex, DestTextureRHI->GetNumMips());
+			for (uint32 MipIndex = 0; MipIndex < CopyInfo.NumMips; ++MipIndex)
+			{
+				uint32 SourceMipIndex = CopyInfo.SourceMipIndex + MipIndex;
+				uint32 DestMipIndex = CopyInfo.DestMipIndex + MipIndex;
+
+				const uint32 SourceSubresource = D3D11CalcSubresource(SourceMipIndex, SourceSliceIndex, SourceTextureRHI->GetNumMips());
+				const uint32 DestSubresource = D3D11CalcSubresource(DestMipIndex, DestSliceIndex, DestTextureRHI->GetNumMips());
 
-			Direct3DDeviceIMContext->CopySubresourceRegion(DestTexture->GetResource(), DestSubresource, CopyInfo.DestPosition.X, CopyInfo.DestPosition.Y, CopyInfo.DestPosition.Z, SourceTexture->GetResource(), SourceSubresource, &SrcBox);
+				SrcBox.right = CopyInfo.SourcePosition.X + FMath::Max(CopyInfo.Size.X >> MipIndex, 1);
+				SrcBox.bottom = CopyInfo.SourcePosition.Y + FMath::Max(CopyInfo.Size.Y >> MipIndex, 1);
+				SrcBox.back = CopyInfo.SourcePosition.Z + FMath::Max(CopyInfo.Size.Z >> MipIndex, 1);
+
+				Direct3DDeviceIMContext->CopySubresourceRegion(DestTexture->GetResource(), DestSubresource, CopyInfo.DestPosition.X, CopyInfo.DestPosition.Y, CopyInfo.DestPosition.Z, SourceTexture->GetResource(), SourceSubresource, &SrcBox);
+			}
 		}
 	}
 	else
