Fixing audio mixer crash when swapping to null audio device.

diff --git a/Engine/Source/Runtime/Windows/AudioMixerXAudio2/Private/AudioMixerPlatformXAudio2.cpp b/Engine/Source/Runtime/Windows/AudioMixerXAudio2/Private/AudioMixerPlatformXAudio2.cpp
--- a/Engine/Source/Runtime/Windows/AudioMixerXAudio2/Private/AudioMixerPlatformXAudio2.cpp
+++ b/Engine/Source/Runtime/Windows/AudioMixerXAudio2/Private/AudioMixerPlatformXAudio2.cpp
@@ -487,10 +487,10 @@ namespace Audio
 		{
 			check(!bIsUsingNullDevice);
 
-			AudioStreamInfo.NumOutputFrames = 512;
+			AudioStreamInfo.NumOutputFrames = OpenStreamParams.NumFrames;
 			AudioStreamInfo.DeviceInfo.OutputChannelArray = { EAudioMixerChannel::FrontLeft, EAudioMixerChannel::FrontRight };
 			AudioStreamInfo.DeviceInfo.NumChannels = 2;
-			AudioStreamInfo.DeviceInfo.SampleRate = 48000;
+			AudioStreamInfo.DeviceInfo.SampleRate = OpenStreamParams.SampleRate;
 			AudioStreamInfo.DeviceInfo.Format = EAudioMixerStreamDataFormat::Float;
 		}
 
