- When UAV is bound in D3D11 RHI, explicitly unbind any VB/IB that reference the same resource
- D3D11 will implicitly unbind these anyway, we need to track this or else state cache will get out of sync

diff --git a/Engine/Source/Runtime/Windows/D3D11RHI/Private/D3D11Commands.cpp b/Engine/Source/Runtime/Windows/D3D11RHI/Private/D3D11Commands.cpp
--- a/Engine/Source/Runtime/Windows/D3D11RHI/Private/D3D11Commands.cpp
+++ b/Engine/Source/Runtime/Windows/D3D11RHI/Private/D3D11Commands.cpp
@@ -159,6 +159,7 @@ void FD3D11DynamicRHI::RHISetStreamSource(uint32 StreamIndex, FVertexBufferRHIPa
 	FD3D11VertexBuffer* VertexBuffer = ResourceCast(VertexBufferRHI);
 
 	ID3D11Buffer* D3DBuffer = VertexBuffer ? VertexBuffer->Resource : NULL;
+	TrackResourceBoundAsVB(VertexBuffer, StreamIndex);
 	StateCache.SetStreamSource(D3DBuffer, StreamIndex, Offset);
 }
 
@@ -442,7 +443,7 @@ void FD3D11DynamicRHI::RHISetUAVParameter(FComputeShaderRHIParamRef ComputeShade
 
 	if(UAV)
 	{
-		ConditionalClearShaderResource(UAV->Resource);		
+		ConditionalClearShaderResource(UAV->Resource, true);		
 
 		//check it's safe for r/w for this UAV
 		const EResourceTransitionAccess CurrentUAVAccess = UAV->Resource->GetCurrentGPUAccess();
@@ -467,7 +468,7 @@ void FD3D11DynamicRHI::RHISetUAVParameter(FComputeShaderRHIParamRef ComputeShade
 	
 	if(UAV)
 	{
-		ConditionalClearShaderResource(UAV->Resource);
+		ConditionalClearShaderResource(UAV->Resource, true);
 
 		//check it's safe for r/w for this UAV
 		const EResourceTransitionAccess CurrentUAVAccess = UAV->Resource->GetCurrentGPUAccess();
@@ -969,7 +970,7 @@ void FD3D11DynamicRHI::RHISetRenderTargets(
 		DepthStencilView = NewDepthStencilTarget->GetDepthStencilView(CurrentDSVAccessType);
 
 		// Unbind any shader views of the depth stencil target that are bound.
-		ConditionalClearShaderResource(NewDepthStencilTarget);
+		ConditionalClearShaderResource(NewDepthStencilTarget, false);
 	}
 
 	// Check if the depth stencil target is different from the old state.
@@ -1050,7 +1051,7 @@ void FD3D11DynamicRHI::RHISetRenderTargets(
 #endif
 			
 			// Unbind any shader views of the render target that are bound.
-			ConditionalClearShaderResource(NewRenderTarget);
+			ConditionalClearShaderResource(NewRenderTarget, false);
 
 #if UE_BUILD_DEBUG	
 			// A check to allow you to pinpoint what is using mismatching targets
@@ -1116,7 +1117,7 @@ void FD3D11DynamicRHI::RHISetRenderTargets(
 			}
 
 			// Unbind any shader views of the UAV's resource.
-			ConditionalClearShaderResource(RHIUAV->Resource);
+			ConditionalClearShaderResource(RHIUAV->Resource, true);
 		}
 
 		if(CurrentUAVs[UAVIndex] != UAV)
@@ -1630,6 +1631,7 @@ void FD3D11DynamicRHI::RHIDrawIndexedIndirect(FIndexBufferRHIParamRef IndexBuffe
 	uint32 SizeFormat = sizeof(DXGI_FORMAT);
 	const DXGI_FORMAT Format = (IndexBuffer->GetStride() == sizeof(uint16) ? DXGI_FORMAT_R16_UINT : DXGI_FORMAT_R32_UINT);
 
+	TrackResourceBoundAsIB(IndexBuffer);
 	StateCache.SetIndexBuffer(IndexBuffer->Resource, Format, 0);
 	StateCache.SetPrimitiveTopology(GetD3D11PrimitiveType(PrimitiveType,bUsingTessellation));
 
@@ -1668,6 +1670,7 @@ void FD3D11DynamicRHI::RHIDrawIndexedPrimitive(FIndexBufferRHIParamRef IndexBuff
 	checkf((StartIndex + IndexCount) * IndexBuffer->GetStride() <= IndexBuffer->GetSize(), 		
 		TEXT("Start %u, Count %u, Type %u, Buffer Size %u, Buffer stride %u"), StartIndex, IndexCount, PrimitiveType, IndexBuffer->GetSize(), IndexBuffer->GetStride());
 
+	TrackResourceBoundAsIB(IndexBuffer);
 	StateCache.SetIndexBuffer(IndexBuffer->Resource, Format, 0);
 	StateCache.SetPrimitiveTopology(GetD3D11PrimitiveType(PrimitiveType,bUsingTessellation));
 
@@ -1698,6 +1701,7 @@ void FD3D11DynamicRHI::RHIDrawIndexedPrimitiveIndirect(FIndexBufferRHIParamRef I
 	// Set the index buffer.
 	const uint32 SizeFormat = sizeof(DXGI_FORMAT);
 	const DXGI_FORMAT Format = (IndexBuffer->GetStride() == sizeof(uint16) ? DXGI_FORMAT_R16_UINT : DXGI_FORMAT_R32_UINT);
+	TrackResourceBoundAsIB(IndexBuffer);
 	StateCache.SetIndexBuffer(IndexBuffer->Resource, Format, 0);
 	StateCache.SetPrimitiveTopology(GetD3D11PrimitiveType(PrimitiveType,bUsingTessellation));
 	Direct3DDeviceIMContext->DrawIndexedInstancedIndirect(ArgumentBuffer->Resource,ArgumentOffset);
@@ -1742,6 +1746,7 @@ void FD3D11DynamicRHI::RHIEndDrawPrimitiveUP()
 	// Issue the draw call.
 	CommitGraphicsResourceTables();
 	CommitNonComputeShaderConstants();
+	TrackResourceBoundAsVB(nullptr, 0);
 	StateCache.SetStreamSource(D3DBuffer, 0, PendingVertexDataStride, VBOffset);
 	StateCache.SetPrimitiveTopology(GetD3D11PrimitiveType(PrimitiveType,bUsingTessellation));
 	Direct3DDeviceIMContext->Draw(PendingNumVertices,0);
@@ -1801,6 +1806,8 @@ void FD3D11DynamicRHI::RHIEndDrawIndexedPrimitiveUP()
 	// Issue the draw call.
 	CommitGraphicsResourceTables();
 	CommitNonComputeShaderConstants();
+	TrackResourceBoundAsVB(nullptr, 0);
+	TrackResourceBoundAsIB(nullptr);
 	StateCache.SetStreamSource(VertexBuffer, 0, PendingVertexDataStride, VBOffset);
 	StateCache.SetIndexBuffer(IndexBuffer, PendingIndexDataStride == sizeof(uint16) ? DXGI_FORMAT_R16_UINT : DXGI_FORMAT_R32_UINT, 0);
 	StateCache.SetPrimitiveTopology(GetD3D11PrimitiveType(PrimitiveType, bUsingTessellation));
diff --git a/Engine/Source/Runtime/Windows/D3D11RHI/Private/D3D11IndexBuffer.cpp b/Engine/Source/Runtime/Windows/D3D11RHI/Private/D3D11IndexBuffer.cpp
--- a/Engine/Source/Runtime/Windows/D3D11RHI/Private/D3D11IndexBuffer.cpp
+++ b/Engine/Source/Runtime/Windows/D3D11RHI/Private/D3D11IndexBuffer.cpp
@@ -66,7 +66,7 @@ void* FD3D11DynamicRHI::RHILockIndexBuffer(FIndexBufferRHIParamRef IndexBufferRH
 	FD3D11IndexBuffer* IndexBuffer = ResourceCast(IndexBufferRHI);
 	
 	// If this resource is bound to the device, unbind it
-	ConditionalClearShaderResource(IndexBuffer);
+	ConditionalClearShaderResource(IndexBuffer, true);
 
 	// Determine whether the index buffer is dynamic or not.
 	D3D11_BUFFER_DESC Desc;
diff --git a/Engine/Source/Runtime/Windows/D3D11RHI/Private/D3D11RHI.cpp b/Engine/Source/Runtime/Windows/D3D11RHI/Private/D3D11RHI.cpp
--- a/Engine/Source/Runtime/Windows/D3D11RHI/Private/D3D11RHI.cpp
+++ b/Engine/Source/Runtime/Windows/D3D11RHI/Private/D3D11RHI.cpp
@@ -64,10 +64,13 @@ void FD3D11DynamicRHI::ClearState()
 	StateCache.ClearState();
 
 	FMemory::Memzero(CurrentResourcesBoundAsSRVs, sizeof(CurrentResourcesBoundAsSRVs));
+	FMemory::Memzero(CurrentResourcesBoundAsVBs, sizeof(CurrentResourcesBoundAsVBs));
+	CurrentResourceBoundAsIB = nullptr;
 	for (int32 Frequency = 0; Frequency < SF_NumFrequencies; Frequency++)
 	{
 		MaxBoundShaderResourcesIndex[Frequency] = INDEX_NONE;
 	}
+	MaxBoundVertexBufferIndex = INDEX_NONE;
 }
 
 void GetMipAndSliceInfoFromSRV(ID3D11ShaderResourceView* SRV, int32& MipLevel, int32& NumMips, int32& ArraySlice, int32& NumSlices)
@@ -236,6 +239,39 @@ void FD3D11DynamicRHI::InternalSetShaderResourceView(FD3D11BaseShaderResource* R
 	StateCache.SetShaderResourceView<ShaderFrequency>(SRV, ResourceIndex, SrvType);
 }
 
+void FD3D11DynamicRHI::TrackResourceBoundAsVB(FD3D11BaseShaderResource* Resource, int32 StreamIndex)
+{
+	check(StreamIndex >= 0 && StreamIndex < D3D11_IA_VERTEX_INPUT_RESOURCE_SLOT_COUNT);
+	if (Resource)
+	{
+		// We are binding a new VB.
+		// Update the max resource index to the highest bound resource index.
+		MaxBoundVertexBufferIndex = FMath::Max(MaxBoundVertexBufferIndex, StreamIndex);
+		CurrentResourcesBoundAsVBs[StreamIndex] = Resource;
+	}
+	else if (CurrentResourcesBoundAsVBs[StreamIndex] != nullptr)
+	{
+		// Unbind the resource from the slot.
+		CurrentResourcesBoundAsVBs[StreamIndex] = nullptr;
+
+		// If this was the highest bound resource...
+		if (MaxBoundVertexBufferIndex == StreamIndex)
+		{
+			// Adjust the max resource index downwards until we
+			// hit the next non-null slot, or we've run out of slots.
+			do
+			{
+				MaxBoundVertexBufferIndex--;
+			} while (MaxBoundVertexBufferIndex >= 0 && CurrentResourcesBoundAsVBs[MaxBoundVertexBufferIndex] == nullptr);
+		}
+	}
+}
+
+void FD3D11DynamicRHI::TrackResourceBoundAsIB(FD3D11BaseShaderResource* Resource)
+{
+	CurrentResourceBoundAsIB = Resource;
+}
+
 template <EShaderFrequency ShaderFrequency>
 void FD3D11DynamicRHI::ClearShaderResourceViews(FD3D11BaseShaderResource* Resource)
 {
@@ -250,7 +286,7 @@ void FD3D11DynamicRHI::ClearShaderResourceViews(FD3D11BaseShaderResource* Resour
 	}
 }
 
-void FD3D11DynamicRHI::ConditionalClearShaderResource(FD3D11BaseShaderResource* Resource)
+void FD3D11DynamicRHI::ConditionalClearShaderResource(FD3D11BaseShaderResource* Resource, bool bCheckBoundInputAssembler)
 {
 	SCOPE_CYCLE_COUNTER(STAT_D3D11ClearShaderResourceTime);
 	check(Resource);
@@ -260,6 +296,25 @@ void FD3D11DynamicRHI::ConditionalClearShaderResource(FD3D11BaseShaderResource*
 	ClearShaderResourceViews<SF_Pixel>(Resource);
 	ClearShaderResourceViews<SF_Geometry>(Resource);
 	ClearShaderResourceViews<SF_Compute>(Resource);
+
+	if (bCheckBoundInputAssembler)
+	{
+		for (int32 ResourceIndex = MaxBoundVertexBufferIndex; ResourceIndex >= 0; --ResourceIndex)
+		{
+			if (CurrentResourcesBoundAsVBs[ResourceIndex] == Resource)
+			{
+				// Unset the vertex buffer from the device context
+				TrackResourceBoundAsVB(nullptr, ResourceIndex);
+				StateCache.SetStreamSource(nullptr, ResourceIndex, 0);
+			}
+		}
+
+		if (Resource == CurrentResourceBoundAsIB)
+		{
+			TrackResourceBoundAsIB(nullptr);
+			StateCache.SetIndexBuffer(nullptr, DXGI_FORMAT_R16_UINT, 0);
+		}
+	}
 }
 
 template <EShaderFrequency ShaderFrequency>
diff --git a/Engine/Source/Runtime/Windows/D3D11RHI/Private/D3D11RHIPrivate.h b/Engine/Source/Runtime/Windows/D3D11RHI/Private/D3D11RHIPrivate.h
--- a/Engine/Source/Runtime/Windows/D3D11RHI/Private/D3D11RHIPrivate.h
+++ b/Engine/Source/Runtime/Windows/D3D11RHI/Private/D3D11RHIPrivate.h
@@ -665,6 +665,9 @@ private:
 	template <EShaderFrequency ShaderFrequency>
 	void InternalSetShaderResourceView(FD3D11BaseShaderResource* Resource, ID3D11ShaderResourceView* SRV, int32 ResourceIndex, FName SRVName, FD3D11StateCache::ESRV_Type SrvType = FD3D11StateCache::SRV_Unknown);
 
+	void TrackResourceBoundAsVB(FD3D11BaseShaderResource* Resource, int32 StreamIndex);
+	void TrackResourceBoundAsIB(FD3D11BaseShaderResource* Resource);
+
 	void SetCurrentComputeShader(FComputeShaderRHIParamRef ComputeShader)
 	{
 		CurrentComputeShader = ComputeShader;
@@ -684,7 +687,7 @@ public:
 	}
 
 	void ClearState();
-	void ConditionalClearShaderResource(FD3D11BaseShaderResource* Resource);
+	void ConditionalClearShaderResource(FD3D11BaseShaderResource* Resource, bool bCheckBoundInputAssembler);
 	void ClearAllShaderResources();
 
 
@@ -764,7 +767,10 @@ protected:
 	TRefCountPtr<ID3D11DepthStencilView> CurrentDepthStencilTarget;
 	TRefCountPtr<FD3D11TextureBase> CurrentDepthTexture;
 	FD3D11BaseShaderResource* CurrentResourcesBoundAsSRVs[SF_NumFrequencies][D3D11_COMMONSHADER_INPUT_RESOURCE_SLOT_COUNT];
+	FD3D11BaseShaderResource* CurrentResourcesBoundAsVBs[D3D11_IA_VERTEX_INPUT_RESOURCE_SLOT_COUNT];
+	FD3D11BaseShaderResource* CurrentResourceBoundAsIB;
 	int32 MaxBoundShaderResourcesIndex[SF_NumFrequencies];
+	int32 MaxBoundVertexBufferIndex;
 	uint32 NumSimultaneousRenderTargets;
 	uint32 NumUAVs;
 
diff --git a/Engine/Source/Runtime/Windows/D3D11RHI/Private/D3D11RenderTarget.cpp b/Engine/Source/Runtime/Windows/D3D11RHI/Private/D3D11RenderTarget.cpp
--- a/Engine/Source/Runtime/Windows/D3D11RHI/Private/D3D11RenderTarget.cpp
+++ b/Engine/Source/Runtime/Windows/D3D11RHI/Private/D3D11RenderTarget.cpp
@@ -71,7 +71,7 @@ void FD3D11DynamicRHI::ResolveTextureUsingShader(
 	// Make sure the destination is not bound as a shader resource.
 	if (DestTexture)
 	{
-		ConditionalClearShaderResource(DestTexture);
+		ConditionalClearShaderResource(DestTexture, false);
 	}
 
 	// Determine if the entire destination surface is being resolved to.
@@ -196,7 +196,7 @@ void FD3D11DynamicRHI::ResolveTextureUsingShader(
 
 	if (SourceTexture)
 	{
-		ConditionalClearShaderResource(SourceTexture);
+		ConditionalClearShaderResource(SourceTexture, false);
 	}
 
 	// Reset saved render targets
diff --git a/Engine/Source/Runtime/Windows/D3D11RHI/Private/D3D11StructuredBuffer.cpp b/Engine/Source/Runtime/Windows/D3D11RHI/Private/D3D11StructuredBuffer.cpp
--- a/Engine/Source/Runtime/Windows/D3D11RHI/Private/D3D11StructuredBuffer.cpp
+++ b/Engine/Source/Runtime/Windows/D3D11RHI/Private/D3D11StructuredBuffer.cpp
@@ -92,7 +92,7 @@ void* FD3D11DynamicRHI::RHILockStructuredBuffer(FStructuredBufferRHIParamRef Str
 	FD3D11StructuredBuffer* StructuredBuffer = ResourceCast(StructuredBufferRHI);
 	
 	// If this resource is bound to the device, unbind it
-	ConditionalClearShaderResource(StructuredBuffer);
+	ConditionalClearShaderResource(StructuredBuffer, true);
 
 	// Determine whether the Structured buffer is dynamic or not.
 	D3D11_BUFFER_DESC Desc;
diff --git a/Engine/Source/Runtime/Windows/D3D11RHI/Private/D3D11Texture.cpp b/Engine/Source/Runtime/Windows/D3D11RHI/Private/D3D11Texture.cpp
--- a/Engine/Source/Runtime/Windows/D3D11RHI/Private/D3D11Texture.cpp
+++ b/Engine/Source/Runtime/Windows/D3D11RHI/Private/D3D11Texture.cpp
@@ -1620,7 +1620,7 @@ void* FD3D11DynamicRHI::RHILockTexture2D(FTexture2DRHIParamRef TextureRHI,uint32
 {
 	check(TextureRHI);
 	FD3D11Texture2D* Texture = ResourceCast(TextureRHI);
-	ConditionalClearShaderResource(Texture);
+	ConditionalClearShaderResource(Texture, false);
 	return Texture->Lock(MipIndex,0,LockMode,DestStride);
 }
 
@@ -1634,7 +1634,7 @@ void FD3D11DynamicRHI::RHIUnlockTexture2D(FTexture2DRHIParamRef TextureRHI,uint3
 void* FD3D11DynamicRHI::RHILockTexture2DArray(FTexture2DArrayRHIParamRef TextureRHI,uint32 TextureIndex,uint32 MipIndex,EResourceLockMode LockMode,uint32& DestStride,bool bLockWithinMiptail)
 {
 	FD3D11Texture2DArray* Texture = ResourceCast(TextureRHI);
-	ConditionalClearShaderResource(Texture);
+	ConditionalClearShaderResource(Texture, false);
 	return Texture->Lock(MipIndex,TextureIndex,LockMode,DestStride);
 }
 
@@ -1699,7 +1699,7 @@ FTextureCubeRHIRef FD3D11DynamicRHI::RHICreateTextureCubeArray(uint32 Size, uint
 void* FD3D11DynamicRHI::RHILockTextureCubeFace(FTextureCubeRHIParamRef TextureCubeRHI,uint32 FaceIndex,uint32 ArrayIndex,uint32 MipIndex,EResourceLockMode LockMode,uint32& DestStride,bool bLockWithinMiptail)
 {
 	FD3D11TextureCube* TextureCube = ResourceCast(TextureCubeRHI);
-	ConditionalClearShaderResource(TextureCube);
+	ConditionalClearShaderResource(TextureCube, false);
 	uint32 D3DFace = GetD3D11CubeFace((ECubeFace)FaceIndex);
 	return TextureCube->Lock(MipIndex,D3DFace + ArrayIndex * 6,LockMode,DestStride);
 }
diff --git a/Engine/Source/Runtime/Windows/D3D11RHI/Private/D3D11VertexBuffer.cpp b/Engine/Source/Runtime/Windows/D3D11RHI/Private/D3D11VertexBuffer.cpp
--- a/Engine/Source/Runtime/Windows/D3D11RHI/Private/D3D11VertexBuffer.cpp
+++ b/Engine/Source/Runtime/Windows/D3D11RHI/Private/D3D11VertexBuffer.cpp
@@ -101,7 +101,7 @@ void* FD3D11DynamicRHI::RHILockVertexBuffer(FVertexBufferRHIParamRef VertexBuffe
 	FD3D11VertexBuffer* VertexBuffer = ResourceCast(VertexBufferRHI);
 	
 	// If this resource is bound to the device, unbind it
-	ConditionalClearShaderResource(VertexBuffer);
+	ConditionalClearShaderResource(VertexBuffer, true);
 
 	// Determine whether the vertex buffer is dynamic or not.
 	D3D11_BUFFER_DESC Desc;
diff --git a/Engine/Source/Runtime/Windows/D3D11RHI/Private/D3D11Viewport.cpp b/Engine/Source/Runtime/Windows/D3D11RHI/Private/D3D11Viewport.cpp
--- a/Engine/Source/Runtime/Windows/D3D11RHI/Private/D3D11Viewport.cpp
+++ b/Engine/Source/Runtime/Windows/D3D11RHI/Private/D3D11Viewport.cpp
@@ -649,6 +649,11 @@ void FD3D11DynamicRHI::RHIEndDrawingViewport(FViewportRHIParamRef ViewportRHI,bo
 	}
 
 	StateCache.SetIndexBuffer(nullptr, DXGI_FORMAT_R16_UINT, 0);
+	
+	CurrentResourceBoundAsIB = nullptr;
+	FMemory::Memzero(CurrentResourcesBoundAsVBs, sizeof(CurrentResourcesBoundAsVBs));
+	MaxBoundVertexBufferIndex = INDEX_NONE;
+	
 	StateCache.SetPixelShader(nullptr);
 	StateCache.SetHullShader(nullptr);
 	StateCache.SetDomainShader(nullptr);
