Disable GRHINeedsExtraDeletionLatency for DX11.  There's no evidence it ever helped GPU hangs, and often hides bugs that other platforms find.  It should not be necessary on any RHI ideally

diff --git a/Engine/Source/Runtime/Windows/D3D11RHI/Private/Windows/WindowsD3D11Device.cpp b/Engine/Source/Runtime/Windows/D3D11RHI/Private/Windows/WindowsD3D11Device.cpp
--- a/Engine/Source/Runtime/Windows/D3D11RHI/Private/Windows/WindowsD3D11Device.cpp
+++ b/Engine/Source/Runtime/Windows/D3D11RHI/Private/Windows/WindowsD3D11Device.cpp
@@ -1657,7 +1657,7 @@ void FD3D11DynamicRHI::InitD3DDevice()
 
 		GRHISupportsTextureStreaming = true;
 		GRHISupportsFirstInstance = true;
-		GRHINeedsExtraDeletionLatency = true;
+		GRHINeedsExtraDeletionLatency = false;
 		// Set the RHI initialized flag.
 		GIsRHIInitialized = true;
 	}
