DR - Show warning when running under RenderDoc

diff --git a/Engine/Source/Runtime/Windows/D3D11RHI/Private/D3D11Commands.cpp b/Engine/Source/Runtime/Windows/D3D11RHI/Private/D3D11Commands.cpp
--- a/Engine/Source/Runtime/Windows/D3D11RHI/Private/D3D11Commands.cpp
+++ b/Engine/Source/Runtime/Windows/D3D11RHI/Private/D3D11Commands.cpp
@@ -1926,27 +1926,41 @@ void FD3D11DynamicRHI::EnableDepthBoundsTest(bool bEnable,float MinDepth,float M
 
 	if (IsRHIDeviceNVIDIA())
 	{
-		auto result = NvAPI_D3D11_SetDepthBoundsTest( Direct3DDevice, bEnable, MinDepth, MaxDepth );
-		if(result != NVAPI_OK)
+		auto Result = NvAPI_D3D11_SetDepthBoundsTest( Direct3DDevice, bEnable, MinDepth, MaxDepth );
+		if (Result != NVAPI_OK)
 		{
 			static bool bOnce = false;
 			if (!bOnce)
 			{
 				bOnce = true;
-				UE_LOG(LogD3D11RHI, Error,TEXT("NvAPI_D3D11_SetDepthBoundsTest(%i,%f, %f) returned error code %i. **********PLEASE UPDATE YOUR VIDEO DRIVERS*********"),bEnable,MinDepth,MaxDepth,(unsigned int)result);
+				if (bRenderDoc)
+				{
+					UE_LOG(LogD3D11RHI, Error, TEXT("NvAPI is not available under RenderDoc"));
+				}
+				else
+				{
+					UE_LOG(LogD3D11RHI, Error, TEXT("NvAPI_D3D11_SetDepthBoundsTest(%i,%f, %f) returned error code %i. **********PLEASE UPDATE YOUR VIDEO DRIVERS*********"), bEnable, MinDepth, MaxDepth, (unsigned int)Result);
+				}
 			}
 		}
 	}
 	else if (IsRHIDeviceAMD())
 	{
-		auto result = agsDriverExtensionsDX11_SetDepthBounds( AmdAgsContext, bEnable, MinDepth, MaxDepth );
-		if(result != AGS_SUCCESS)
+		auto Result = agsDriverExtensionsDX11_SetDepthBounds( AmdAgsContext, bEnable, MinDepth, MaxDepth );
+		if(Result != AGS_SUCCESS)
 		{
 			static bool bOnce = false;
 			if (!bOnce)
 			{
 				bOnce = true;
-				UE_LOG(LogD3D11RHI, Error,TEXT("agsDriverExtensionsDX11_SetDepthBounds(%i,%f, %f) returned error code %i. **********PLEASE UPDATE YOUR VIDEO DRIVERS*********"),bEnable,MinDepth,MaxDepth,(unsigned int)result);
+				if (bRenderDoc)
+				{
+					UE_LOG(LogD3D11RHI, Error, TEXT("AGS is not available under RenderDoc"));
+				}
+				else
+				{
+					UE_LOG(LogD3D11RHI, Error, TEXT("agsDriverExtensionsDX11_SetDepthBounds(%i,%f, %f) returned error code %i. **********PLEASE UPDATE YOUR VIDEO DRIVERS*********"), bEnable, MinDepth, MaxDepth, (unsigned int)Result);
+				}
 			}
 		}
 	}
diff --git a/Engine/Source/Runtime/Windows/D3D11RHI/Private/D3D11Device.cpp b/Engine/Source/Runtime/Windows/D3D11RHI/Private/D3D11Device.cpp
--- a/Engine/Source/Runtime/Windows/D3D11RHI/Private/D3D11Device.cpp
+++ b/Engine/Source/Runtime/Windows/D3D11RHI/Private/D3D11Device.cpp
@@ -410,6 +410,8 @@ void FD3D11DynamicRHI::SetupAfterDeviceCreation()
 	{
 		if (SUCCEEDED(Direct3DDevice->QueryInterface(RenderDocID, (void**)(&RenderDoc))))
 		{
+			bRenderDoc = true;
+
 			// Running under RenderDoc, so enable capturing mode
 			GDynamicRHI->EnableIdealGPUCaptureOptions(true);
 		}
diff --git a/Engine/Source/Runtime/Windows/D3D11RHI/Private/D3D11RHIPrivate.h b/Engine/Source/Runtime/Windows/D3D11RHI/Private/D3D11RHIPrivate.h
--- a/Engine/Source/Runtime/Windows/D3D11RHI/Private/D3D11RHIPrivate.h
+++ b/Engine/Source/Runtime/Windows/D3D11RHI/Private/D3D11RHIPrivate.h
@@ -836,6 +836,7 @@ protected:
 	uint32 HDRDetectedDisplayIndex;
 	uint32 HDRDetectedDisplayIHVIndex;
 
+	bool bRenderDoc = false;
 #if CHECK_SRV_TRANSITIONS
 	/*
 	 * Rendertargets must be explicitly 'resolved' to manage their transition to an SRV on some platforms and DX12
