DR - Disable depth bounds if RenderDoc is enabled

diff --git a/Engine/Source/Runtime/Windows/D3D11RHI/Private/Windows/WindowsD3D11Device.cpp b/Engine/Source/Runtime/Windows/D3D11RHI/Private/Windows/WindowsD3D11Device.cpp
--- a/Engine/Source/Runtime/Windows/D3D11RHI/Private/Windows/WindowsD3D11Device.cpp
+++ b/Engine/Source/Runtime/Windows/D3D11RHI/Private/Windows/WindowsD3D11Device.cpp
@@ -1520,33 +1520,36 @@ void FD3D11DynamicRHI::InitD3DDevice()
 
 		GNumAlternateFrameRenderingGroups = 1;
 
-		if (IsRHIDeviceNVIDIA())
+		if (!bRenderDoc)
 		{
-			GSupportsDepthBoundsTest = true;
-			NV_GET_CURRENT_SLI_STATE SLICaps;
-			FMemory::Memzero(SLICaps);
-			SLICaps.version = NV_GET_CURRENT_SLI_STATE_VER;
-			NvAPI_Status SLIStatus = NvAPI_D3D_GetCurrentSLIState(Direct3DDevice, &SLICaps);
-			if (SLIStatus == NVAPI_OK)
+			if (IsRHIDeviceNVIDIA())
 			{
-				if (SLICaps.numAFRGroups > 1)
+				GSupportsDepthBoundsTest = true;
+				NV_GET_CURRENT_SLI_STATE SLICaps;
+				FMemory::Memzero(SLICaps);
+				SLICaps.version = NV_GET_CURRENT_SLI_STATE_VER;
+				NvAPI_Status SLIStatus = NvAPI_D3D_GetCurrentSLIState(Direct3DDevice, &SLICaps);
+				if (SLIStatus == NVAPI_OK)
 				{
-					GNumAlternateFrameRenderingGroups = SLICaps.numAFRGroups;
-					UE_LOG(LogD3D11RHI, Log, TEXT("Detected %i SLI GPUs Setting GNumAlternateFrameRenderingGroups to: %i."), SLICaps.numAFRGroups, GNumAlternateFrameRenderingGroups);
+					if (SLICaps.numAFRGroups > 1)
+					{
+						GNumAlternateFrameRenderingGroups = SLICaps.numAFRGroups;
+						UE_LOG(LogD3D11RHI, Log, TEXT("Detected %i SLI GPUs Setting GNumAlternateFrameRenderingGroups to: %i."), SLICaps.numAFRGroups, GNumAlternateFrameRenderingGroups);
+					}
+				}
+				else
+				{
+					UE_LOG(LogD3D11RHI, Log, TEXT("NvAPI_D3D_GetCurrentSLIState failed: 0x%x"), (int32)SLIStatus);
 				}
-			}
-			else
-			{
-				UE_LOG(LogD3D11RHI, Log, TEXT("NvAPI_D3D_GetCurrentSLIState failed: 0x%x"), (int32)SLIStatus);
-			}
 
-			START_NV_AFTERMATH();
-		}
-		else if (IsRHIDeviceAMD() && AmdAgsContext)
-		{
-			if ((AmdSupportedExtensionFlags & AGS_DX11_EXTENSION_DEPTH_BOUNDS_TEST) != 0)
+				START_NV_AFTERMATH();
+			}
+			else if (IsRHIDeviceAMD() && AmdAgsContext)
 			{
-				GSupportsDepthBoundsTest = true;
+				if ((AmdSupportedExtensionFlags & AGS_DX11_EXTENSION_DEPTH_BOUNDS_TEST) != 0)
+				{
+					GSupportsDepthBoundsTest = true;
+				}
 			}
 		}
 
