DR - Change default back to 24-bit depth buffers as that is incompatible with Oculus

diff --git a/Engine/Source/Runtime/D3D12RHI/Private/D3D12RHI.cpp b/Engine/Source/Runtime/D3D12RHI/Private/D3D12RHI.cpp
--- a/Engine/Source/Runtime/D3D12RHI/Private/D3D12RHI.cpp
+++ b/Engine/Source/Runtime/D3D12RHI/Private/D3D12RHI.cpp
@@ -22,9 +22,9 @@ DEFINE_LOG_CATEGORY(LogD3D12RHI);
 
 
 static TAutoConsoleVariable<int32> CVarD3D12UseD24(
-	TEXT("r.d3d12.Depth24Bit"),
-	0,
-	TEXT("0: Use 32bit float(default)\n1: Use 24bit fixed point"),
+	TEXT("r.D3D12.Depth24Bit"),
+	1,
+	TEXT("0: Use 32-bit float depth buffer\n1: Use 24-bit fixed point depth buffer(default)\n"),
 	ECVF_ReadOnly
 );
 
diff --git a/Engine/Source/Runtime/Windows/D3D11RHI/Private/D3D11Device.cpp b/Engine/Source/Runtime/Windows/D3D11RHI/Private/D3D11Device.cpp
--- a/Engine/Source/Runtime/Windows/D3D11RHI/Private/D3D11Device.cpp
+++ b/Engine/Source/Runtime/Windows/D3D11RHI/Private/D3D11Device.cpp
@@ -32,9 +32,9 @@ bool D3D11RHI_ShouldAllowAsyncResourceCreation()
 IMPLEMENT_MODULE(FD3D11DynamicRHIModule, D3D11RHI);
 
 static TAutoConsoleVariable<int32> CVarD3D11UseD24(
-	TEXT("r.d3d11.Depth24Bit"),
-	0,
-	TEXT("0: Use 32bit float(default)\n1: Use 24bit fixed point"),
+	TEXT("r.D3D11.Depth24Bit"),
+	1,
+	TEXT("0: Use 32-bit float depth buffer\n1: Use 24-bit fixed point depth buffer(default)\n"),
 	ECVF_ReadOnly
 );
 
