Added some logging for a previously silent failure during D3D11 adapter enumeration in the hopes of catching a rare failure case.

diff --git a/Engine/Source/Runtime/Windows/D3D11RHI/Private/Windows/WindowsD3D11Device.cpp b/Engine/Source/Runtime/Windows/D3D11RHI/Private/Windows/WindowsD3D11Device.cpp
--- a/Engine/Source/Runtime/Windows/D3D11RHI/Private/Windows/WindowsD3D11Device.cpp
+++ b/Engine/Source/Runtime/Windows/D3D11RHI/Private/Windows/WindowsD3D11Device.cpp
@@ -844,6 +844,14 @@ void FD3D11DynamicRHIModule::FindAdapter()
 					}
 				}
 			}
+			else
+			{
+				UE_LOG(LogD3D11RHI, Log, TEXT("  %2u. Unknown, failed to create test device."), AdapterIndex);
+			}
+		}
+		else
+		{
+			UE_LOG(LogD3D11RHI, Log, TEXT("  %2u. Unknown, failed to create adapter."), AdapterIndex);
 		}
 	}
 
