Change XAudio2_7 dll unloading so that it is only unloaded when the engine is shutting down as opposed to when devices are lost.

* Modify how the XAudio2_7 dll is loaded/unloaded such that its lifetime matches the engine lifetime rather than a single device's lifetime.
* This may mitigate a potential deadlock from occuring when attempting to free XAudio2 resources, such as when calling DestroyVoice.
* More details here: https://blogs.msdn.microsoft.com/chuckw/2015/10/09/known-issues-xaudio-2-7/

diff --git a/Engine/Source/Runtime/Engine/Private/AudioThread.cpp b/Engine/Source/Runtime/Engine/Private/AudioThread.cpp
--- a/Engine/Source/Runtime/Engine/Private/AudioThread.cpp
+++ b/Engine/Source/Runtime/Engine/Private/AudioThread.cpp
@@ -478,12 +478,13 @@ void FAudioCommandFence::Wait(bool bProcessGameThreadTasks) const
 		FTaskGraphInterface::Get().TriggerEventWhenTaskCompletes(Event, CompletionEvent, ENamedThreads::GameThread);
 
 		bool bDone;
-		const uint32 WaitTime = 35;
+		const uint32 WaitTimeMs = 35;
 		do
 		{
-			bDone = Event->Wait(WaitTime);
+			bDone = Event->Wait(WaitTimeMs);
+
 			float ThisTime = FPlatformTime::Seconds() - StartTime;
- 			if (ThisTime > .036f)
+ 			if (ThisTime > static_cast<float>(WaitTimeMs) / 1000.0f + SMALL_NUMBER)
 			{
 				if (GCVarEnableAudioCommandLogging == 1)
 				{
@@ -497,7 +498,7 @@ void FAudioCommandFence::Wait(bool bProcessGameThreadTasks) const
 				}
 				else
 				{
-					UE_LOG(LogAudio, Warning,  TEXT("Waited %fms for audio thread."), ThisTime * 1000.0f);
+					UE_LOG(LogAudio, Warning,  TEXT("Waited %f ms for audio thread."), ThisTime * 1000.0f);
 				}
 			}
 		} while (!bDone);
diff --git a/Engine/Source/Runtime/Windows/XAudio2/Private/XAudio2Device.cpp b/Engine/Source/Runtime/Windows/XAudio2/Private/XAudio2Device.cpp
--- a/Engine/Source/Runtime/Windows/XAudio2/Private/XAudio2Device.cpp
+++ b/Engine/Source/Runtime/Windows/XAudio2/Private/XAudio2Device.cpp
@@ -62,6 +62,10 @@ const float* FXAudioDeviceProperties::OutputMixMatrix = NULL;
 XAUDIO2_DEVICE_DETAILS FXAudioDeviceProperties::DeviceDetails;
 #endif	//XAUDIO_SUPPORTS_DEVICE_DETAILS
 
+#if PLATFORM_WINDOWS
+HMODULE FXAudioDeviceProperties::XAudio2Dll = nullptr;
+#endif
+
 /*------------------------------------------------------------------------------------
 	FAudioDevice Interface.
 ------------------------------------------------------------------------------------*/
@@ -111,13 +115,22 @@ bool FXAudio2Device::InitializeHardware()
 	// Load the xaudio2 library and keep a handle so we can free it on teardown
 	// Note: windows internally ref-counts the library per call to load library so 
 	// when we call FreeLibrary, it will only free it once the refcount is zero
-	DeviceProperties->XAudio2Dll = LoadLibraryA("XAudio2_7.dll");
+	// It's important that this dll stay around for the duration of the application's lifetime;
+	// attempting to free the library and losing our final ref when we change devices still exhibits the bad behavior.
 
-	// returning null means we failed to load XAudio2, which means everything will fail
-	if (DeviceProperties->XAudio2Dll == nullptr)
+	// Work around a known XAudio 2.7 issue: https://blogs.msdn.microsoft.com/chuckw/2015/10/09/known-issues-xaudio-2-7/
+	if (FXAudioDeviceProperties::XAudio2Dll == nullptr)
 	{
-		UE_LOG(LogInit, Warning, TEXT("Failed to load XAudio2 dll"));
-		return false;
+		UE_LOG(LogInit, Verbose, TEXT("Loading XAudio2 dll"));
+
+		FXAudioDeviceProperties::XAudio2Dll = LoadLibraryA("XAudio2_7.dll");
+
+		// returning null means we failed to load XAudio2, which means everything will fail
+		if (FXAudioDeviceProperties::XAudio2Dll == nullptr)
+		{
+			UE_LOG(LogInit, Error, TEXT("Failed to load XAudio2 dll"));
+			return false;
+		}
 	}
 #endif	//PLATFORM_64BITS
 #endif	//PLATFORM_WINDOWS
diff --git a/Engine/Source/Runtime/Windows/XAudio2/Private/XAudio2Support.h b/Engine/Source/Runtime/Windows/XAudio2/Private/XAudio2Support.h
--- a/Engine/Source/Runtime/Windows/XAudio2/Private/XAudio2Support.h
+++ b/Engine/Source/Runtime/Windows/XAudio2/Private/XAudio2Support.h
@@ -743,7 +743,6 @@ struct FXAudioDeviceProperties final : public IDeviceChangedListener
 	// These variables are non-static to support multiple audio device instances
 	struct IXAudio2*					XAudio2;
 	struct IXAudio2MasteringVoice*		MasteringVoice;
-	HMODULE								XAudio2Dll;
 
 	// These variables are static because they are common across all audio device instances
 	static int32						NumSpeakers;
@@ -754,6 +753,7 @@ struct FXAudioDeviceProperties final : public IDeviceChangedListener
 
 #if PLATFORM_WINDOWS
 	FMMNotificationClient* NotificationClient;
+	static HMODULE XAudio2Dll;
 #endif
 
 	// For calculating speaker maps for 3d audio
@@ -774,7 +774,6 @@ struct FXAudioDeviceProperties final : public IDeviceChangedListener
 	FXAudioDeviceProperties()
 		: XAudio2(nullptr)
 		, MasteringVoice(nullptr)
-		, XAudio2Dll(nullptr)
 		, NumActiveVoices(0)
 		, bDeviceChanged(false)
 		, bAllowNewVoices(true)
@@ -815,12 +814,15 @@ struct FXAudioDeviceProperties final : public IDeviceChangedListener
 		}
 
 #if PLATFORM_WINDOWS && PLATFORM_64BITS
-		if (XAudio2Dll)
+		// Only free the library if we're shutting down
+		if (XAudio2Dll != nullptr && GIsRequestingExit)
 		{
+			UE_LOG(LogAudio, Verbose, TEXT("Freeing XAudio2 dll"));
 			if (!FreeLibrary(XAudio2Dll))
 			{
 				UE_LOG(LogAudio, Warning, TEXT("Failed to free XAudio2 Dll"));
 			}
+			XAudio2Dll = nullptr;
 		}
 #endif
 	}
