github #5020
Added formats requested + additional missing formats. Both d3d11 and d3d12

diff --git a/Engine/Source/Runtime/D3D12RHI/Private/D3D12RenderTarget.cpp b/Engine/Source/Runtime/D3D12RHI/Private/D3D12RenderTarget.cpp
--- a/Engine/Source/Runtime/D3D12RHI/Private/D3D12RenderTarget.cpp
+++ b/Engine/Source/Runtime/D3D12RHI/Private/D3D12RenderTarget.cpp
@@ -648,14 +648,25 @@ struct FD3DFloatR11G11B10
 	}
 };
 
-// Only supports the formats that are supported by ConvertRAWSurfaceDataToFColor()
+
 static uint32 ComputeBytesPerPixel(DXGI_FORMAT Format)
 {
 	uint32 BytesPerPixel = 0;
 
 	switch (Format)
 	{
+	case DXGI_FORMAT_R8G8_TYPELESS:
+	case DXGI_FORMAT_R8G8_UNORM:
+	case DXGI_FORMAT_R8G8_UINT:
+	case DXGI_FORMAT_R8G8_SNORM:
+	case DXGI_FORMAT_R8G8_SINT:
 	case DXGI_FORMAT_R16_TYPELESS:
+	case DXGI_FORMAT_R16_FLOAT:
+	case DXGI_FORMAT_D16_UNORM:
+	case DXGI_FORMAT_R16_UNORM:
+	case DXGI_FORMAT_R16_UINT:
+	case DXGI_FORMAT_R16_SNORM:
+	case DXGI_FORMAT_R16_SINT:
 		BytesPerPixel = 2;
 		break;
 	case DXGI_FORMAT_B8G8R8A8_TYPELESS:
@@ -669,29 +680,55 @@ static uint32 ComputeBytesPerPixel(DXGI_FORMAT Format)
 	case DXGI_FORMAT_R32_UINT:
 	case DXGI_FORMAT_R32_TYPELESS:
 	case DXGI_FORMAT_R32_FLOAT:
+	case DXGI_FORMAT_B8G8R8A8_UNORM_SRGB:
+	case DXGI_FORMAT_B8G8R8X8_TYPELESS:
+	case DXGI_FORMAT_B8G8R8X8_UNORM_SRGB:
+	case DXGI_FORMAT_B8G8R8X8_UNORM:
+	case DXGI_FORMAT_R10G10B10A2_TYPELESS:
+	case DXGI_FORMAT_R10G10B10A2_UINT:
+	case DXGI_FORMAT_R8G8B8A8_UNORM_SRGB:
+	case DXGI_FORMAT_R8G8B8A8_UINT:
+	case DXGI_FORMAT_R8G8B8A8_SNORM:
+	case DXGI_FORMAT_R8G8B8A8_SINT:
+	case DXGI_FORMAT_R16G16_TYPELESS:
+	case DXGI_FORMAT_R16G16_FLOAT:
+	case DXGI_FORMAT_R16G16_UINT:
+	case DXGI_FORMAT_R16G16_SNORM:
+	case DXGI_FORMAT_R16G16_SINT:
+	case DXGI_FORMAT_D32_FLOAT:
+	case DXGI_FORMAT_R32_SINT:
 		BytesPerPixel = 4;
 		break;
 	case DXGI_FORMAT_R16G16B16A16_FLOAT:
 	case DXGI_FORMAT_R16G16B16A16_UNORM:
+	case DXGI_FORMAT_R16G16B16A16_TYPELESS:
+	case DXGI_FORMAT_R16G16B16A16_UINT:
+	case DXGI_FORMAT_R16G16B16A16_SNORM:
+	case DXGI_FORMAT_R16G16B16A16_SINT:
+	case DXGI_FORMAT_R32G32_TYPELESS:
+	case DXGI_FORMAT_R32G32_FLOAT:
+	case DXGI_FORMAT_R32G32_UINT:
+	case DXGI_FORMAT_R32G32_SINT:
 		BytesPerPixel = 8;
 		break;
-		// Changing Depth Buffers to 32 bit on Dingo as D24S8 is actually implemented as a 32 bit buffer in the hardware
 	case DXGI_FORMAT_R32G8X24_TYPELESS:
 	case DXGI_FORMAT_D32_FLOAT_S8X24_UINT:
 	case DXGI_FORMAT_R32_FLOAT_X8X24_TYPELESS:
 		BytesPerPixel = 5;
 		break;
-	case DXGI_FORMAT_R32G32B32A32_FLOAT:
-		BytesPerPixel = 16;
-		break;
 	case DXGI_FORMAT_R8_TYPELESS:
 	case DXGI_FORMAT_R8_UNORM:
 	case DXGI_FORMAT_R8_UINT:
 	case DXGI_FORMAT_R8_SNORM:
 	case DXGI_FORMAT_R8_SINT:
+	case DXGI_FORMAT_A8_UNORM:
+	case DXGI_FORMAT_R1_UNORM:
 		BytesPerPixel = 1;
 		break;
 	case DXGI_FORMAT_R32G32B32A32_UINT:
+	case DXGI_FORMAT_R32G32B32A32_TYPELESS:
+	case DXGI_FORMAT_R32G32B32A32_FLOAT:
+	case DXGI_FORMAT_R32G32B32A32_SINT:
 		BytesPerPixel = 16;
 		break;
 	}
diff --git a/Engine/Source/Runtime/Windows/D3D11RHI/Private/D3D11RenderTarget.cpp b/Engine/Source/Runtime/Windows/D3D11RHI/Private/D3D11RenderTarget.cpp
--- a/Engine/Source/Runtime/Windows/D3D11RHI/Private/D3D11RenderTarget.cpp
+++ b/Engine/Source/Runtime/Windows/D3D11RHI/Private/D3D11RenderTarget.cpp
@@ -480,24 +480,60 @@ static uint32 ComputeBytesPerPixel(DXGI_FORMAT Format)
 
 	switch(Format)
 	{
+		case DXGI_FORMAT_R8G8_TYPELESS:
+		case DXGI_FORMAT_R8G8_UNORM:
+		case DXGI_FORMAT_R8G8_UINT:
+		case DXGI_FORMAT_R8G8_SNORM:
+		case DXGI_FORMAT_R8G8_SINT:
 		case DXGI_FORMAT_R16_TYPELESS:
+		case DXGI_FORMAT_R16_FLOAT:
+		case DXGI_FORMAT_D16_UNORM:
+		case DXGI_FORMAT_R16_UNORM:
+		case DXGI_FORMAT_R16_UINT:
+		case DXGI_FORMAT_R16_SNORM:
+		case DXGI_FORMAT_R16_SINT:
 			BytesPerPixel = 2;
 			break;
-		case DXGI_FORMAT_R8G8B8A8_TYPELESS:
-		case DXGI_FORMAT_R8G8B8A8_UNORM:
-		case DXGI_FORMAT_R8G8B8A8_UNORM_SRGB:
 		case DXGI_FORMAT_B8G8R8A8_TYPELESS:
 		case DXGI_FORMAT_B8G8R8A8_UNORM:
-		case DXGI_FORMAT_B8G8R8A8_UNORM_SRGB:
+		case DXGI_FORMAT_R8G8B8A8_TYPELESS:
+		case DXGI_FORMAT_R8G8B8A8_UNORM:
 		case DXGI_FORMAT_R24G8_TYPELESS:
 		case DXGI_FORMAT_R10G10B10A2_UNORM:
 		case DXGI_FORMAT_R11G11B10_FLOAT:
 		case DXGI_FORMAT_R16G16_UNORM:
 		case DXGI_FORMAT_R32_UINT:
+		case DXGI_FORMAT_R32_TYPELESS:
+		case DXGI_FORMAT_R32_FLOAT:
+		case DXGI_FORMAT_B8G8R8A8_UNORM_SRGB:
+		case DXGI_FORMAT_B8G8R8X8_TYPELESS:
+		case DXGI_FORMAT_B8G8R8X8_UNORM_SRGB:
+		case DXGI_FORMAT_B8G8R8X8_UNORM:
+		case DXGI_FORMAT_R10G10B10A2_TYPELESS:
+		case DXGI_FORMAT_R10G10B10A2_UINT:
+		case DXGI_FORMAT_R8G8B8A8_UNORM_SRGB:
+		case DXGI_FORMAT_R8G8B8A8_UINT:
+		case DXGI_FORMAT_R8G8B8A8_SNORM:
+		case DXGI_FORMAT_R8G8B8A8_SINT:
+		case DXGI_FORMAT_R16G16_TYPELESS:
+		case DXGI_FORMAT_R16G16_FLOAT:
+		case DXGI_FORMAT_R16G16_UINT:
+		case DXGI_FORMAT_R16G16_SNORM:
+		case DXGI_FORMAT_R16G16_SINT:
+		case DXGI_FORMAT_D32_FLOAT:
+		case DXGI_FORMAT_R32_SINT:
 			BytesPerPixel = 4;
 			break;
 		case DXGI_FORMAT_R16G16B16A16_FLOAT:
 		case DXGI_FORMAT_R16G16B16A16_UNORM:
+		case DXGI_FORMAT_R16G16B16A16_TYPELESS:
+		case DXGI_FORMAT_R16G16B16A16_UINT:
+		case DXGI_FORMAT_R16G16B16A16_SNORM:
+		case DXGI_FORMAT_R16G16B16A16_SINT:
+		case DXGI_FORMAT_R32G32_TYPELESS:
+		case DXGI_FORMAT_R32G32_FLOAT:
+		case DXGI_FORMAT_R32G32_UINT:
+		case DXGI_FORMAT_R32G32_SINT:
 			BytesPerPixel = 8;
 			break;
 		// Changing Depth Buffers to 32 bit on Dingo as D24S8 is actually implemented as a 32 bit buffer in the hardware
@@ -506,20 +542,19 @@ static uint32 ComputeBytesPerPixel(DXGI_FORMAT Format)
 		case DXGI_FORMAT_R32_FLOAT_X8X24_TYPELESS:
 			BytesPerPixel = 5;
 			break;
-		case DXGI_FORMAT_R32G32B32A32_FLOAT:
-			BytesPerPixel = 16;
-			break;
-		case DXGI_FORMAT_R32G32_FLOAT:
-			BytesPerPixel = 8;
-			break;
 		case DXGI_FORMAT_R8_TYPELESS:
 		case DXGI_FORMAT_R8_UNORM:
 		case DXGI_FORMAT_R8_UINT:
 		case DXGI_FORMAT_R8_SNORM:
 		case DXGI_FORMAT_R8_SINT:
+		case DXGI_FORMAT_A8_UNORM:
+		case DXGI_FORMAT_R1_UNORM:
 			BytesPerPixel = 1;
 			break;
 		case DXGI_FORMAT_R32G32B32A32_UINT:
+		case DXGI_FORMAT_R32G32B32A32_TYPELESS:
+		case DXGI_FORMAT_R32G32B32A32_FLOAT:
+		case DXGI_FORMAT_R32G32B32A32_SINT:
 			BytesPerPixel = 16;
 			break;
 	}
