Fixing xma2 streaming for audio mixer

diff --git a/Engine/Source/Runtime/Windows/AudioMixerXAudio2/Private/AudioMixerPlatformXAudio2.cpp b/Engine/Source/Runtime/Windows/AudioMixerXAudio2/Private/AudioMixerPlatformXAudio2.cpp
--- a/Engine/Source/Runtime/Windows/AudioMixerXAudio2/Private/AudioMixerPlatformXAudio2.cpp
+++ b/Engine/Source/Runtime/Windows/AudioMixerXAudio2/Private/AudioMixerPlatformXAudio2.cpp
@@ -804,6 +804,13 @@ namespace Audio
 	{
 		check(InSoundWave);
 
+#if WITH_XMA2 && USE_XMA2_FOR_STREAMING
+		if (InSoundWave->IsStreaming() && InSoundWave->NumChannels <= 2 )
+		{
+			return new FXMAAudioInfo();
+		}
+#endif
+
 		if (InSoundWave->IsStreaming())
 		{
 #if USE_VORBIS_FOR_STREAMING
