Fix for Client Crash in OggReadStreaming

diff --git a/Engine/Source/Runtime/Windows/XAudio2/Private/XAudio2Source.cpp b/Engine/Source/Runtime/Windows/XAudio2/Private/XAudio2Source.cpp
--- a/Engine/Source/Runtime/Windows/XAudio2/Private/XAudio2Source.cpp
+++ b/Engine/Source/Runtime/Windows/XAudio2/Private/XAudio2Source.cpp
@@ -112,7 +112,7 @@ void FXAudio2SoundSource::FreeResources( void )
 
 	if (RealtimeAsyncTask)
 	{
-		RealtimeAsyncTask->EnsureCompletion();
+		RealtimeAsyncTask->EnsureCompletion(false);
 		delete RealtimeAsyncTask;
 		RealtimeAsyncTask = nullptr;
 		check(bResourcesNeedFreeing);
