[NO NON-PII COMMIT MESSAGE]

diff --git a/Engine/Source/Runtime/Windows/XAudio2/Private/XAudio2Source.cpp b/Engine/Source/Runtime/Windows/XAudio2/Private/XAudio2Source.cpp
--- a/Engine/Source/Runtime/Windows/XAudio2/Private/XAudio2Source.cpp
+++ b/Engine/Source/Runtime/Windows/XAudio2/Private/XAudio2Source.cpp
@@ -80,6 +80,7 @@ FXAudio2SoundSource::FXAudio2SoundSource(FAudioDevice* InAudioDevice)
 FXAudio2SoundSource::~FXAudio2SoundSource( void )
 {
 	FreeResources();
+	FreeBuffer();
 }
 
 void FXAudio2SoundSource::InitializeSourceEffects(uint32 InVoiceId)
@@ -112,12 +113,16 @@ void FXAudio2SoundSource::FreeResources( void )
 
 	if (RealtimeAsyncTask)
 	{
-		RealtimeAsyncTask->EnsureCompletion(false);
+		RealtimeAsyncTask->EnsureCompletion();
 		delete RealtimeAsyncTask;
 		RealtimeAsyncTask = nullptr;
 		check(bResourcesNeedFreeing);
 	}
+}
+
 
+void FXAudio2SoundSource::FreeBuffer()
+{
 	if (bResourcesNeedFreeing && Buffer)
 	{
 		// If we failed to initialize, then we will have a non-zero resource ID, but still need to delete the buffer
@@ -131,6 +136,7 @@ void FXAudio2SoundSource::FreeResources( void )
 	CurrentBuffer = 0;
 }
 
+
 /** 
  * Submit the relevant audio buffers to the system
  */
@@ -637,6 +643,7 @@ bool FXAudio2SoundSource::Init(FWaveInstance* InWaveInstance)
 
 	// Initialization failed.
 	FreeResources();
+	FreeBuffer();
 	return false;
 }
 
@@ -1256,6 +1263,11 @@ void FXAudio2SoundSourceCallback::OnBufferEnd(void* BufferContext)
 				return;
 			}
 		}
+		else if (SoundSource->RealtimeAsyncTask && !SoundSource->RealtimeAsyncTask->IsDone())
+		{
+			// If Playing was set to false in Stop(), we may still have an async decode task in flight:
+			SoundSource->RealtimeAsyncTask->EnsureCompletion(false);
+		}
 	}
 }
 
@@ -1465,7 +1477,6 @@ void FXAudio2SoundSource::Play()
 void FXAudio2SoundSource::Stop()
 {
 	bInitialized = false;
-	IStreamingManager::Get().GetAudioStreamingManager().RemoveStreamingSoundSource(this);
 
 	if( WaveInstance )
 	{	
@@ -1474,12 +1485,18 @@ void FXAudio2SoundSource::Stop()
 
 		// Free resources
 		FreeResources();
+	}
 
+	IStreamingManager::Get().GetAudioStreamingManager().RemoveStreamingSoundSource(this);
+	
+	if (WaveInstance)
+	{
+		FreeBuffer();
 		bBuffersToFlush = false;
 		bLoopCallback = false;
 		bResourcesNeedFreeing = false;
 	}
-
+	
 	FSoundSource::Stop();
 }
 
diff --git a/Engine/Source/Runtime/Windows/XAudio2/Private/XAudio2Support.h b/Engine/Source/Runtime/Windows/XAudio2/Private/XAudio2Support.h
--- a/Engine/Source/Runtime/Windows/XAudio2/Private/XAudio2Support.h
+++ b/Engine/Source/Runtime/Windows/XAudio2/Private/XAudio2Support.h
@@ -409,10 +409,15 @@ public:
 	virtual ~FXAudio2SoundSource();
 
 	/**
-	 * Frees existing resources. Called from destructor and therefore not virtual.
+	 * Frees existing resources except for the buffer. Called from destructor and therefore not virtual.
 	 */
 	void FreeResources();
 
+	/**
+	 * Frees the source's underlying buffer, if neccessary. 
+	 */
+	void FreeBuffer();
+
 	/**
 	* Initializes any effects used with this source voice
 	*/
