1. Fix null device callback, which was broken when swapping to the null device rather than launching with it.
2. Refactoring device swap code from using a thread safe bool to a mutex, ensuing we are not in an XAudio2 callback when we clean up XAudio2 resources.
3. Introduing au.DisableDeviceSwaps cvar, which can be used to ignore device notifications from the client.
4. Introducing some basic logic for timing out when the audio device stops performing it's callback, as well as logic to attempt to recover every so often when no devices are available.

diff --git a/Engine/Source/Runtime/AudioMixer/Private/AudioMixer.cpp b/Engine/Source/Runtime/AudioMixer/Private/AudioMixer.cpp
--- a/Engine/Source/Runtime/AudioMixer/Private/AudioMixer.cpp
+++ b/Engine/Source/Runtime/AudioMixer/Private/AudioMixer.cpp
@@ -24,6 +24,30 @@ FAutoConsoleVariableRef CVarLogRenderThreadPriority(
 	TEXT("0: Normal, 1: Above Normal, 2: Below Normal, 3: Highest, 4: Lowest, 5: Slightly Below Normal, 6: Time Critical"),
 	ECVF_Default);
 
+static int32 EnableDetailedWindowsDeviceLoggingCVar = 0;
+FAutoConsoleVariableRef CVarEnableDetailedWindowsDeviceLogging(
+	TEXT("au.EnableDetailedWindowsDeviceLogging"),
+	EnableDetailedWindowsDeviceLoggingCVar,
+	TEXT("Enables detailed windows device logging.\n")
+	TEXT("0: Not Enabled, 1: Enabled"),
+	ECVF_Default);
+
+static int32 DisableDeviceSwapCVar = 0;
+FAutoConsoleVariableRef CVarDisableDeviceSwap(
+	TEXT("au.DisableDeviceSwap"),
+	DisableDeviceSwapCVar,
+	TEXT("Disable device swap handling code for Audio Mixer on Windows.\n")
+	TEXT("0: Not Enabled, 1: Enabled"),
+	ECVF_Default);
+
+
+static int32 OverrunTimeoutCVar = 1000;
+FAutoConsoleVariableRef CVarOverrunTimeout(
+	TEXT("au.OverrunTimeoutMSec"),
+	OverrunTimeoutCVar,
+	TEXT("Amount of time to wait for the render thread to time out before swapping to the null device. \n"),
+	ECVF_Default);
+
 DEFINE_STAT(STAT_AudioMixerRenderAudio);
 DEFINE_STAT(STAT_AudioMixerSourceManagerUpdate);
 DEFINE_STAT(STAT_AudioMixerSourceBuffers);
@@ -217,10 +241,10 @@ namespace Audio
 		, NumOutputBuffers(0)
 		, FadeVolume(0.0f)
 		, LastError(TEXT("None"))
-		, bAudioDeviceChanging(false)
 		, bPerformingFade(true)
 		, bFadedOut(false)
 		, bIsDeviceInitialized(false)
+		, bMoveAudioStreamToNewAudioDevice(false)
 		, bIsUsingNullDevice(false)
 		, NullDeviceCallback(nullptr)
 	{
@@ -286,6 +310,20 @@ namespace Audio
 	{
 		if (!NullDeviceCallback.IsValid())
 		{
+			// Reset all of the buffers, then immediately kick off another render.
+			for (int32 Index = 0; Index < OutputBuffers.Num(); ++Index)
+			{
+				OutputBuffers[Index].Reset(OpenStreamParams.NumFrames * AudioStreamInfo.DeviceInfo.NumChannels);
+			}
+
+			CurrentBufferReadIndex = 0;
+			CurrentBufferWriteIndex = 1;
+
+			SubmitBuffer(OutputBuffers[CurrentBufferReadIndex].GetBufferData());
+			check(OpenStreamParams.NumFrames * AudioStreamInfo.DeviceInfo.NumChannels == OutputBuffers[CurrentBufferReadIndex].GetBuffer().Num());
+
+			AudioRenderEvent->Trigger();
+
 			float BufferDuration = ((float) OpenStreamParams.NumFrames) / OpenStreamParams.SampleRate;
 			NullDeviceCallback.Reset(new FMixerNullCallback( BufferDuration, [this]()
 			{
@@ -327,7 +365,7 @@ namespace Audio
 		LLM_SCOPE(ELLMTag::AudioMixer);
 
 		// Don't read any more audio if we're not running or changing device
-		if (AudioStreamInfo.StreamState != EAudioOutputStreamState::Running || bAudioDeviceChanging)
+		if (AudioStreamInfo.StreamState != EAudioOutputStreamState::Running || !DeviceSwapCriticalSection.TryLock())
 		{
 			return;
 		}
@@ -379,6 +417,8 @@ namespace Audio
 			CurrentBufferReadIndex = NextReadIndex;
 		}
 
+		DeviceSwapCriticalSection.Unlock();
+
 		// Kick off rendering of the next set of buffers
 		AudioRenderEvent->Trigger();
 	}
@@ -509,8 +549,16 @@ namespace Audio
 
 			RenderTimeAnalysis.End();
 
+			// Bounds check the timeout for our audio render event.
+			OverrunTimeoutCVar = FMath::Clamp(OverrunTimeoutCVar, 500, 5000);
+
 			// Now wait for a buffer to be consumed, which will bump up the read index.
-			AudioRenderEvent->Wait();
+			if (!AudioRenderEvent->Wait(static_cast<uint32>(OverrunTimeoutCVar)))
+			{
+				// if we reached this block, we timed out, and should attempt to
+				// bail on our current device.
+				bMoveAudioStreamToNewAudioDevice = true;
+			}
 		}
 
 		OpenStreamParams.AudioMixer->OnAudioStreamShutdown();
@@ -626,4 +674,14 @@ namespace Audio
 		}
 		return false;
 	}
+
+	bool IAudioMixer::ShouldIgnoreDeviceSwaps()
+	{
+		return DisableDeviceSwapCVar != 0;
+	}
+
+	bool IAudioMixer::ShouldLogDeviceSwaps()
+	{
+		return EnableDetailedWindowsDeviceLoggingCVar != 0;
+	}
 }
diff --git a/Engine/Source/Runtime/AudioMixer/Public/AudioMixer.h b/Engine/Source/Runtime/AudioMixer/Public/AudioMixer.h
--- a/Engine/Source/Runtime/AudioMixer/Public/AudioMixer.h
+++ b/Engine/Source/Runtime/AudioMixer/Public/AudioMixer.h
@@ -139,6 +139,12 @@ namespace Audio
 
 		bool IsMainAudioMixer() const { return bIsMainAudioMixer; }
 
+		/** Called by FWindowsMMNotificationClient to bypass notifications for audio device changes: */
+		AUDIOMIXER_API static bool ShouldIgnoreDeviceSwaps();
+
+		/** Called by FWindowsMMNotificationClient to toggle logging for audio device changes: */
+		AUDIOMIXER_API static bool ShouldLogDeviceSwaps();
+
 	protected:
 
 		IAudioMixer() 
@@ -506,8 +512,8 @@ namespace Audio
 		/** The render thread sync event. */
 		FEvent* AudioRenderEvent;
 
-		/** Event for a single buffer render. */
-		FEvent* AudioBufferEvent;
+		/** Critical Section used for times when we need the render loop to halt for the device swap. */
+		FCriticalSection DeviceSwapCriticalSection;
 
 		/** Event allows you to block until fadeout is complete. */
 		FEvent* AudioFadeEvent;
@@ -533,13 +539,11 @@ namespace Audio
 		/** Struct used to store render time analysis data. */
 		FAudioRenderTimeAnalysis RenderTimeAnalysis;
 
-		/** Flag if the audio device is in the process of changing. Prevents more buffers from being submitted to platform. */
-		FThreadSafeBool bAudioDeviceChanging;
-
 		FThreadSafeBool bPerformingFade;
 		FThreadSafeBool bFadedOut;
 		FThreadSafeBool bIsDeviceInitialized;
 
+		FThreadSafeBool bMoveAudioStreamToNewAudioDevice;
 		FThreadSafeBool bIsUsingNullDevice;
 
 	private:
diff --git a/Engine/Source/Runtime/Windows/AudioMixerXAudio2/Private/AudioMixerPlatformWindows.cpp b/Engine/Source/Runtime/Windows/AudioMixerXAudio2/Private/AudioMixerPlatformWindows.cpp
--- a/Engine/Source/Runtime/Windows/AudioMixerXAudio2/Private/AudioMixerPlatformWindows.cpp
+++ b/Engine/Source/Runtime/Windows/AudioMixerXAudio2/Private/AudioMixerPlatformWindows.cpp
@@ -18,15 +18,6 @@
 #include <mmdeviceapi.h>
 #include <functiondiscoverykeys_devpkey.h>
 
-static int32 GEnableDetailedWindowsDeviceLogging = 0;
-FAutoConsoleVariableRef CVarEnableDetailedWindowsDeviceLogging(
-	TEXT("au.EnableDetailedWindowsDeviceLogging"),
-	GEnableDetailedWindowsDeviceLogging,
-	TEXT("Enables detailed windows device logging.\n")
-	TEXT("0: Not Enabled, 1: Enabled"),
-	ECVF_Default);
-
-
 class FWindowsMMNotificationClient final : public IMMNotificationClient
 {
 public:
@@ -58,13 +49,18 @@ public:
 
 	HRESULT STDMETHODCALLTYPE OnDefaultDeviceChanged(EDataFlow InFlow, ERole InRole, LPCWSTR pwstrDeviceId) override
 	{
-		if (GEnableDetailedWindowsDeviceLogging)
+		if (Audio::IAudioMixer::ShouldLogDeviceSwaps())
 		{
 			UE_LOG(LogAudioMixer, Warning, TEXT("OnDefaultDeviceChanged: %d, %d, %s"), InFlow, InRole, pwstrDeviceId);
 		}
 
 		Audio::EAudioDeviceRole AudioDeviceRole;
 
+		if (Audio::IAudioMixer::ShouldIgnoreDeviceSwaps())
+		{
+			return S_OK;
+		}
+
 		if (InRole == eConsole)
 		{
 			AudioDeviceRole = Audio::EAudioDeviceRole::Console;
@@ -107,10 +103,15 @@ public:
 
 	HRESULT STDMETHODCALLTYPE OnDeviceAdded(LPCWSTR pwstrDeviceId) override
 	{
-		if (GEnableDetailedWindowsDeviceLogging)
+		if (Audio::IAudioMixer::ShouldLogDeviceSwaps())
 		{
 			UE_LOG(LogAudioMixer, Warning, TEXT("OnDeviceAdded: %s"), pwstrDeviceId);
 		}
+		
+		if (Audio::IAudioMixer::ShouldIgnoreDeviceSwaps())
+		{
+			return S_OK;
+		}
 
 		for (Audio::IAudioMixerDeviceChangedLister* Listener : Listeners)
 		{
@@ -121,11 +122,16 @@ public:
 
 	HRESULT STDMETHODCALLTYPE OnDeviceRemoved(LPCWSTR pwstrDeviceId) override
 	{
-		if (GEnableDetailedWindowsDeviceLogging)
+		if (Audio::IAudioMixer::ShouldLogDeviceSwaps())
 		{
 			UE_LOG(LogAudioMixer, Warning, TEXT("OnDeviceRemoved: %s"), pwstrDeviceId);
 		}
 
+		if (Audio::IAudioMixer::ShouldIgnoreDeviceSwaps())
+		{
+			return S_OK;
+		}
+
 		for (Audio::IAudioMixerDeviceChangedLister* Listener : Listeners)
 		{
 			Listener->OnDeviceRemoved(FString(pwstrDeviceId));
@@ -135,11 +141,16 @@ public:
 
 	HRESULT STDMETHODCALLTYPE OnDeviceStateChanged(LPCWSTR pwstrDeviceId, DWORD dwNewState) override
 	{
-		if (GEnableDetailedWindowsDeviceLogging)
+		if (Audio::IAudioMixer::ShouldLogDeviceSwaps())
 		{
 			UE_LOG(LogAudioMixer, Warning, TEXT("OnDeviceStateChanged: %s, %d"), pwstrDeviceId, dwNewState);
 		}
 
+		if (Audio::IAudioMixer::ShouldIgnoreDeviceSwaps())
+		{
+			return S_OK;
+		}
+
 		if (dwNewState == DEVICE_STATE_DISABLED || dwNewState == DEVICE_STATE_UNPLUGGED || dwNewState == DEVICE_STATE_NOTPRESENT)
 		{
 			for (Audio::IAudioMixerDeviceChangedLister* Listener : Listeners)
@@ -166,11 +177,16 @@ public:
 
 	HRESULT STDMETHODCALLTYPE OnPropertyValueChanged(LPCWSTR pwstrDeviceId, const PROPERTYKEY key)
 	{
-		if (GEnableDetailedWindowsDeviceLogging)
+		if (Audio::IAudioMixer::ShouldLogDeviceSwaps())
 		{
 			UE_LOG(LogAudioMixer, Warning, TEXT("OnPropertyValueChanged: %s, %d"), pwstrDeviceId, key.pid);
 		}
 
+		if (Audio::IAudioMixer::ShouldIgnoreDeviceSwaps())
+		{
+			return S_OK;
+		}
+
 		FString ChangedId = FString(pwstrDeviceId);
 
 		// look for ids we care about!
@@ -180,7 +196,7 @@ public:
 		{
 			for (Audio::IAudioMixerDeviceChangedLister* Listener : Listeners)
 			{
-				Listener->OnDeviceRemoved(FString(pwstrDeviceId));
+				Listener->OnDeviceRemoved(ChangedId);
 			}
 		}
 
diff --git a/Engine/Source/Runtime/Windows/AudioMixerXAudio2/Private/AudioMixerPlatformXAudio2.cpp b/Engine/Source/Runtime/Windows/AudioMixerXAudio2/Private/AudioMixerPlatformXAudio2.cpp
--- a/Engine/Source/Runtime/Windows/AudioMixerXAudio2/Private/AudioMixerPlatformXAudio2.cpp
+++ b/Engine/Source/Runtime/Windows/AudioMixerXAudio2/Private/AudioMixerPlatformXAudio2.cpp
@@ -48,6 +48,8 @@
 
 namespace Audio
 {
+	HMODULE FMixerPlatformXAudio2::XAudio2Dll = nullptr;
+
 	void FXAudio2VoiceCallback::OnBufferEnd(void* BufferContext)
 	{
 		check(BufferContext);
@@ -60,8 +62,8 @@ namespace Audio
 		, XAudio2System(nullptr)
 		, OutputAudioStreamMasteringVoice(nullptr)
 		, OutputAudioStreamSourceVoice(nullptr)
-		, bMoveAudioStreamToNewAudioDevice(false)
 		, LastDeviceSwapTime(0.0)
+		, TimeSinceNullDeviceWasLastChecked(0.0f)
 		, bIsComInitialized(false)
 		, bIsInitialized(false)
 		, bIsDeviceOpen(false)
@@ -126,7 +128,7 @@ namespace Audio
 			return false;
 		}
 
-		// Some devices spam device swap notifications, so we want to rate-limit them to prevent double/tripple triggering.
+		// Some devices spam device swap notifications, so we want to rate-limit them to prevent double/triple triggering.
 		static const int32 MinSwapTimeMs = 10;
 		if (CurrentTime - LastDeviceSwapTime > (double)MinSwapTimeMs / 1000.0)
 		{
@@ -154,7 +156,10 @@ namespace Audio
 		// Load the xaudio2 library and keep a handle so we can free it on teardown
 		// Note: windows internally ref-counts the library per call to load library so 
 		// when we call FreeLibrary, it will only free it once the refcount is zero
-		XAudio2Dll = LoadLibraryA("XAudio2_7.dll");
+		if (XAudio2Dll == nullptr)
+		{
+			XAudio2Dll = LoadLibraryA("XAudio2_7.dll");
+		}
 
 		// returning null means we failed to load XAudio2, which means everything will fail
 		if (XAudio2Dll == nullptr)
@@ -173,7 +178,7 @@ namespace Audio
 		Flags |= XAUDIO2_DO_NOT_USE_SHAPE;
 #endif
 
-		if (FAILED(XAudio2Create(&XAudio2System, Flags, (XAUDIO2_PROCESSOR)FPlatformAffinity::GetAudioThreadMask())))
+		if (!XAudio2System && FAILED(XAudio2Create(&XAudio2System, Flags, (XAUDIO2_PROCESSOR)FPlatformAffinity::GetAudioThreadMask())))
 		{
 			FMessageDialog::Open(EAppMsgType::Ok, NSLOCTEXT("Audio", "XAudio2Error", "Failed to initialize audio. This may be an issue with your installation of XAudio 2.7. XAudio2 is available in the DirectX End-User Runtime (June 2010)."));
 			return false;
@@ -205,12 +210,14 @@ namespace Audio
 #if PLATFORM_WINDOWS
 
 #if PLATFORM_64BITS
-		if (XAudio2Dll)
+		if (XAudio2Dll != nullptr && GIsRequestingExit)
 		{
 			if (!FreeLibrary(XAudio2Dll))
 			{
 				UE_LOG(LogAudio, Warning, TEXT("Failed to free XAudio2 Dll"));
 			}
+
+			XAudio2Dll = nullptr;
 		}
 #endif
 
@@ -591,6 +598,7 @@ namespace Audio
 			// Signal that the thread that is running the update that we're stopping
 			if (OutputAudioStreamSourceVoice)
 			{
+				FScopeLock ScopeLock(&DeviceSwapCriticalSection);
 				OutputAudioStreamSourceVoice->DestroyVoice();
 				OutputAudioStreamSourceVoice = nullptr;
 			}
@@ -618,6 +626,19 @@ namespace Audio
 	{
 #if PLATFORM_WINDOWS
 
+		uint32 NumDevices = 0;
+
+		check(XAudio2System != nullptr);
+		XAUDIO2_RETURN_ON_FAIL(XAudio2System->GetDeviceCount(&NumDevices));
+
+		// If we're running the null device, This function is called every second or so.
+		// Because of this, we early exit from this function if we're running the null device
+		// and there still are no devices.
+		if (bIsUsingNullDevice && !NumDevices)
+		{
+			return true;
+		}
+
 		UE_LOG(LogTemp, Log, TEXT("Resetting audio stream to device id %s"), *InNewDeviceId);
 
 		if (bIsUsingNullDevice)
@@ -632,8 +653,9 @@ namespace Audio
 				return true;
 			}
 
-			// Flag that we're changing audio devices so we stop submitting audio in the callbacks
-			bAudioDeviceChanging = true;
+			// If an XAudio2 callback is in flight,
+			// we have to wait for it here.
+			FScopeLock ScopeLock(&DeviceSwapCriticalSection);
 
 			if (OutputAudioStreamSourceVoice)
 			{
@@ -648,26 +670,7 @@ namespace Audio
 				OutputAudioStreamMasteringVoice->DestroyVoice();
 				OutputAudioStreamMasteringVoice = nullptr;
 			}
-
-			// Stop the engine from generating audio
-			if (XAudio2System)
-			{
-				XAudio2System->StopEngine();
-				SAFE_RELEASE(XAudio2System);
-			}
 		}
-		uint32 Flags = 0;
-
-#if WITH_XMA2
-		// We need to raise this flag explicitly to prevent initializing SHAPE twice, because we are allocating SHAPE in FXMAAudioInfo
-		Flags |= XAUDIO2_DO_NOT_USE_SHAPE;
-#endif
-
-		// Create a new xaudio2 system
-		XAUDIO2_RETURN_ON_FAIL(XAudio2Create(&XAudio2System, Flags, (XAUDIO2_PROCESSOR)FPlatformAffinity::GetAudioThreadMask()));
-
-		uint32 NumDevices = 0;
-		XAUDIO2_RETURN_ON_FAIL(XAudio2System->GetDeviceCount(&NumDevices));
 
 		if (NumDevices > 0)
 		{
@@ -708,25 +711,20 @@ namespace Audio
 			// Create the output source voice
 			XAUDIO2_RETURN_ON_FAIL(XAudio2System->CreateSourceVoice(&OutputAudioStreamSourceVoice, &Format, XAUDIO2_VOICE_NOPITCH, 2.0f, &OutputVoiceCallback));
 
-			// Start the xaudio2 system back up
-			XAudio2System->StartEngine();
+			const int32 NewNumSamples = OpenStreamParams.NumFrames * AudioStreamInfo.DeviceInfo.NumChannels;
+
+			// Clear the output buffers with zero's and submit one
+			for (int32 Index = 0; Index < OutputBuffers.Num(); ++Index)
+			{
+				OutputBuffers[Index].Reset(NewNumSamples);
+			}
 		}
 		else
-		{
+		{	
 			// If we don't have any hardware playback devices available, use the null device callback to render buffers.
 			StartRunningNullDevice();
 		}
 
-		const int32 NewNumSamples = OpenStreamParams.NumFrames * AudioStreamInfo.DeviceInfo.NumChannels;
-
-		// Clear the output buffers with zero's and submit one
- 		for (int32 Index = 0; Index < OutputBuffers.Num(); ++Index)
- 		{
- 			OutputBuffers[Index].Reset(NewNumSamples);
- 		}
-
-		bAudioDeviceChanging = false;
-
 #endif // #if PLATFORM_WINDOWS
 
 		return true;
@@ -740,6 +738,7 @@ namespace Audio
 			CurrentBufferWriteIndex = 1;
 
 			SubmitBuffer(OutputBuffers[CurrentBufferReadIndex].GetBufferData());
+			check(OpenStreamParams.NumFrames * AudioStreamInfo.DeviceInfo.NumChannels == OutputBuffers[CurrentBufferReadIndex].GetBuffer().Num());
 
 			AudioRenderEvent->Trigger();
 
@@ -848,6 +847,19 @@ namespace Audio
 		return FAudioPlatformSettings::GetPlatformSettings(TEXT("/Script/WindowsTargetPlatform.WindowsTargetSettings"));
 	}
 
+	void FMixerPlatformXAudio2::OnHardwareUpdate()
+	{
+		if (bIsUsingNullDevice)
+		{
+			float CurrentTime = FPlatformTime::Seconds();
+			if (CurrentTime - TimeSinceNullDeviceWasLastChecked > 1.0f)
+			{
+				bMoveAudioStreamToNewAudioDevice = true;
+				TimeSinceNullDeviceWasLastChecked = CurrentTime;
+			}
+		}
+	}
+
 	bool FMixerPlatformXAudio2::DisablePCMAudioCaching() const
 	{
 #if PLATFORM_WINDOWS
diff --git a/Engine/Source/Runtime/Windows/AudioMixerXAudio2/Private/AudioMixerPlatformXAudio2.h b/Engine/Source/Runtime/Windows/AudioMixerXAudio2/Private/AudioMixerPlatformXAudio2.h
--- a/Engine/Source/Runtime/Windows/AudioMixerXAudio2/Private/AudioMixerPlatformXAudio2.h
+++ b/Engine/Source/Runtime/Windows/AudioMixerXAudio2/Private/AudioMixerPlatformXAudio2.h
@@ -68,6 +68,7 @@ namespace Audio
 		virtual ICompressedAudioInfo* CreateCompressedAudioInfo(USoundWave* InSoundWave) override;
 		virtual FString GetDefaultDeviceName() override;
 		virtual FAudioPlatformSettings GetPlatformSettings() const override;
+		virtual void OnHardwareUpdate() override;
 		//~ End IAudioMixerPlatformInterface
 
 		//~ Begin IAudioMixerDeviceChangedLister
@@ -88,8 +89,10 @@ namespace Audio
 
 		typedef TArray<long> TChannelTypeMap;
 		
+#if PLATFORM_WINDOWS
 		// Handle to XAudio2DLL
-		HMODULE XAudio2Dll;
+		static HMODULE XAudio2Dll;
+#endif //PLATFORM_WINDOWS
 
 		// Bool indicating that the default audio device changed
 		// And that we need to restart the audio device.
@@ -103,9 +106,15 @@ namespace Audio
 		FCriticalSection AudioDeviceSwapCriticalSection;
 		FString OriginalAudioDeviceId;
 		FString NewAudioDeviceId;
-		FThreadSafeBool bMoveAudioStreamToNewAudioDevice;
+
+		// Because certain device drivers tend to spam IMMNotificationClient callbacks,
+		// We rate limit our device swap using time below:
 		double LastDeviceSwapTime;
 
+		// When we are running the null device,
+		// we check whether a new audio device was connected every second or so.
+		float TimeSinceNullDeviceWasLastChecked;
+
 		uint32 bIsComInitialized : 1;
 		uint32 bIsInitialized : 1;
 		uint32 bIsDeviceOpen : 1;
