Fix build for potential non-windows platforms.

diff --git a/Engine/Source/Runtime/Windows/AudioMixerXAudio2/Private/AudioMixerPlatformXAudio2.cpp b/Engine/Source/Runtime/Windows/AudioMixerXAudio2/Private/AudioMixerPlatformXAudio2.cpp
--- a/Engine/Source/Runtime/Windows/AudioMixerXAudio2/Private/AudioMixerPlatformXAudio2.cpp
+++ b/Engine/Source/Runtime/Windows/AudioMixerXAudio2/Private/AudioMixerPlatformXAudio2.cpp
@@ -48,7 +48,9 @@
 
 namespace Audio
 {
+#if PLATFORM_WINDOWS
 	HMODULE FMixerPlatformXAudio2::XAudio2Dll = nullptr;
+#endif
 
 	void FXAudio2VoiceCallback::OnBufferEnd(void* BufferContext)
 	{
