RHI Thread for DirectX11(Change by Intel)
Experimental, disabled by default for nowl.
define EXPERIMENTAL_D3D11_RHITHREAD to 1 in D3D11RHIPrivate.h to enable.

diff --git a/Engine/Shaders/Private/HZBOcclusion.usf b/Engine/Shaders/Private/HZBOcclusion.usf
--- a/Engine/Shaders/Private/HZBOcclusion.usf
+++ b/Engine/Shaders/Private/HZBOcclusion.usf
@@ -12,7 +12,6 @@ SamplerState	TextureSampler;
 void HZBBuildPS(noperspective float4 RenderTargetUVAndScreenPos : TEXCOORD0, out float4 OutColor : SV_Target0)
 {
 	float4 Depth;
-	float2 UV[4];
 
 #if STAGE == 0
 	const float2 ViewUV = RenderTargetUVAndScreenPos.zw * float2(0.5f, -0.5f) + 0.5f;
@@ -20,6 +19,8 @@ void HZBBuildPS(noperspective float4 RenderTargetUVAndScreenPos : TEXCOORD0, out
 	
 	// min(..., InputViewportMaxBound) because we don't want to sample outside of the viewport
 	// when the view size has odd dimensions on X/Y axis.
+#if COMPILER_GLSL || COMPILER_GLSL_ES2 || COMPILER_GLSL_ES3_1 || FEATURE_LEVEL < FEATURE_LEVEL_SM5
+	float2 UV[4];
 	UV[0] = min(InUV + float2(-0.25f, -0.25f) * InvSize, InputViewportMaxBound);
 	UV[1] = min(InUV + float2( 0.25f, -0.25f) * InvSize, InputViewportMaxBound);
 	UV[2] = min(InUV + float2(-0.25f,  0.25f) * InvSize, InputViewportMaxBound);
@@ -29,9 +30,16 @@ void HZBBuildPS(noperspective float4 RenderTargetUVAndScreenPos : TEXCOORD0, out
 	Depth.y = SceneTexturesStruct.SceneDepthTexture.SampleLevel( SceneTexturesStruct.SceneDepthTextureSampler, UV[1], 0 ).r;
 	Depth.z = SceneTexturesStruct.SceneDepthTexture.SampleLevel( SceneTexturesStruct.SceneDepthTextureSampler, UV[2], 0 ).r;
 	Depth.w = SceneTexturesStruct.SceneDepthTexture.SampleLevel( SceneTexturesStruct.SceneDepthTextureSampler, UV[3], 0 ).r;
+#else
+	float2 UV = min(InUV + float2(-0.25f, -0.25f) * InvSize, InputViewportMaxBound);
+	Depth = SceneTexturesStruct.SceneDepthTexture.GatherRed( SceneTexturesStruct.SceneDepthTextureSampler, UV, 0 );
+#endif
+
 #else
 	const float2 InUV = RenderTargetUVAndScreenPos.xy;
 	
+#if COMPILER_GLSL || COMPILER_GLSL_ES2 || COMPILER_GLSL_ES3_1 || FEATURE_LEVEL < FEATURE_LEVEL_SM5
+	float2 UV[4];
 	UV[0] = InUV + float2(-0.25f, -0.25f) * InvSize;
 	UV[1] = InUV + float2( 0.25f, -0.25f) * InvSize;
 	UV[2] = InUV + float2(-0.25f,  0.25f) * InvSize;
@@ -41,6 +49,11 @@ void HZBBuildPS(noperspective float4 RenderTargetUVAndScreenPos : TEXCOORD0, out
 	Depth.y = Texture.SampleLevel( TextureSampler, UV[1], 0 ).r;
 	Depth.z = Texture.SampleLevel( TextureSampler, UV[2], 0 ).r;
 	Depth.w = Texture.SampleLevel( TextureSampler, UV[3], 0 ).r;
+#else
+	float2 UV = InUV + float2(-0.25f, -0.25f) * InvSize;
+	Depth = Texture.GatherRed( TextureSampler, UV, 0 );
+#endif
+
 #endif
 
 	OutColor = min( min(Depth.x, Depth.y), min(Depth.z, Depth.w) );
diff --git a/Engine/Source/Runtime/RHI/Private/PipelineStateCache.cpp b/Engine/Source/Runtime/RHI/Private/PipelineStateCache.cpp
--- a/Engine/Source/Runtime/RHI/Private/PipelineStateCache.cpp
+++ b/Engine/Source/Runtime/RHI/Private/PipelineStateCache.cpp
@@ -794,7 +794,7 @@ static bool IsAsyncCompilationAllowed(FRHICommandList& RHICmdList)
 {
 	return !IsOpenGLPlatform(GMaxRHIShaderPlatform) &&  // The PSO cache is a waste of time on OpenGL and async compilation is a double waste of time.
 		!IsSwitchPlatform(GMaxRHIShaderPlatform) &&
-		GCVarAsyncPipelineCompile.GetValueOnAnyThread() && !RHICmdList.Bypass() && IsRunningRHIInSeparateThread();
+		GCVarAsyncPipelineCompile.GetValueOnAnyThread() && !RHICmdList.Bypass() && (IsRunningRHIInSeparateThread() && !IsInRHIThread());
 }
 
 FComputePipelineState* PipelineStateCache::GetAndOrCreateComputePipelineState(FRHICommandList& RHICmdList, FRHIComputeShader* ComputeShader)
diff --git a/Engine/Source/Runtime/RHI/Public/DynamicRHI.h b/Engine/Source/Runtime/RHI/Public/DynamicRHI.h
--- a/Engine/Source/Runtime/RHI/Public/DynamicRHI.h
+++ b/Engine/Source/Runtime/RHI/Public/DynamicRHI.h
@@ -1105,6 +1105,7 @@ public:
 	virtual void RHISignalFlipEvent() { }
 
 	virtual void RHICalibrateTimers() {}
+	virtual void RHIPollRenderQueryResults() {}
 
 	virtual bool RHIIsTypedUAVLoadSupported(EPixelFormat PixelFormat) { return true; }
 
diff --git a/Engine/Source/Runtime/RHI/Public/RHICommandList.h b/Engine/Source/Runtime/RHI/Public/RHICommandList.h
--- a/Engine/Source/Runtime/RHI/Public/RHICommandList.h
+++ b/Engine/Source/Runtime/RHI/Public/RHICommandList.h
@@ -4291,6 +4291,11 @@ public:
 	}
 	void UpdateTextureReference(FTextureReferenceRHIParamRef TextureRef, FTextureRHIParamRef NewTexture);
 
+	FORCEINLINE void PollRenderQueryResults()
+	{
+		GDynamicRHI->RHIPollRenderQueryResults();
+	}
+
 	void UpdateRHIResources(FRHIResourceUpdateInfo* UpdateInfos, int32 Num);
 };
 
diff --git a/Engine/Source/Runtime/Renderer/Private/DeferredShadingRenderer.cpp b/Engine/Source/Runtime/Renderer/Private/DeferredShadingRenderer.cpp
--- a/Engine/Source/Runtime/Renderer/Private/DeferredShadingRenderer.cpp
+++ b/Engine/Source/Runtime/Renderer/Private/DeferredShadingRenderer.cpp
@@ -554,6 +554,11 @@ void ServiceLocalQueue()
 {
 	SCOPE_CYCLE_COUNTER(STAT_FDeferredShadingSceneRenderer_Render_ServiceLocalQueue);
 	FTaskGraphInterface::Get().ProcessThreadUntilIdle(ENamedThreads::GetRenderThread_Local());
+
+	if (IsRunningRHIInSeparateThread())
+	{
+		FRHICommandListExecutor::GetImmediateCommandList().ImmediateFlush(EImmediateFlushType::DispatchToRHIThread);
+	}
 }
 
 // @return 0/1
diff --git a/Engine/Source/Runtime/Renderer/Private/SceneOcclusion.cpp b/Engine/Source/Runtime/Renderer/Private/SceneOcclusion.cpp
--- a/Engine/Source/Runtime/Renderer/Private/SceneOcclusion.cpp
+++ b/Engine/Source/Runtime/Renderer/Private/SceneOcclusion.cpp
@@ -1638,6 +1638,7 @@ void FSceneRenderer::FenceOcclusionTests(FRHICommandListImmediate& RHICmdList)
 		}
 		OcclusionSubmittedFence[0] = RHICmdList.RHIThreadFence();
 		RHICmdList.ImmediateFlush(EImmediateFlushType::DispatchToRHIThread);
+		RHICmdList.PollRenderQueryResults();
 	}
 }
 
diff --git a/Engine/Source/Runtime/Windows/D3D11RHI/Private/D3D11Commands.cpp b/Engine/Source/Runtime/Windows/D3D11RHI/Private/D3D11Commands.cpp
--- a/Engine/Source/Runtime/Windows/D3D11RHI/Private/D3D11Commands.cpp
+++ b/Engine/Source/Runtime/Windows/D3D11RHI/Private/D3D11Commands.cpp
@@ -2155,6 +2155,18 @@ void FD3D11DynamicRHI::RHIFlushComputeShaderCache()
 	EndUAVOverlap();
 }
 
+FD3D11DynamicRHI::FD3D11LockTracker& FD3D11DynamicRHI::GetThreadLocalLockTracker()
+{
+	if (IsInRHIThread())
+	{
+		return LockTrackers[1];
+	}
+	else
+	{
+		return LockTrackers[0];
+	}
+}
+
 //*********************** StagingBuffer Implementation ***********************//
 
 FStagingBufferRHIRef FD3D11DynamicRHI::RHICreateStagingBuffer()
diff --git a/Engine/Source/Runtime/Windows/D3D11RHI/Private/D3D11IndexBuffer.cpp b/Engine/Source/Runtime/Windows/D3D11RHI/Private/D3D11IndexBuffer.cpp
--- a/Engine/Source/Runtime/Windows/D3D11RHI/Private/D3D11IndexBuffer.cpp
+++ b/Engine/Source/Runtime/Windows/D3D11RHI/Private/D3D11IndexBuffer.cpp
@@ -61,10 +61,19 @@ FIndexBufferRHIRef FD3D11DynamicRHI::RHICreateIndexBuffer(uint32 Stride,uint32 S
 	return new FD3D11IndexBuffer(IndexBufferResource, Stride, Size, InUsage);
 }
 
+FIndexBufferRHIRef FD3D11DynamicRHI::CreateIndexBuffer_RenderThread(
+	class FRHICommandListImmediate& RHICmdList,
+	uint32 Stride,
+	uint32 Size,
+	uint32 InUsage,
+	FRHIResourceCreateInfo& CreateInfo)
+{
+	return RHICreateIndexBuffer(Stride, Size, InUsage, CreateInfo);
+}
+
 void* FD3D11DynamicRHI::RHILockIndexBuffer(FIndexBufferRHIParamRef IndexBufferRHI,uint32 Offset,uint32 Size,EResourceLockMode LockMode)
 {
 	FD3D11IndexBuffer* IndexBuffer = ResourceCast(IndexBufferRHI);
-	
 	// If this resource is bound to the device, unbind it
 	ConditionalClearShaderResource(IndexBuffer, true);
 
@@ -120,7 +129,7 @@ void* FD3D11DynamicRHI::RHILockIndexBuffer(FIndexBufferRHIParamRef IndexBufferRH
 	}
 	
 	// Add the lock to the lock map.
-	OutstandingLocks.Add(LockedKey,LockedData);
+	GetThreadLocalLockTracker().Add(LockedKey, LockedData);
 
 	// Return the offset pointer
 	return (void*)((uint8*)LockedData.GetData() + Offset);
@@ -136,9 +145,10 @@ void FD3D11DynamicRHI::RHIUnlockIndexBuffer(FIndexBufferRHIParamRef IndexBufferR
 	const bool bIsDynamic = (Desc.Usage == D3D11_USAGE_DYNAMIC);
 
 	// Find the outstanding lock for this IB.
+	FD3D11LockTracker& OutstandingLocks = GetThreadLocalLockTracker();
 	FD3D11LockedKey LockedKey(IndexBuffer->Resource);
 	FD3D11LockedData* LockedData = OutstandingLocks.Find(LockedKey);
-	check(LockedData);
+	checkf(LockedData, TEXT("Index buffer is either not locked or locked on a different thread"));
 
 	if(bIsDynamic)
 	{
diff --git a/Engine/Source/Runtime/Windows/D3D11RHI/Private/D3D11Query.cpp b/Engine/Source/Runtime/Windows/D3D11RHI/Private/D3D11Query.cpp
--- a/Engine/Source/Runtime/Windows/D3D11RHI/Private/D3D11Query.cpp
+++ b/Engine/Source/Runtime/Windows/D3D11RHI/Private/D3D11Query.cpp
@@ -6,11 +6,145 @@
 
 #include "D3D11RHIPrivate.h"
 
+class FD3D11RenderQueryBatcher final
+{
+public:
+	static FD3D11RenderQueryBatcher& Get()
+	{
+		static FD3D11RenderQueryBatcher Singleton;
+		return Singleton;
+	}
+
+	inline void StartNewBatch()
+	{
+		FRenderQueryBatch& NewBatch = ActiveBatches[ActiveBatches.AddDefaulted()];
+		NewBatch.Id = CurrentId;
+	}
+
+	void PerFrameFlush()
+	{
+		for (int32 BatchIdx = 0; BatchIdx < ActiveBatches.Num(); ++BatchIdx)
+		{
+			FRenderQueryBatch& Batch = ActiveBatches[BatchIdx];
+
+			// This is correct even if CurrentId overflowed
+			if (CurrentId - Batch.Id >= NumBufferedFrames)
+			{
+				ActiveBatches.RemoveAtSwap(BatchIdx--);
+			}
+		}
+
+		++CurrentId;
+	}
+
+	inline void Release()
+	{
+		ActiveBatches.Empty();
+	}
+
+	inline void Add(FRenderQueryRHIRef NewQuery)
+	{
+		FRenderQueryBatch& CurrentBatch = ActiveBatches.Last();
+		CurrentBatch.Queries.Add(NewQuery);
+	}
+
+	void PollQueryResults()
+	{
+		FD3D11DynamicRHI* D3D11RHI = static_cast<FD3D11DynamicRHI*>(GDynamicRHI);
+
+		for (int32 BatchIdx = 0; BatchIdx < ActiveBatches.Num(); ++BatchIdx)
+		{
+			TArray<FRenderQueryRHIRef>& Queries = ActiveBatches[BatchIdx].Queries;
+
+			for (int32 QueryIdx = 0; QueryIdx < Queries.Num(); ++QueryIdx)
+			{
+				FD3D11RenderQuery* Query = FD3D11DynamicRHI::ResourceCast(Queries[QueryIdx].GetReference());
+
+				if (Query->bResultIsCached || Query->GetRefCount() == 1)
+				{
+					Queries.RemoveAtSwap(QueryIdx--);
+				}
+				else if (D3D11RHI->GetQueryData(Query->Resource, &Query->Result, sizeof(Query->Result), Query->QueryType, false, false))
+				{
+					Query->bResultIsCached = true;
+					Queries.RemoveAtSwap(QueryIdx--);
+				}
+			}
+
+			if (!Queries.Num())
+			{
+				ActiveBatches.RemoveAtSwap(BatchIdx--);
+			}
+		}
+	}
+
+	void PollQueryResults_RenderThread()
+	{
+		FScopedRHIThreadStaller RHITStaller(FRHICommandListExecutor::GetImmediateCommandList());
+		PollQueryResults();
+	}
+
+	inline uint32 GetCurrentId() const
+	{
+		return CurrentId;
+	}
+
+	FD3D11RenderQueryBatcher(const FD3D11RenderQueryBatcher&) = delete;
+	FD3D11RenderQueryBatcher(FD3D11RenderQueryBatcher&&) = delete;
+	FD3D11RenderQueryBatcher& operator=(const FD3D11RenderQueryBatcher&) = delete;
+	FD3D11RenderQueryBatcher& operator=(FD3D11RenderQueryBatcher&&) = delete;
+
+private:
+	struct FRenderQueryBatch
+	{
+		uint32 Id;
+		TArray<FRenderQueryRHIRef> Queries;
+	};
+
+	static const int32 NumBufferedFrames = 5;
+
+	uint32 CurrentId;
+	TArray<FRenderQueryBatch> ActiveBatches;
+
+	FD3D11RenderQueryBatcher()
+		: CurrentId(0)
+	{}
+
+	virtual ~FD3D11RenderQueryBatcher() = default;
+};
+
+void D3D11RHIQueryBatcherPerFrameCleanup()
+{
+	FD3D11RenderQueryBatcher::Get().PerFrameFlush();
+}
+
+void FD3D11DynamicRHI::ReleaseCachedQueries()
+{
+	FD3D11RenderQueryBatcher::Get().Release();
+}
+
+void FD3D11DynamicRHI::RHIPollRenderQueryResults()
+{
+	if (ShouldNotEnqueueRHICommand())
+	{
+		FD3D11RenderQueryBatcher::Get().PollQueryResults();
+	}
+	else
+	{
+		RunOnRHIThread(
+			[]()
+		{
+			FD3D11RenderQueryBatcher::Get().PollQueryResults();
+		});
+	}
+}
+
 void FD3D11DynamicRHI::RHIBeginOcclusionQueryBatch(uint32 NumQueriesInBatch)
 {
 	check(RequestedOcclusionQueriesInBatch == 0);
 	ActualOcclusionQueriesInBatch = 0;
 	RequestedOcclusionQueriesInBatch = NumQueriesInBatch;
+	FD3D11RenderQueryBatcher::Get().StartNewBatch();
 }
 
 void FD3D11DynamicRHI::RHIEndOcclusionQueryBatch()
@@ -19,6 +153,13 @@ void FD3D11DynamicRHI::RHIEndOcclusionQueryBatch()
 	RequestedOcclusionQueriesInBatch = 0;
 }
 
+FRenderQueryRHIRef FD3D11DynamicRHI::RHICreateRenderQuery_RenderThread(
+	class FRHICommandListImmediate& RHICmdList,
+	ERenderQueryType QueryType)
+{
+	return RHICreateRenderQuery(QueryType);
+}
+
 FRenderQueryRHIRef FD3D11DynamicRHI::RHICreateRenderQuery(ERenderQueryType QueryType)
 {
 	TRefCountPtr<ID3D11Query> Query;
@@ -47,10 +188,18 @@ bool FD3D11DynamicRHI::RHIGetRenderQueryResult(FRenderQueryRHIParamRef QueryRHI,
 	check(IsInRenderingThread());
 	FD3D11RenderQuery* Query = ResourceCast(QueryRHI);
 
+	static uint32 LastQueryBatchId = 0xffffffff;
+	uint32 CurrentQueryBatchId = FD3D11RenderQueryBatcher::Get().GetCurrentId();
+	if (LastQueryBatchId != CurrentQueryBatchId && Query->QueryType == RQT_Occlusion)
+	{
+		LastQueryBatchId = CurrentQueryBatchId;
+		FD3D11RenderQueryBatcher::Get().PollQueryResults_RenderThread();
+	}
+
 	bool bSuccess = true;
 	if(!Query->bResultIsCached)
 	{
-		bSuccess = GetQueryData(Query->Resource,&Query->Result,sizeof(Query->Result),bWait, Query->QueryType);
+		bSuccess = GetQueryData(Query->Resource,&Query->Result,sizeof(Query->Result),Query->QueryType,bWait,true);
 
 		Query->bResultIsCached = bSuccess;
 	}
@@ -81,6 +230,7 @@ void FD3D11DynamicRHI::RHIBeginRenderQuery(FRenderQueryRHIParamRef QueryRHI)
 		++ActualOcclusionQueriesInBatch;
 		Query->bResultIsCached = false;
 		Direct3DDeviceIMContext->Begin(Query->Resource);
+		FD3D11RenderQueryBatcher::Get().Add(QueryRHI);
 	}
 	else
 	{
@@ -98,10 +248,16 @@ void FD3D11DynamicRHI::RHIEndRenderQuery(FRenderQueryRHIParamRef QueryRHI)
 	//@todo - d3d debug spews warnings about OQ's that are being issued but not polled, need to investigate
 }
 
-bool FD3D11DynamicRHI::GetQueryData(ID3D11Query* Query,void* Data,SIZE_T DataSize,bool bWait, ERenderQueryType QueryType)
+bool FD3D11DynamicRHI::GetQueryData(ID3D11Query* Query, void* Data, SIZE_T DataSize, ERenderQueryType QueryType, bool bWait, bool bStallRHIThread)
 {
+#define SAFE_GET_QUERY_DATA \
+	if (bStallRHIThread) D3D11StallRHIThread(); \
+	Result = Direct3DDeviceIMContext->GetData(Query, Data, DataSize, 0); \
+	if (bStallRHIThread) D3D11UnstallRHIThread();
+
 	// Request the data from the query.
-	HRESULT Result = Direct3DDeviceIMContext->GetData(Query,Data,DataSize,0);
+	HRESULT Result;
+	SAFE_GET_QUERY_DATA
 
 	// Isn't the query finished yet, and can we wait for it?
 	if ( Result == S_FALSE && bWait )
@@ -111,7 +267,7 @@ bool FD3D11DynamicRHI::GetQueryData(ID3D11Query* Query,void* Data,SIZE_T DataSiz
 		double StartTime = FPlatformTime::Seconds();
 		do 
 		{
-			Result = Direct3DDeviceIMContext->GetData(Query,Data,DataSize,0);
+			SAFE_GET_QUERY_DATA
 
 			// timer queries are used for Benchmarks which can stall a bit more
 			double TimeoutValue = (QueryType == RQT_AbsoluteTime) ? 2.0 : 0.5;
@@ -140,18 +296,30 @@ bool FD3D11DynamicRHI::GetQueryData(ID3D11Query* Query,void* Data,SIZE_T DataSiz
 		VERIFYD3D11RESULT_EX(Result, Direct3DDevice);
 		return false;
 	}
+#undef SAFE_GET_QUERY_DATA
 }
 
 void FD3D11EventQuery::IssueEvent()
 {
-	D3DRHI->GetDeviceContext()->End(Query);
+	if (ShouldNotEnqueueRHICommand())
+	{
+		D3DRHI->GetDeviceContext()->End(Query);
+	}
+	else
+	{
+		RunOnRHIThread(
+			[InQuery = Query]()
+		{
+			D3D11RHI_IMMEDIATE_CONTEXT->End(InQuery);
+		});
+	}
 }
 
 void FD3D11EventQuery::WaitForCompletion()
 {
 	BOOL bRenderingIsFinished = false;
 	while(
-		D3DRHI->GetQueryData(Query,&bRenderingIsFinished,sizeof(bRenderingIsFinished),true,RQT_Undefined) &&
+		D3DRHI->GetQueryData(Query,&bRenderingIsFinished,sizeof(bRenderingIsFinished),RQT_Undefined,true,true) &&
 		!bRenderingIsFinished
 		)
 	{};
@@ -226,13 +394,19 @@ void FD3D11BufferedGPUTiming::PlatformStaticInitialize(void* UserData)
 
 			D3D11_QUERY_DATA_TIMESTAMP_DISJOINT FreqQueryData;
 
+			D3D11StallRHIThread();
 			D3DResult = D3D11DeviceContext->GetData(FreqQuery, &FreqQueryData, sizeof(D3D11_QUERY_DATA_TIMESTAMP_DISJOINT), 0);
+			D3D11UnstallRHIThread();
+
 			double StartTime = FPlatformTime::Seconds();
 			while (D3DResult == S_FALSE && (FPlatformTime::Seconds() - StartTime) < 0.5f)
 			{
 				++DebugCounter;
 				FPlatformProcess::Sleep(0.005f);
+
+				D3D11StallRHIThread();
 				D3DResult = D3D11DeviceContext->GetData(FreqQuery, &FreqQueryData, sizeof(D3D11_QUERY_DATA_TIMESTAMP_DISJOINT), 0);
+				D3D11UnstallRHIThread();
 			}
 
 			if (D3DResult == S_OK)
diff --git a/Engine/Source/Runtime/Windows/D3D11RHI/Private/D3D11RHI.cpp b/Engine/Source/Runtime/Windows/D3D11RHI/Private/D3D11RHI.cpp
--- a/Engine/Source/Runtime/Windows/D3D11RHI/Private/D3D11RHI.cpp
+++ b/Engine/Source/Runtime/Windows/D3D11RHI/Private/D3D11RHI.cpp
@@ -648,7 +648,9 @@ bool FD3DGPUProfiler::CheckGpuHeartbeat() const
 	if (GDX11NVAfterMathEnabled && bTrackingGPUCrashData)
 	{
 		GFSDK_Aftermath_Device_Status Status;
+		D3D11StallRHIThread();
 		GFSDK_Aftermath_Result Result = GFSDK_Aftermath_GetDeviceStatus(&Status);
+		D3D11UnstallRHIThread();
 		if (Result == GFSDK_Aftermath_Result_Success)
 		{
 			if (Status != GFSDK_Aftermath_Device_Status_Active)
@@ -662,7 +664,9 @@ bool FD3DGPUProfiler::CheckGpuHeartbeat() const
 				if (AftermathContext)
 				{
 					GFSDK_Aftermath_ContextData ContextDataOut;
+					D3D11StallRHIThread();
 					Result = GFSDK_Aftermath_GetData(1, &AftermathContext, &ContextDataOut);
+					D3D11UnstallRHIThread();
 					if (Result == GFSDK_Aftermath_Result_Success)
 					{
 						UE_LOG(LogRHI, Error, TEXT("[Aftermath] GPU Stack Dump"));
diff --git a/Engine/Source/Runtime/Windows/D3D11RHI/Private/D3D11RHIPrivate.h b/Engine/Source/Runtime/Windows/D3D11RHI/Private/D3D11RHIPrivate.h
--- a/Engine/Source/Runtime/Windows/D3D11RHI/Private/D3D11RHIPrivate.h
+++ b/Engine/Source/Runtime/Windows/D3D11RHI/Private/D3D11RHIPrivate.h
@@ -84,6 +84,15 @@ struct Intel_MetricsDiscovery_ContextData
 // DX11 doesn't support higher MSAA count
 #define DX_MAX_MSAA_COUNT 8
 
+#ifndef EXPERIMENTAL_D3D11_RHITHREAD
+#define EXPERIMENTAL_D3D11_RHITHREAD 0
+#endif
+
+#if EXPERIMENTAL_D3D11_RHITHREAD
+#define D3D11_NUM_THREAD_LOCAL_CACHES 2
+#else
+#define D3D11_NUM_THREAD_LOCAL_CACHES 1
+#endif
 
 /**
  * The D3D RHI stats.
@@ -99,6 +108,7 @@ DECLARE_CYCLE_STAT_EXTERN(TEXT("New bound shader state time"),STAT_D3D11NewBound
 DECLARE_CYCLE_STAT_EXTERN(TEXT("Clean uniform buffer pool"),STAT_D3D11CleanUniformBufferTime,STATGROUP_D3D11RHI, );
 DECLARE_CYCLE_STAT_EXTERN(TEXT("Clear shader resources"),STAT_D3D11ClearShaderResourceTime,STATGROUP_D3D11RHI, );
 DECLARE_DWORD_ACCUMULATOR_STAT_EXTERN(TEXT("Uniform buffer pool num free"),STAT_D3D11NumFreeUniformBuffers,STATGROUP_D3D11RHI, );
+DECLARE_DWORD_COUNTER_STAT_EXTERN(TEXT("Immutable Uniform buffers"), STAT_D3D11NumImmutableUniformBuffers, STATGROUP_D3D11RHI, );
 DECLARE_DWORD_ACCUMULATOR_STAT_EXTERN(TEXT("Num Bound Shader State"), STAT_D3D11NumBoundShaderState, STATGROUP_D3D11RHI, );
 DECLARE_MEMORY_STAT_EXTERN(TEXT("Uniform buffer pool memory"), STAT_D3D11FreeUniformBufferMemory, STATGROUP_D3D11RHI, );
 DECLARE_CYCLE_STAT_EXTERN(TEXT("Update uniform buffer"),STAT_D3D11UpdateUniformBufferTime,STATGROUP_D3D11RHI, );
@@ -225,12 +235,12 @@ public:
 		Timing(InRHI, 1)
 	{
 		// Initialize Buffered timestamp queries 
-		Timing.InitResource(); // can't do this from the RHI thread
+		Timing.InitDynamicRHI(); // can't do this from the RHI thread
 	}
 
 	virtual ~FD3D11EventNode()
 	{
-		Timing.ReleaseResource();  // can't do this from the RHI thread
+		Timing.ReleaseDynamicRHI();  // can't do this from the RHI thread
 	}
 
 	/** 
@@ -263,14 +273,14 @@ public:
 		RootEventTiming(InRHI, 1),
 		DisjointQuery(InRHI)
 	{
-		RootEventTiming.InitResource();
-		DisjointQuery.InitResource();
+		RootEventTiming.InitDynamicRHI();
+		DisjointQuery.InitDynamicRHI();
 	}
 
 	~FD3D11EventNodeFrame()
 	{
-		RootEventTiming.ReleaseResource();
-		DisjointQuery.ReleaseResource();
+		RootEventTiming.ReleaseDynamicRHI();
+		DisjointQuery.ReleaseDynamicRHI();
 	}
 
 	/** Start this frame of per tracking */
@@ -334,12 +344,9 @@ struct AGSContext;
 class D3D11RHI_API FD3D11DynamicRHI : public FDynamicRHI, public IRHICommandContextPSOFallback
 {
 public:
-
+	typedef TMap<FD3D11LockedKey, FD3D11LockedData> FD3D11LockTracker;
 	friend class FD3D11Viewport;
 
-	/** Global D3D11 lock list */
-	TMap<FD3D11LockedKey,FD3D11LockedData> OutstandingLocks;
-
 	/** Initialization constructor. */
 	FD3D11DynamicRHI(IDXGIFactory1* InDXGIFactory1,D3D_FEATURE_LEVEL InFeatureLevel,int32 InChosenAdapter, const DXGI_ADAPTER_DESC& ChosenDescription);
 
@@ -371,11 +378,12 @@ public:
 	 * @param Query - The D3D query to read data from.
 	 * @param Data - The buffer to read the data into.
 	 * @param DataSize - The size of the buffer.
-	 * @param bWait - If true, it will wait for the query to finish.
 	 * @param QueryType e.g. RQT_Occlusion or RQT_AbsoluteTime
+	 * @param bWait - If true, it will wait for the query to finish.
+	 * @param bStallRHIThread - if true, stall RHIT before accessing immediate context
 	 * @return true if the query finished.
 	 */
-	bool GetQueryData(ID3D11Query* Query,void* Data,SIZE_T DataSize,bool bWait, ERenderQueryType QueryType);
+	bool GetQueryData(ID3D11Query* Query, void* Data, SIZE_T DataSize, ERenderQueryType QueryType, bool bWait, bool bStallRHIThread);
 
 	virtual void RHIBeginUpdateMultiFrameResource(FTextureRHIParamRef Texture) override;
 	virtual void RHIEndUpdateMultiFrameResource(FTextureRHIParamRef Texture) override;
@@ -440,6 +448,7 @@ public:
 	virtual FShaderResourceViewRHIRef RHICreateShaderResourceView(FTextureCubeRHIParamRef TextureCubeRHI, uint8 MipLevel) final override;
 	virtual void RHIGenerateMips(FTextureRHIParamRef Texture) final override;
 	virtual uint32 RHIComputeMemorySize(FTextureRHIParamRef TextureRHI) final override;
+	virtual void RHIAsyncCopyTexture2DCopy(FTexture2DRHIParamRef NewTexture2DRHI, FTexture2DRHIParamRef Texture2DRHI, int32 NewMipCount, int32 NewSizeX, int32 NewSizeY, FThreadSafeCounter* RequestStatus);
 	virtual FTexture2DRHIRef RHIAsyncReallocateTexture2D(FTexture2DRHIParamRef Texture2D, int32 NewMipCount, int32 NewSizeX, int32 NewSizeY, FThreadSafeCounter* RequestStatus) final override;
 	virtual ETextureReallocationStatus RHIFinalizeAsyncReallocateTexture2D(FTexture2DRHIParamRef Texture2D, bool bBlockUntilCompleted) final override;
 	virtual ETextureReallocationStatus RHICancelAsyncReallocateTexture2D(FTexture2DRHIParamRef Texture2D, bool bBlockUntilCompleted) final override;
@@ -590,7 +599,52 @@ public:
 	virtual FTextureCubeRHIRef RHICreateTextureCubeFromResource(EPixelFormat Format, uint32 TexCreateFlags, const FClearValueBinding& ClearValueBinding, ID3D11Texture2D* Resource);
 	virtual void RHIAliasTextureResources(FTextureRHIParamRef DestTexture, FTextureRHIParamRef SrcTexture);
 
-	virtual void RHIPerFrameRHIFlushComplete() override;
+	virtual void RHIPerFrameRHIFlushComplete() final override;
+	virtual void RHIPollRenderQueryResults() final override;
+
+	// *_RenderThread functions. Command lists call these functions on RT. You can implement your own behavior inside these functions.
+	// For example, deferring the actual creation to RHI thread by sending an RHI command.
+	// For D3D11, these functions mainly just remove RHIT stalls because ID3D11Device is thread safe.
+	virtual FVertexBufferRHIRef CreateVertexBuffer_RenderThread(class FRHICommandListImmediate& RHICmdList, uint32 Size, uint32 InUsage, FRHIResourceCreateInfo& CreateInfo) final override;
+	virtual FStructuredBufferRHIRef CreateStructuredBuffer_RenderThread(class FRHICommandListImmediate& RHICmdList, uint32 Stride, uint32 Size, uint32 InUsage, FRHIResourceCreateInfo& CreateInfo) final override;
+	virtual FShaderResourceViewRHIRef CreateShaderResourceView_RenderThread(class FRHICommandListImmediate& RHICmdList, FVertexBufferRHIParamRef VertexBuffer, uint32 Stride, uint8 Format) final override;
+	virtual FShaderResourceViewRHIRef CreateShaderResourceView_RenderThread(class FRHICommandListImmediate& RHICmdList, FIndexBufferRHIParamRef Buffer) final override;
+	virtual FTexture2DRHIRef AsyncReallocateTexture2D_RenderThread(class FRHICommandListImmediate& RHICmdList, FTexture2DRHIParamRef Texture2D, int32 NewMipCount, int32 NewSizeX, int32 NewSizeY, FThreadSafeCounter* RequestStatus) final override;
+	virtual ETextureReallocationStatus FinalizeAsyncReallocateTexture2D_RenderThread(class FRHICommandListImmediate& RHICmdList, FTexture2DRHIParamRef Texture2D, bool bBlockUntilCompleted) final override;
+	virtual ETextureReallocationStatus CancelAsyncReallocateTexture2D_RenderThread(class FRHICommandListImmediate& RHICmdList, FTexture2DRHIParamRef Texture2D, bool bBlockUntilCompleted) final override;
+	virtual FIndexBufferRHIRef CreateIndexBuffer_RenderThread(class FRHICommandListImmediate& RHICmdList, uint32 Stride, uint32 Size, uint32 InUsage, FRHIResourceCreateInfo& CreateInfo) final override;
+	virtual FVertexShaderRHIRef CreateVertexShader_RenderThread(class FRHICommandListImmediate& RHICmdList, const TArray<uint8>& Code) final override;
+	virtual FPixelShaderRHIRef CreatePixelShader_RenderThread(class FRHICommandListImmediate& RHICmdList, const TArray<uint8>& Code) final override;
+	virtual FGeometryShaderRHIRef CreateGeometryShader_RenderThread(class FRHICommandListImmediate& RHICmdList, const TArray<uint8>& Code) final override;
+	virtual FGeometryShaderRHIRef CreateGeometryShaderWithStreamOutput_RenderThread(class FRHICommandListImmediate& RHICmdList, const TArray<uint8>& Code, const FStreamOutElementList& ElementList, uint32 NumStrides, const uint32* Strides, int32 RasterizedStream) final override;
+	virtual FComputeShaderRHIRef CreateComputeShader_RenderThread(class FRHICommandListImmediate& RHICmdList, const TArray<uint8>& Code) final override;
+	virtual FHullShaderRHIRef CreateHullShader_RenderThread(class FRHICommandListImmediate& RHICmdList, const TArray<uint8>& Code) final override;
+	virtual FDomainShaderRHIRef CreateDomainShader_RenderThread(class FRHICommandListImmediate& RHICmdList, const TArray<uint8>& Code);
+	virtual void* LockTexture2D_RenderThread(class FRHICommandListImmediate& RHICmdList, FTexture2DRHIParamRef Texture, uint32 MipIndex, EResourceLockMode LockMode, uint32& DestStride, bool bLockWithinMiptail, bool bNeedsDefaultRHIFlush = true) final override;
+	virtual void UnlockTexture2D_RenderThread(class FRHICommandListImmediate& RHICmdList, FTexture2DRHIParamRef Texture, uint32 MipIndex, bool bLockWithinMiptail, bool bNeedsDefaultRHIFlush = true) final override;
+	virtual void UpdateTexture2D_RenderThread(class FRHICommandListImmediate& RHICmdList, FTexture2DRHIParamRef Texture, uint32 MipIndex, const struct FUpdateTextureRegion2D& UpdateRegion, uint32 SourcePitch, const uint8* SourceData) final override;
+	virtual void EndUpdateTexture3D_RenderThread(class FRHICommandListImmediate& RHICmdList, FUpdateTexture3DData& UpdateData) final override;
+	virtual void UpdateTexture3D_RenderThread(class FRHICommandListImmediate& RHICmdList, FTexture3DRHIParamRef Texture, uint32 MipIndex, const struct FUpdateTextureRegion3D& UpdateRegion, uint32 SourceRowPitch, uint32 SourceDepthPitch, const uint8* SourceData) final override;
+	virtual FTexture2DRHIRef RHICreateTexture2D_RenderThread(class FRHICommandListImmediate& RHICmdList, uint32 SizeX, uint32 SizeY, uint8 Format, uint32 NumMips, uint32 NumSamples, uint32 Flags, FRHIResourceCreateInfo& CreateInfo) final override;
+	virtual FTexture2DArrayRHIRef RHICreateTexture2DArray_RenderThread(class FRHICommandListImmediate& RHICmdList, uint32 SizeX, uint32 SizeY, uint32 SizeZ, uint8 Format, uint32 NumMips, uint32 Flags, FRHIResourceCreateInfo& CreateInfo) final override;
+	virtual FTexture3DRHIRef RHICreateTexture3D_RenderThread(class FRHICommandListImmediate& RHICmdList, uint32 SizeX, uint32 SizeY, uint32 SizeZ, uint8 Format, uint32 NumMips, uint32 Flags, FRHIResourceCreateInfo& CreateInfo) final override;
+	virtual FUnorderedAccessViewRHIRef RHICreateUnorderedAccessView_RenderThread(class FRHICommandListImmediate& RHICmdList, FStructuredBufferRHIParamRef StructuredBuffer, bool bUseUAVCounter, bool bAppendBuffer) final override;
+	virtual FUnorderedAccessViewRHIRef RHICreateUnorderedAccessView_RenderThread(class FRHICommandListImmediate& RHICmdList, FTextureRHIParamRef Texture, uint32 MipLevel) final override;
+	virtual FUnorderedAccessViewRHIRef RHICreateUnorderedAccessView_RenderThread(class FRHICommandListImmediate& RHICmdList, FVertexBufferRHIParamRef VertexBuffer, uint8 Format) final override;
+	virtual FUnorderedAccessViewRHIRef RHICreateUnorderedAccessView_RenderThread(class FRHICommandListImmediate& RHICmdList, FIndexBufferRHIParamRef IndexBuffer, uint8 Format) final override;
+	virtual FShaderResourceViewRHIRef RHICreateShaderResourceView_RenderThread(class FRHICommandListImmediate& RHICmdList, FTexture2DRHIParamRef Texture2DRHI, uint8 MipLevel) final override;
+	virtual FShaderResourceViewRHIRef RHICreateShaderResourceView_RenderThread(class FRHICommandListImmediate& RHICmdList, FTexture2DRHIParamRef Texture2DRHI, uint8 MipLevel, uint8 NumMipLevels, uint8 Format) final override;
+	virtual FShaderResourceViewRHIRef RHICreateShaderResourceView_RenderThread(class FRHICommandListImmediate& RHICmdList, FTexture3DRHIParamRef Texture3DRHI, uint8 MipLevel) final override;
+	virtual FShaderResourceViewRHIRef RHICreateShaderResourceView_RenderThread(class FRHICommandListImmediate& RHICmdList, FTexture2DArrayRHIParamRef Texture2DArrayRHI, uint8 MipLevel) final override;
+	virtual FShaderResourceViewRHIRef RHICreateShaderResourceView_RenderThread(class FRHICommandListImmediate& RHICmdList, FTextureCubeRHIParamRef TextureCubeRHI, uint8 MipLevel) final override;
+	virtual FShaderResourceViewRHIRef RHICreateShaderResourceView_RenderThread(class FRHICommandListImmediate& RHICmdList, FStructuredBufferRHIParamRef StructuredBuffer) final override;
+	virtual FTextureCubeRHIRef RHICreateTextureCube_RenderThread(class FRHICommandListImmediate& RHICmdList, uint32 Size, uint8 Format, uint32 NumMips, uint32 Flags, FRHIResourceCreateInfo& CreateInfo) final override;
+	virtual FTextureCubeRHIRef RHICreateTextureCubeArray_RenderThread(class FRHICommandListImmediate& RHICmdList, uint32 Size, uint32 ArraySize, uint8 Format, uint32 NumMips, uint32 Flags, FRHIResourceCreateInfo& CreateInfo) final override;
+	virtual FRenderQueryRHIRef RHICreateRenderQuery_RenderThread(class FRHICommandListImmediate& RHICmdList, ERenderQueryType QueryType) final override;
+	virtual FTextureReferenceRHIRef RHICreateTextureReference_RenderThread(class FRHICommandListImmediate& RHICmdList, FLastRenderTimeContainer* LastRenderTime) final override;
+	virtual void* RHILockTextureCubeFace_RenderThread(class FRHICommandListImmediate& RHICmdList, FTextureCubeRHIParamRef Texture, uint32 FaceIndex, uint32 ArrayIndex, uint32 MipIndex, EResourceLockMode LockMode, uint32& DestStride, bool bLockWithinMiptail) final override;
+	virtual void RHIUnlockTextureCubeFace_RenderThread(class FRHICommandListImmediate& RHICmdList, FTextureCubeRHIParamRef Texture, uint32 FaceIndex, uint32 ArrayIndex, uint32 MipIndex, bool bLockWithinMiptail) final override;
+	virtual void RHICopySubTextureRegion_RenderThread(class FRHICommandListImmediate& RHICmdList, FTexture2DRHIParamRef SourceTexture, FTexture2DRHIParamRef DestinationTexture, FBox2D SourceBox, FBox2D DestinationBox) final override;
 
 	virtual void RHIBeginRenderPass(const FRHIRenderPassInfo& InInfo, const TCHAR* InName) final override
 	{
@@ -641,6 +695,8 @@ public:
 		return GPUProfilingData.CheckGpuHeartbeat();
 	}
 
+	FD3D11LockTracker& GetThreadLocalLockTracker();
+
 	bool IsQuadBufferStereoEnabled();
 	void DisableQuadBufferStereo();
 
@@ -735,6 +791,9 @@ protected:
 
 	FD3D11StateCache StateCache;
 
+	/** Tracks outstanding locks on each thread */
+	FD3D11LockTracker LockTrackers[D3D11_NUM_THREAD_LOCAL_CACHES];
+
 	/** A list of all viewport RHIs that have been created. */
 	TArray<FD3D11Viewport*> Viewports;
 
@@ -932,6 +991,7 @@ protected:
 	virtual void CleanupD3DDevice();
 
 	void ReleasePooledUniformBuffers();
+	void ReleaseCachedQueries();
 
 	template<typename TPixelShader>
 	void ResolveTextureUsingShader(
diff --git a/Engine/Source/Runtime/Windows/D3D11RHI/Private/D3D11Shaders.cpp b/Engine/Source/Runtime/Windows/D3D11RHI/Private/D3D11Shaders.cpp
--- a/Engine/Source/Runtime/Windows/D3D11RHI/Private/D3D11Shaders.cpp
+++ b/Engine/Source/Runtime/Windows/D3D11RHI/Private/D3D11Shaders.cpp
@@ -54,6 +54,13 @@ FVertexShaderRHIRef FD3D11DynamicRHI::RHICreateVertexShader(const TArray<uint8>&
 	return Shader;
 }
 
+FVertexShaderRHIRef FD3D11DynamicRHI::CreateVertexShader_RenderThread(
+	class FRHICommandListImmediate& RHICmdList,
+	const TArray<uint8>& Code)
+{
+	return RHICreateVertexShader(Code);
+}
+
 FGeometryShaderRHIRef FD3D11DynamicRHI::RHICreateGeometryShader(const TArray<uint8>& Code) 
 { 
 	FShaderCodeReader ShaderCode(Code);
@@ -73,6 +80,13 @@ FGeometryShaderRHIRef FD3D11DynamicRHI::RHICreateGeometryShader(const TArray<uin
 	return Shader;
 }
 
+FGeometryShaderRHIRef FD3D11DynamicRHI::CreateGeometryShader_RenderThread(
+	class FRHICommandListImmediate& RHICmdList,
+	const TArray<uint8>& Code)
+{
+	return RHICreateGeometryShader(Code);
+}
+
 FGeometryShaderRHIRef FD3D11DynamicRHI::RHICreateGeometryShaderWithStreamOutput(const TArray<uint8>& Code, const FStreamOutElementList& ElementList, uint32 NumStrides, const uint32* Strides, int32 RasterizedStream) 
 { 
 	FShaderCodeReader ShaderCode(Code);
@@ -127,6 +141,17 @@ FGeometryShaderRHIRef FD3D11DynamicRHI::RHICreateGeometryShaderWithStreamOutput(
 	return Shader;
 }
 
+FGeometryShaderRHIRef FD3D11DynamicRHI::CreateGeometryShaderWithStreamOutput_RenderThread(
+	class FRHICommandListImmediate& RHICmdList,
+	const TArray<uint8>& Code,
+	const FStreamOutElementList& ElementList,
+	uint32 NumStrides,
+	const uint32* Strides,
+	int32 RasterizedStream)
+{
+	return RHICreateGeometryShaderWithStreamOutput(Code, ElementList, NumStrides, Strides, RasterizedStream);
+}
+
 FHullShaderRHIRef FD3D11DynamicRHI::RHICreateHullShader(const TArray<uint8>& Code) 
 { 
 	FShaderCodeReader ShaderCode(Code);
@@ -145,6 +170,13 @@ FHullShaderRHIRef FD3D11DynamicRHI::RHICreateHullShader(const TArray<uint8>& Cod
 	return Shader;
 }
 
+FHullShaderRHIRef FD3D11DynamicRHI::CreateHullShader_RenderThread(
+	class FRHICommandListImmediate& RHICmdList,
+	const TArray<uint8>& Code)
+{
+	return RHICreateHullShader(Code);
+}
+
 FDomainShaderRHIRef FD3D11DynamicRHI::RHICreateDomainShader(const TArray<uint8>& Code) 
 { 
 	FShaderCodeReader ShaderCode(Code);
@@ -163,6 +195,13 @@ FDomainShaderRHIRef FD3D11DynamicRHI::RHICreateDomainShader(const TArray<uint8>&
 	return Shader;
 }
 
+FDomainShaderRHIRef FD3D11DynamicRHI::CreateDomainShader_RenderThread(
+	class FRHICommandListImmediate& RHICmdList,
+	const TArray<uint8>& Code)
+{
+	return RHICreateDomainShader(Code);
+}
+
 FPixelShaderRHIRef FD3D11DynamicRHI::RHICreatePixelShader(const TArray<uint8>& Code)
 {
 	FShaderCodeReader ShaderCode(Code);
@@ -181,6 +220,13 @@ FPixelShaderRHIRef FD3D11DynamicRHI::RHICreatePixelShader(const TArray<uint8>& C
 	return Shader;
 }
 
+FPixelShaderRHIRef FD3D11DynamicRHI::CreatePixelShader_RenderThread(
+	class FRHICommandListImmediate& RHICmdList,
+	const TArray<uint8>& Code)
+{
+	return RHICreatePixelShader(Code);
+}
+
 FComputeShaderRHIRef FD3D11DynamicRHI::RHICreateComputeShader(const TArray<uint8>& Code) 
 { 
 	FShaderCodeReader ShaderCode(Code);
@@ -199,6 +245,13 @@ FComputeShaderRHIRef FD3D11DynamicRHI::RHICreateComputeShader(const TArray<uint8
 	return Shader;
 }
 
+FComputeShaderRHIRef FD3D11DynamicRHI::CreateComputeShader_RenderThread(
+	class FRHICommandListImmediate& RHICmdList,
+	const TArray<uint8>& Code)
+{
+	return RHICreateComputeShader(Code);
+}
+
 void FD3D11DynamicRHI::RHISetMultipleViewports(uint32 Count, const FViewportBounds* Data) 
 { 
 	check(Count > 0);
@@ -299,7 +352,7 @@ FBoundShaderStateRHIRef FD3D11DynamicRHI::RHICreateBoundShaderState(
 	FGeometryShaderRHIParamRef GeometryShaderRHI
 	)
 {
-	check(IsInRenderingThread());
+	check(IsInRenderingThread() || IsInRHIThread());
 
 	SCOPE_CYCLE_COUNTER(STAT_D3D11CreateBoundShaderStateTime);
 
diff --git a/Engine/Source/Runtime/Windows/D3D11RHI/Private/D3D11StructuredBuffer.cpp b/Engine/Source/Runtime/Windows/D3D11RHI/Private/D3D11StructuredBuffer.cpp
--- a/Engine/Source/Runtime/Windows/D3D11RHI/Private/D3D11StructuredBuffer.cpp
+++ b/Engine/Source/Runtime/Windows/D3D11RHI/Private/D3D11StructuredBuffer.cpp
@@ -92,6 +92,16 @@ FStructuredBufferRHIRef FD3D11DynamicRHI::RHICreateStructuredBuffer(uint32 Strid
 	return new FD3D11StructuredBuffer(StructuredBufferResource,Stride,Size,InUsage);
 }
 
+FStructuredBufferRHIRef FD3D11DynamicRHI::CreateStructuredBuffer_RenderThread(
+	class FRHICommandListImmediate& RHICmdList,
+	uint32 Stride,
+	uint32 Size,
+	uint32 InUsage,
+	FRHIResourceCreateInfo& CreateInfo)
+{
+	return RHICreateStructuredBuffer(Stride, Size, InUsage, CreateInfo);
+}
+
 void* FD3D11DynamicRHI::RHILockStructuredBuffer(FStructuredBufferRHIParamRef StructuredBufferRHI,uint32 Offset,uint32 Size,EResourceLockMode LockMode)
 {
 	FD3D11StructuredBuffer* StructuredBuffer = ResourceCast(StructuredBufferRHI);
@@ -151,7 +161,7 @@ void* FD3D11DynamicRHI::RHILockStructuredBuffer(FStructuredBufferRHIParamRef Str
 	}
 
 	// Add the lock to the lock map.
-	OutstandingLocks.Add(LockedKey,LockedData);
+	GetThreadLocalLockTracker().Add(LockedKey,LockedData);
 
 	// Return the offset pointer
 	return (void*)((uint8*)LockedData.GetData() + Offset);
@@ -167,9 +177,10 @@ void FD3D11DynamicRHI::RHIUnlockStructuredBuffer(FStructuredBufferRHIParamRef St
 	const bool bIsDynamic = (Desc.Usage == D3D11_USAGE_DYNAMIC);
 
 	// Find the outstanding lock for this VB.
+	FD3D11LockTracker& OutstandingLocks = GetThreadLocalLockTracker();
 	FD3D11LockedKey LockedKey(StructuredBuffer->Resource);
 	FD3D11LockedData* LockedData = OutstandingLocks.Find(LockedKey);
-	check(LockedData);
+	checkf(LockedData, TEXT("Structured buffer is either not locked or locked on a different thread"));
 
 	if(bIsDynamic)
 	{
diff --git a/Engine/Source/Runtime/Windows/D3D11RHI/Private/D3D11Texture.cpp b/Engine/Source/Runtime/Windows/D3D11RHI/Private/D3D11Texture.cpp
--- a/Engine/Source/Runtime/Windows/D3D11RHI/Private/D3D11Texture.cpp
+++ b/Engine/Source/Runtime/Windows/D3D11RHI/Private/D3D11Texture.cpp
@@ -1098,6 +1098,19 @@ FTexture2DRHIRef FD3D11DynamicRHI::RHICreateTexture2D(uint32 SizeX,uint32 SizeY,
 	return CreateD3D11Texture2D<FD3D11BaseTexture2D>(SizeX,SizeY,1,false,false,Format,NumMips,NumSamples,Flags,CreateInfo);
 }
 
+FTexture2DRHIRef FD3D11DynamicRHI::RHICreateTexture2D_RenderThread(
+	class FRHICommandListImmediate& RHICmdList,
+	uint32 SizeX,
+	uint32 SizeY,
+	uint8 Format,
+	uint32 NumMips,
+	uint32 NumSamples,
+	uint32 Flags,
+	FRHIResourceCreateInfo& CreateInfo)
+{
+	return RHICreateTexture2D(SizeX, SizeY, Format, NumMips, NumSamples, Flags, CreateInfo);
+}
+
 FTexture2DRHIRef FD3D11DynamicRHI::RHIAsyncCreateTexture2D(uint32 SizeX,uint32 SizeY,uint8 Format,uint32 NumMips,uint32 Flags,void** InitialMipData,uint32 NumInitialMips)
 {
 	FD3D11Texture2D* NewTexture = NULL;
@@ -1237,12 +1250,38 @@ FTexture2DArrayRHIRef FD3D11DynamicRHI::RHICreateTexture2DArray(uint32 SizeX,uin
 	return CreateD3D11Texture2D<FD3D11BaseTexture2DArray>(SizeX,SizeY,SizeZ,true,false,Format,NumMips,1,Flags,CreateInfo);
 }
 
+FTexture2DArrayRHIRef FD3D11DynamicRHI::RHICreateTexture2DArray_RenderThread(
+	class FRHICommandListImmediate& RHICmdList,
+	uint32 SizeX,
+	uint32 SizeY,
+	uint32 SizeZ,
+	uint8 Format,
+	uint32 NumMips,
+	uint32 Flags,
+	FRHIResourceCreateInfo& CreateInfo)
+{
+	return RHICreateTexture2DArray(SizeX, SizeY, SizeZ, Format, NumMips, Flags, CreateInfo);
+}
+
 FTexture3DRHIRef FD3D11DynamicRHI::RHICreateTexture3D(uint32 SizeX,uint32 SizeY,uint32 SizeZ,uint8 Format,uint32 NumMips,uint32 Flags,FRHIResourceCreateInfo& CreateInfo)
 {
 	check(SizeZ >= 1);
 	return CreateD3D11Texture3D(SizeX,SizeY,SizeZ,Format,NumMips,Flags,CreateInfo);
 }
 
+FTexture3DRHIRef FD3D11DynamicRHI::RHICreateTexture3D_RenderThread(
+	class FRHICommandListImmediate& RHICmdList,
+	uint32 SizeX,
+	uint32 SizeY,
+	uint32 SizeZ,
+	uint8 Format,
+	uint32 NumMips,
+	uint32 Flags,
+	FRHIResourceCreateInfo& CreateInfo)
+{
+	return RHICreateTexture3D(SizeX, SizeY, SizeZ, Format, NumMips, Flags, CreateInfo);
+}
+
 void FD3D11DynamicRHI::RHIGetResourceInfo(FTextureRHIParamRef Ref, FRHIResourceInfo& OutInfo)
 {
 	if(Ref)
@@ -1274,6 +1313,14 @@ FShaderResourceViewRHIRef FD3D11DynamicRHI::RHICreateShaderResourceView(FTexture
 	return new FD3D11ShaderResourceView(ShaderResourceView,Texture2D);
 }
 
+FShaderResourceViewRHIRef FD3D11DynamicRHI::RHICreateShaderResourceView_RenderThread(
+	class FRHICommandListImmediate& RHICmdList,
+	FTexture2DRHIParamRef Texture2DRHI,
+	uint8 MipLevel)
+{
+	return RHICreateShaderResourceView(Texture2DRHI, MipLevel);
+}
+
 FShaderResourceViewRHIRef FD3D11DynamicRHI::RHICreateShaderResourceView(FTexture2DRHIParamRef Texture2DRHI, uint8 MipLevel, uint8 NumMipLevels, uint8 Format)
 {
 	FD3D11Texture2D* Texture2D = ResourceCast(Texture2DRHI);
@@ -1318,6 +1365,16 @@ FShaderResourceViewRHIRef FD3D11DynamicRHI::RHICreateShaderResourceView(FTexture
 	return new FD3D11ShaderResourceView(ShaderResourceView,Texture2D);
 }
 
+FShaderResourceViewRHIRef FD3D11DynamicRHI::RHICreateShaderResourceView_RenderThread(
+	class FRHICommandListImmediate& RHICmdList,
+	FTexture2DRHIParamRef Texture2DRHI,
+	uint8 MipLevel,
+	uint8 NumMipLevels,
+	uint8 Format)
+{
+	return RHICreateShaderResourceView(Texture2DRHI, MipLevel, NumMipLevels, Format);
+}
+
 FShaderResourceViewRHIRef FD3D11DynamicRHI::RHICreateShaderResourceView(FTexture3DRHIParamRef Texture3DRHI, uint8 MipLevel)
 {
 	FD3D11Texture3D* Texture3D = ResourceCast(Texture3DRHI);
@@ -1341,6 +1398,14 @@ FShaderResourceViewRHIRef FD3D11DynamicRHI::RHICreateShaderResourceView(FTexture
 	return new FD3D11ShaderResourceView(ShaderResourceView,Texture3D);
 }
 
+FShaderResourceViewRHIRef FD3D11DynamicRHI::RHICreateShaderResourceView_RenderThread(
+	class FRHICommandListImmediate& RHICmdList,
+	FTexture3DRHIParamRef Texture3DRHI,
+	uint8 MipLevel)
+{
+	return RHICreateShaderResourceView(Texture3DRHI, MipLevel);
+}
+
 FShaderResourceViewRHIRef FD3D11DynamicRHI::RHICreateShaderResourceView(FTexture2DArrayRHIParamRef Texture2DArrayRHI, uint8 MipLevel)
 {
 	FD3D11Texture2DArray* Texture2DArray = ResourceCast(Texture2DArrayRHI);
@@ -1366,6 +1431,14 @@ FShaderResourceViewRHIRef FD3D11DynamicRHI::RHICreateShaderResourceView(FTexture
 	return new FD3D11ShaderResourceView(ShaderResourceView,Texture2DArray);
 }
 
+FShaderResourceViewRHIRef FD3D11DynamicRHI::RHICreateShaderResourceView_RenderThread(
+	class FRHICommandListImmediate& RHICmdList,
+	FTexture2DArrayRHIParamRef Texture2DArrayRHI,
+	uint8 MipLevel)
+{
+	return RHICreateShaderResourceView(Texture2DArrayRHI, MipLevel);
+}
+
 FShaderResourceViewRHIRef FD3D11DynamicRHI::RHICreateShaderResourceView(FTextureCubeRHIParamRef TextureCubeRHI, uint8 MipLevel)
 {
 	FD3D11TextureCube* TextureCube = ResourceCast(TextureCubeRHI);
@@ -1389,6 +1462,14 @@ FShaderResourceViewRHIRef FD3D11DynamicRHI::RHICreateShaderResourceView(FTexture
 	return new FD3D11ShaderResourceView(ShaderResourceView,TextureCube);
 }
 
+FShaderResourceViewRHIRef FD3D11DynamicRHI::RHICreateShaderResourceView_RenderThread(
+	class FRHICommandListImmediate& RHICmdList,
+	FTextureCubeRHIParamRef TextureCubeRHI,
+	uint8 MipLevel)
+{
+	return RHICreateShaderResourceView(TextureCubeRHI, MipLevel);
+}
+
 /** Generates mip maps for the surface. */
 void FD3D11DynamicRHI::RHIGenerateMips(FTextureRHIParamRef TextureRHI)
 {
@@ -1417,6 +1498,48 @@ uint32 FD3D11DynamicRHI::RHIComputeMemorySize(FTextureRHIParamRef TextureRHI)
 	return Texture->GetMemorySize();
 }
 
+/**
+ * Asynchronous texture copy helper
+ *
+ * @param NewTexture2DRHI		- Texture to reallocate
+ * @param Texture2DRHI			- Texture to reallocate
+ * @param NewMipCount			- New number of mip-levels
+ * @param NewSizeX				- New width, in pixels
+ * @param NewSizeY				- New height, in pixels
+ * @param RequestStatus			- Will be decremented by 1 when the reallocation is complete (success or failure).
+ * @return						- New reference to the texture, or an invalid reference upon failure
+ */
+void FD3D11DynamicRHI::RHIAsyncCopyTexture2DCopy(FTexture2DRHIParamRef NewTexture2DRHI, FTexture2DRHIParamRef Texture2DRHI, int32 NewMipCount, int32 NewSizeX, int32 NewSizeY, FThreadSafeCounter* RequestStatus)
+{
+	FD3D11Texture2D* Texture2D = ResourceCast(Texture2DRHI);
+	FD3D11Texture2D* NewTexture2D = ResourceCast(NewTexture2DRHI);
+
+	// Use the GPU to asynchronously copy the old mip-maps into the new texture.
+	const uint32 NumSharedMips = FMath::Min(Texture2D->GetNumMips(), NewTexture2D->GetNumMips());
+	const uint32 SourceMipOffset = Texture2D->GetNumMips() - NumSharedMips;
+	const uint32 DestMipOffset = NewTexture2D->GetNumMips() - NumSharedMips;
+	for (uint32 MipIndex = 0; MipIndex < NumSharedMips; ++MipIndex)
+	{
+		// Use the GPU to copy between mip-maps.
+		// This is serialized with other D3D commands, so it isn't necessary to increment Counter to signal a pending asynchronous copy.
+		Direct3DDeviceIMContext->CopySubresourceRegion(
+			NewTexture2D->GetResource(),
+			D3D11CalcSubresource(MipIndex + DestMipOffset, 0, NewTexture2D->GetNumMips()),
+			0,
+			0,
+			0,
+			Texture2D->GetResource(),
+			D3D11CalcSubresource(MipIndex + SourceMipOffset, 0, Texture2D->GetNumMips()),
+			NULL
+		);
+	}
+
+	// Decrement the thread-safe counter used to track the completion of the reallocation, since D3D handles sequencing the
+	// async mip copies with other D3D calls.
+	RequestStatus->Decrement();
+}
+
+
 /**
  * Starts an asynchronous texture reallocation. It may complete immediately if the reallocation
  * could be performed without any reshuffling of texture memory, or if there isn't enough memory.
@@ -1439,31 +1562,38 @@ FTexture2DRHIRef FD3D11DynamicRHI::RHIAsyncReallocateTexture2D(FTexture2DRHIPara
 	// Allocate a new texture.
 	FRHIResourceCreateInfo CreateInfo;
 	FD3D11Texture2D* NewTexture2D = CreateD3D11Texture2D<FD3D11BaseTexture2D>(NewSizeX,NewSizeY,1,false,false,Texture2D->GetFormat(),NewMipCount,1,Texture2D->GetFlags(),CreateInfo);
-	
-	// Use the GPU to asynchronously copy the old mip-maps into the new texture.
-	const uint32 NumSharedMips = FMath::Min(Texture2D->GetNumMips(),NewTexture2D->GetNumMips());
-	const uint32 SourceMipOffset = Texture2D->GetNumMips()    - NumSharedMips;
-	const uint32 DestMipOffset   = NewTexture2D->GetNumMips() - NumSharedMips;
-	for(uint32 MipIndex = 0;MipIndex < NumSharedMips;++MipIndex)
+
+	RHIAsyncCopyTexture2DCopy(NewTexture2D, Texture2D, NewMipCount, NewSizeX, NewSizeY, RequestStatus);
+
+	return NewTexture2D;
+}
+
+FTexture2DRHIRef FD3D11DynamicRHI::AsyncReallocateTexture2D_RenderThread(
+	class FRHICommandListImmediate& RHICmdList,
+	FTexture2DRHIParamRef Texture2D,
+	int32 NewMipCount,
+	int32 NewSizeX,
+	int32 NewSizeY,
+	FThreadSafeCounter* RequestStatus)
+{
+	FTexture2DRHIRef NewTexture2D;
+
+	if (ShouldNotEnqueueRHICommand())
 	{
-		// Use the GPU to copy between mip-maps.
-		// This is serialized with other D3D commands, so it isn't necessary to increment Counter to signal a pending asynchronous copy.
-		Direct3DDeviceIMContext->CopySubresourceRegion(
-			NewTexture2D->GetResource(),
-			D3D11CalcSubresource(MipIndex + DestMipOffset,0,NewTexture2D->GetNumMips()),
-			0,
-			0,
-			0,
-			Texture2D->GetResource(),
-			D3D11CalcSubresource(MipIndex + SourceMipOffset,0,Texture2D->GetNumMips()),
-			NULL
-			);
+		NewTexture2D = RHIAsyncReallocateTexture2D(Texture2D, NewMipCount, NewSizeX, NewSizeY, RequestStatus);
 	}
+	else
+	{
+		// Allocate a new texture.
+		FRHIResourceCreateInfo CreateInfo;
+		FD3D11Texture2D* NewTexture2DPointer = CreateD3D11Texture2D<FD3D11BaseTexture2D>(NewSizeX, NewSizeY, 1, false, false, Texture2D->GetFormat(), NewMipCount, 1, Texture2D->GetFlags(), CreateInfo);
+		NewTexture2D = NewTexture2DPointer;
 
-	// Decrement the thread-safe counter used to track the completion of the reallocation, since D3D handles sequencing the
-	// async mip copies with other D3D calls.
-	RequestStatus->Decrement();
-
+		RunOnRHIThread([this, NewTexture2D, Texture2D, NewMipCount, NewSizeX, NewSizeY, RequestStatus]()
+		{
+			RHIAsyncCopyTexture2DCopy(NewTexture2D, Texture2D, NewMipCount, NewSizeX, NewSizeY, RequestStatus);
+		});
+	}
 	return NewTexture2D;
 }
 
@@ -1481,6 +1611,14 @@ ETextureReallocationStatus FD3D11DynamicRHI::RHIFinalizeAsyncReallocateTexture2D
 	return TexRealloc_Succeeded;
 }
 
+ETextureReallocationStatus FD3D11DynamicRHI::FinalizeAsyncReallocateTexture2D_RenderThread(
+	class FRHICommandListImmediate& RHICmdList,
+	FTexture2DRHIParamRef Texture2D,
+	bool bBlockUntilCompleted)
+{
+	return RHIFinalizeAsyncReallocateTexture2D(Texture2D, bBlockUntilCompleted);
+}
+
 /**
  * Cancels an async reallocation for the specified texture.
  * This should be called for the new texture, not the original.
@@ -1494,6 +1632,14 @@ ETextureReallocationStatus FD3D11DynamicRHI::RHICancelAsyncReallocateTexture2D(
 	return TexRealloc_Succeeded;
 }
 
+ETextureReallocationStatus FD3D11DynamicRHI::CancelAsyncReallocateTexture2D_RenderThread(
+	class FRHICommandListImmediate& RHICmdList,
+	FTexture2DRHIParamRef Texture2D,
+	bool bBlockUntilCompleted)
+{
+	return RHICancelAsyncReallocateTexture2D(Texture2D, bBlockUntilCompleted);
+}
+
 template<typename RHIResourceType>
 void* TD3D11Texture2D<RHIResourceType>::Lock(uint32 MipIndex,uint32 ArrayIndex,EResourceLockMode LockMode,uint32& DestStride)
 {
@@ -1583,7 +1729,7 @@ void* TD3D11Texture2D<RHIResourceType>::Lock(uint32 MipIndex,uint32 ArrayIndex,E
 	}
 
 	// Add the lock to the outstanding lock list.
-	D3DRHI->OutstandingLocks.Add(FD3D11LockedKey(GetResource(),Subresource),LockedData);
+	D3DRHI->GetThreadLocalLockTracker().Add(FD3D11LockedKey(GetResource(),Subresource),LockedData);
 
 	return (void*)LockedData.GetData();
 }
@@ -1597,9 +1743,10 @@ void TD3D11Texture2D<RHIResourceType>::Unlock(uint32 MipIndex,uint32 ArrayIndex)
 	const uint32 Subresource = D3D11CalcSubresource(MipIndex,ArrayIndex,this->GetNumMips());
 
 	// Find the object that is tracking this lock
+	FD3D11DynamicRHI::FD3D11LockTracker& OutstandingLocks = D3DRHI->GetThreadLocalLockTracker();
 	const FD3D11LockedKey LockedKey(GetResource(),Subresource);
-	FD3D11LockedData* LockedData = D3DRHI->OutstandingLocks.Find(LockedKey);
-	check(LockedData);
+	FD3D11LockedData* LockedData = OutstandingLocks.Find(LockedKey);
+	checkf(LockedData, TEXT("Texture is either not locked or locked on a different thread"));
 
 #if PLATFORM_SUPPORTS_VIRTUAL_TEXTURES
 	if (D3DRHI->HandleSpecialUnlock(MipIndex, GetFlags(), GetResource(), RawTextureMemory))
@@ -1615,12 +1762,12 @@ void TD3D11Texture2D<RHIResourceType>::Unlock(uint32 MipIndex,uint32 ArrayIndex)
 	else if(!LockedData->StagingResource)
 	{
 		// If we're writing, we need to update the subresource
-		D3DRHI->GetDeviceContext()->UpdateSubresource(GetResource(),Subresource,NULL,LockedData->GetData(),LockedData->Pitch,0);
+		D3DRHI->GetDeviceContext()->UpdateSubresource(GetResource(), Subresource, NULL, LockedData->GetData(), LockedData->Pitch, 0);
 		LockedData->FreeData();
 	}
 
 	// Remove the lock from the outstanding lock list.
-	D3DRHI->OutstandingLocks.Remove(LockedKey);
+	OutstandingLocks.Remove(LockedKey);
 }
 
 void* FD3D11DynamicRHI::RHILockTexture2D(FTexture2DRHIParamRef TextureRHI,uint32 MipIndex,EResourceLockMode LockMode,uint32& DestStride,bool bLockWithinMiptail)
@@ -1631,6 +1778,32 @@ void* FD3D11DynamicRHI::RHILockTexture2D(FTexture2DRHIParamRef TextureRHI,uint32
 	return Texture->Lock(MipIndex,0,LockMode,DestStride);
 }
 
+void* FD3D11DynamicRHI::LockTexture2D_RenderThread(
+	class FRHICommandListImmediate& RHICmdList,
+	FTexture2DRHIParamRef Texture,
+	uint32 MipIndex,
+	EResourceLockMode LockMode,
+	uint32& DestStride,
+	bool bLockWithinMiptail,
+	bool bNeedsDefaultRHIFlush)
+{
+	void *LockedTexture = nullptr;
+
+	if (ShouldNotEnqueueRHICommand())
+	{
+		LockedTexture = RHILockTexture2D(Texture, MipIndex, LockMode, DestStride, bLockWithinMiptail);
+	}
+	else
+	{
+		RunOnRHIThread([this, Texture, MipIndex, LockMode, &DestStride, bLockWithinMiptail, &LockedTexture]()
+		{
+			LockedTexture = RHILockTexture2D(Texture, MipIndex, LockMode, DestStride, bLockWithinMiptail);
+		});
+		FRHICommandListExecutor::GetImmediateCommandList().ImmediateFlush(EImmediateFlushType::FlushRHIThread);
+	}
+	return LockedTexture;
+}
+
 void FD3D11DynamicRHI::RHIUnlockTexture2D(FTexture2DRHIParamRef TextureRHI,uint32 MipIndex,bool bLockWithinMiptail)
 {
 	check(TextureRHI);
@@ -1638,6 +1811,26 @@ void FD3D11DynamicRHI::RHIUnlockTexture2D(FTexture2DRHIParamRef TextureRHI,uint3
 	Texture->Unlock(MipIndex,0);
 }
 
+void FD3D11DynamicRHI::UnlockTexture2D_RenderThread(
+	class FRHICommandListImmediate& RHICmdList,
+	FTexture2DRHIParamRef Texture,
+	uint32 MipIndex,
+	bool bLockWithinMiptail,
+	bool bNeedsDefaultRHIFlush)
+{
+	if (ShouldNotEnqueueRHICommand())
+	{
+		RHIUnlockTexture2D(Texture, MipIndex, bLockWithinMiptail);
+	}
+	else
+	{
+		RunOnRHIThread([this, Texture, MipIndex, bLockWithinMiptail]()
+		{
+			RHIUnlockTexture2D(Texture, MipIndex, bLockWithinMiptail);
+		});
+	}
+}
+
 void* FD3D11DynamicRHI::RHILockTexture2DArray(FTexture2DArrayRHIParamRef TextureRHI,uint32 TextureIndex,uint32 MipIndex,EResourceLockMode LockMode,uint32& DestStride,bool bLockWithinMiptail)
 {
 	FD3D11Texture2DArray* Texture = ResourceCast(TextureRHI);
@@ -1653,12 +1846,12 @@ void FD3D11DynamicRHI::RHIUnlockTexture2DArray(FTexture2DArrayRHIParamRef Textur
 
 void FD3D11DynamicRHI::RHIUpdateTexture2D(FTexture2DRHIParamRef TextureRHI,uint32 MipIndex,const FUpdateTextureRegion2D& UpdateRegion,uint32 SourcePitch,const uint8* SourceData)
 {
-    FD3D11Texture2D* Texture = ResourceCast(TextureRHI);
+	FD3D11Texture2D* Texture = ResourceCast(TextureRHI);
 
-    D3D11_BOX DestBox =
+	D3D11_BOX DestBox =
 	{
 		UpdateRegion.DestX,                      UpdateRegion.DestY,                       0,
-        UpdateRegion.DestX + UpdateRegion.Width, UpdateRegion.DestY + UpdateRegion.Height, 1
+		UpdateRegion.DestX + UpdateRegion.Width, UpdateRegion.DestY + UpdateRegion.Height, 1
 	};
 
 	check(UpdateRegion.Width % GPixelFormats[Texture->GetFormat()].BlockSizeX == 0);
@@ -1671,6 +1864,27 @@ void FD3D11DynamicRHI::RHIUpdateTexture2D(FTexture2DRHIParamRef TextureRHI,uint3
 	Direct3DDeviceIMContext->UpdateSubresource(Texture->GetResource(), MipIndex, &DestBox, SourceData, SourcePitch, 0);
 }
 
+void FD3D11DynamicRHI::UpdateTexture2D_RenderThread(
+	class FRHICommandListImmediate& RHICmdList,
+	FTexture2DRHIParamRef Texture,
+	uint32 MipIndex,
+	const struct FUpdateTextureRegion2D& UpdateRegion,
+	uint32 SourcePitch,
+	const uint8* SourceData)
+{
+	if (ShouldNotEnqueueRHICommand())
+	{
+		RHIUpdateTexture2D(Texture, MipIndex, UpdateRegion, SourcePitch, SourceData);
+	}
+	else
+	{
+		RunOnRHIThread([this, Texture, MipIndex, UpdateRegion, SourcePitch, SourceData]()
+		{
+			RHIUpdateTexture2D(Texture, MipIndex, UpdateRegion, SourcePitch, SourceData);
+		});
+	}
+}
+
 void FD3D11DynamicRHI::RHIUpdateTexture3D(FTexture3DRHIParamRef TextureRHI,uint32 MipIndex,const FUpdateTextureRegion3D& UpdateRegion,uint32 SourceRowPitch,uint32 SourceDepthPitch,const uint8* SourceData)
 {
 	FD3D11Texture3D* Texture = ResourceCast(TextureRHI);
@@ -1690,6 +1904,57 @@ void FD3D11DynamicRHI::RHIUpdateTexture3D(FTexture3DRHIParamRef TextureRHI,uint3
 	Direct3DDeviceIMContext->UpdateSubresource(Texture->GetResource(), MipIndex, &DestBox, SourceData, SourceRowPitch, SourceDepthPitch);
 }
 
+void FD3D11DynamicRHI::EndUpdateTexture3D_RenderThread(class FRHICommandListImmediate& RHICmdList, FUpdateTexture3DData& UpdateData)
+{
+	if (RHICmdList.Bypass())
+	{
+		RHIUpdateTexture3D(UpdateData.Texture, UpdateData.MipIndex, UpdateData.UpdateRegion, UpdateData.RowPitch, UpdateData.DepthPitch, UpdateData.Data);
+		FMemory::Free(UpdateData.Data);
+	}
+	else
+	{
+		UpdateData.Texture->AddRef();
+		RunOnRHIThread(
+			[UpdateData]()
+		{
+			static_cast<FD3D11DynamicRHI*>(GDynamicRHI)->RHIUpdateTexture3D(
+				UpdateData.Texture,
+				UpdateData.MipIndex,
+				UpdateData.UpdateRegion,
+				UpdateData.RowPitch,
+				UpdateData.DepthPitch,
+				UpdateData.Data);
+			UpdateData.Texture->Release();
+			FMemory::Free(UpdateData.Data);
+		});
+		RHICmdList.RHIThreadFence(true);
+	}
+
+	UpdateData.Data = nullptr;
+}
+
+void FD3D11DynamicRHI::UpdateTexture3D_RenderThread(
+	class FRHICommandListImmediate& RHICmdList,
+	FTexture3DRHIParamRef Texture,
+	uint32 MipIndex,
+	const struct FUpdateTextureRegion3D& UpdateRegion,
+	uint32 SourceRowPitch,
+	uint32 SourceDepthPitch,
+	const uint8* SourceData)
+{
+	if (ShouldNotEnqueueRHICommand())
+	{
+		RHIUpdateTexture3D(Texture, MipIndex, UpdateRegion, SourceRowPitch, SourceDepthPitch, SourceData);
+	}
+	else
+	{
+		RunOnRHIThread([this, Texture, MipIndex, UpdateRegion, SourceRowPitch, SourceDepthPitch, SourceData]()
+		{
+			RHIUpdateTexture3D(Texture, MipIndex, UpdateRegion, SourceRowPitch, SourceDepthPitch, SourceData);
+		});
+	}
+}
+
 /*-----------------------------------------------------------------------------
 	Cubemap texture support.
 -----------------------------------------------------------------------------*/
@@ -1698,11 +1963,34 @@ FTextureCubeRHIRef FD3D11DynamicRHI::RHICreateTextureCube(uint32 Size, uint8 For
 	return CreateD3D11Texture2D<FD3D11BaseTextureCube>(Size,Size,6,false,true,Format,NumMips,1,Flags,CreateInfo);
 }
 
+FTextureCubeRHIRef FD3D11DynamicRHI::RHICreateTextureCube_RenderThread(
+	class FRHICommandListImmediate& RHICmdList,
+	uint32 Size,
+	uint8 Format,
+	uint32 NumMips,
+	uint32 Flags,
+	FRHIResourceCreateInfo& CreateInfo)
+{
+	return RHICreateTextureCube(Size, Format, NumMips, Flags, CreateInfo);
+}
+
 FTextureCubeRHIRef FD3D11DynamicRHI::RHICreateTextureCubeArray(uint32 Size, uint32 ArraySize, uint8 Format, uint32 NumMips, uint32 Flags, FRHIResourceCreateInfo& CreateInfo)
 {
 	return CreateD3D11Texture2D<FD3D11BaseTextureCube>(Size,Size,6 * ArraySize,true,true,Format,NumMips,1,Flags,CreateInfo);
 }
 
+FTextureCubeRHIRef FD3D11DynamicRHI::RHICreateTextureCubeArray_RenderThread(
+	class FRHICommandListImmediate& RHICmdList,
+	uint32 Size,
+	uint32 ArraySize,
+	uint8 Format,
+	uint32 NumMips,
+	uint32 Flags,
+	FRHIResourceCreateInfo& CreateInfo)
+{
+	return RHICreateTextureCubeArray(Size, ArraySize, Format, NumMips, Flags, CreateInfo);
+}
+
 void* FD3D11DynamicRHI::RHILockTextureCubeFace(FTextureCubeRHIParamRef TextureCubeRHI,uint32 FaceIndex,uint32 ArrayIndex,uint32 MipIndex,EResourceLockMode LockMode,uint32& DestStride,bool bLockWithinMiptail)
 {
 	FD3D11TextureCube* TextureCube = ResourceCast(TextureCubeRHI);
@@ -1710,6 +1998,34 @@ void* FD3D11DynamicRHI::RHILockTextureCubeFace(FTextureCubeRHIParamRef TextureCu
 	uint32 D3DFace = GetD3D11CubeFace((ECubeFace)FaceIndex);
 	return TextureCube->Lock(MipIndex,D3DFace + ArrayIndex * 6,LockMode,DestStride);
 }
+
+void* FD3D11DynamicRHI::RHILockTextureCubeFace_RenderThread(
+	class FRHICommandListImmediate& RHICmdList,
+	FTextureCubeRHIParamRef Texture,
+	uint32 FaceIndex,
+	uint32 ArrayIndex,
+	uint32 MipIndex,
+	EResourceLockMode LockMode,
+	uint32& DestStride,
+	bool bLockWithinMiptail)
+{
+	void *LockedTexture = nullptr;
+
+	if (ShouldNotEnqueueRHICommand())
+	{
+		LockedTexture = RHILockTextureCubeFace(Texture, FaceIndex, ArrayIndex, MipIndex, LockMode, DestStride, bLockWithinMiptail);
+	}
+	else
+	{
+		RunOnRHIThread([this, Texture, FaceIndex, ArrayIndex, MipIndex, LockMode, &DestStride, bLockWithinMiptail, &LockedTexture]()
+		{
+			LockedTexture = RHILockTextureCubeFace(Texture, FaceIndex, ArrayIndex, MipIndex, LockMode, DestStride, bLockWithinMiptail);
+		});
+		FRHICommandListExecutor::GetImmediateCommandList().ImmediateFlush(EImmediateFlushType::FlushRHIThread);
+	}
+	return LockedTexture;
+}
+
 void FD3D11DynamicRHI::RHIUnlockTextureCubeFace(FTextureCubeRHIParamRef TextureCubeRHI,uint32 FaceIndex,uint32 ArrayIndex,uint32 MipIndex,bool bLockWithinMiptail)
 {
 	FD3D11TextureCube* TextureCube = ResourceCast(TextureCubeRHI);
@@ -1717,6 +2033,27 @@ void FD3D11DynamicRHI::RHIUnlockTextureCubeFace(FTextureCubeRHIParamRef TextureC
 	TextureCube->Unlock(MipIndex,D3DFace + ArrayIndex * 6);
 }
 
+void FD3D11DynamicRHI::RHIUnlockTextureCubeFace_RenderThread(
+	class FRHICommandListImmediate& RHICmdList,
+	FTextureCubeRHIParamRef Texture,
+	uint32 FaceIndex,
+	uint32 ArrayIndex,
+	uint32 MipIndex,
+	bool bLockWithinMiptail)
+{
+	if (ShouldNotEnqueueRHICommand())
+	{
+		RHIUnlockTextureCubeFace(Texture, FaceIndex, ArrayIndex, MipIndex, bLockWithinMiptail);
+	}
+	else
+	{
+		RunOnRHIThread([this, Texture, FaceIndex, ArrayIndex, MipIndex, bLockWithinMiptail]()
+		{
+			RHIUnlockTextureCubeFace(Texture, FaceIndex, ArrayIndex, MipIndex, bLockWithinMiptail);
+		});
+	}
+}
+
 void FD3D11DynamicRHI::RHIBindDebugLabelName(FTextureRHIParamRef TextureRHI, const TCHAR* Name)
 {
 	//todo: require names at texture creation time.
@@ -1751,6 +2088,13 @@ FTextureReferenceRHIRef FD3D11DynamicRHI::RHICreateTextureReference(FLastRenderT
 	return new FD3D11TextureReference(this,LastRenderTime);
 }
 
+FTextureReferenceRHIRef FD3D11DynamicRHI::RHICreateTextureReference_RenderThread(
+	class FRHICommandListImmediate& RHICmdList,
+	FLastRenderTimeContainer* LastRenderTime)
+{
+	return RHICreateTextureReference(LastRenderTime);
+}
+
 void FD3D11DynamicRHI::RHICopySubTextureRegion(FTexture2DRHIParamRef SourceTextureRHI, FTexture2DRHIParamRef DestinationTextureRHI, FBox2D SourceBox, FBox2D DestinationBox)
 {
 	FD3D11Texture2D* SourceTexture = ResourceCast(SourceTextureRHI);
@@ -1832,6 +2176,26 @@ void FD3D11DynamicRHI::RHICopySubTextureRegion(FTexture2DRHIParamRef SourceTextu
 	Direct3DDeviceIMContext->CopySubresourceRegion(DestinationRessource, 0, DestinationBox.Min.X + DestinationOffsetX, DestinationBox.Min.Y + DestinationOffsetY, 0, SourceTexture->GetResource(), 0, &SourceBoxAdjust);
 }
 
+void FD3D11DynamicRHI::RHICopySubTextureRegion_RenderThread(
+	class FRHICommandListImmediate& RHICmdList,
+	FTexture2DRHIParamRef SourceTexture,
+	FTexture2DRHIParamRef DestinationTexture,
+	FBox2D SourceBox,
+	FBox2D DestinationBox)
+{
+	if (ShouldNotEnqueueRHICommand())
+	{
+		RHICopySubTextureRegion(SourceTexture, DestinationTexture, SourceBox, DestinationBox);
+	}
+	else
+	{
+		RunOnRHIThread([this, SourceTexture, DestinationTexture, SourceBox, DestinationBox]()
+		{
+			RHICopySubTextureRegion(SourceTexture, DestinationTexture, SourceBox, DestinationBox);
+		});
+	}
+}
+
 void FD3D11DynamicRHI::RHIUpdateTextureReference(FTextureReferenceRHIParamRef TextureRefRHI, FTextureRHIParamRef NewTextureRHI)
 {
 	// Updating texture references is disallowed while the RHI could be caching them in referenced resource tables.
diff --git a/Engine/Source/Runtime/Windows/D3D11RHI/Private/D3D11UAV.cpp b/Engine/Source/Runtime/Windows/D3D11RHI/Private/D3D11UAV.cpp
--- a/Engine/Source/Runtime/Windows/D3D11RHI/Private/D3D11UAV.cpp
+++ b/Engine/Source/Runtime/Windows/D3D11RHI/Private/D3D11UAV.cpp
@@ -52,6 +52,15 @@ FUnorderedAccessViewRHIRef FD3D11DynamicRHI::RHICreateUnorderedAccessView(FStruc
 	return new FD3D11UnorderedAccessView(UnorderedAccessView,StructuredBuffer);
 }
 
+FUnorderedAccessViewRHIRef FD3D11DynamicRHI::RHICreateUnorderedAccessView_RenderThread(
+	class FRHICommandListImmediate& RHICmdList,
+	FStructuredBufferRHIParamRef StructuredBuffer,
+	bool bUseUAVCounter,
+	bool bAppendBuffer)
+{
+	return RHICreateUnorderedAccessView(StructuredBuffer, bUseUAVCounter, bAppendBuffer);
+}
+
 FUnorderedAccessViewRHIRef FD3D11DynamicRHI::RHICreateUnorderedAccessView(FTextureRHIParamRef TextureRHI, uint32 MipLevel)
 {
 	FD3D11TextureBase* Texture = GetD3D11TextureFromRHITexture(TextureRHI);
@@ -96,6 +105,14 @@ FUnorderedAccessViewRHIRef FD3D11DynamicRHI::RHICreateUnorderedAccessView(FTextu
 	return new FD3D11UnorderedAccessView(UnorderedAccessView,Texture);
 }
 
+FUnorderedAccessViewRHIRef FD3D11DynamicRHI::RHICreateUnorderedAccessView_RenderThread(
+	class FRHICommandListImmediate& RHICmdList,
+	FTextureRHIParamRef Texture,
+	uint32 MipLevel)
+{
+	return RHICreateUnorderedAccessView(Texture, MipLevel);
+}
+
 FUnorderedAccessViewRHIRef FD3D11DynamicRHI::RHICreateUnorderedAccessView(FVertexBufferRHIParamRef VertexBufferRHI, uint8 Format)
 {
 	FD3D11VertexBuffer* VertexBuffer = ResourceCast(VertexBufferRHI);
@@ -125,6 +142,14 @@ FUnorderedAccessViewRHIRef FD3D11DynamicRHI::RHICreateUnorderedAccessView(FVerte
 	return new FD3D11UnorderedAccessView(UnorderedAccessView,VertexBuffer);
 }
 
+FUnorderedAccessViewRHIRef FD3D11DynamicRHI::RHICreateUnorderedAccessView_RenderThread(
+	class FRHICommandListImmediate& RHICmdList,
+	FVertexBufferRHIParamRef VertexBuffer,
+	uint8 Format)
+{
+	return RHICreateUnorderedAccessView(VertexBuffer, Format);
+}
+
 FUnorderedAccessViewRHIRef FD3D11DynamicRHI::RHICreateUnorderedAccessView(FIndexBufferRHIParamRef IndexBufferRHI, uint8 Format)
 {
 	FD3D11IndexBuffer* IndexBuffer = ResourceCast(IndexBufferRHI);
@@ -154,6 +179,14 @@ FUnorderedAccessViewRHIRef FD3D11DynamicRHI::RHICreateUnorderedAccessView(FIndex
 	return new FD3D11UnorderedAccessView(UnorderedAccessView, IndexBuffer);
 }
 
+FUnorderedAccessViewRHIRef FD3D11DynamicRHI::RHICreateUnorderedAccessView_RenderThread(
+	class FRHICommandListImmediate& RHICmdList,
+	FIndexBufferRHIParamRef IndexBuffer,
+	uint8 Format)
+{
+	return RHICreateUnorderedAccessView(IndexBuffer, Format);
+}
+
 FShaderResourceViewRHIRef FD3D11DynamicRHI::RHICreateShaderResourceView(FStructuredBufferRHIParamRef StructuredBufferRHI)
 {
 	FD3D11StructuredBuffer* StructuredBuffer = ResourceCast(StructuredBufferRHI);
@@ -188,6 +221,13 @@ FShaderResourceViewRHIRef FD3D11DynamicRHI::RHICreateShaderResourceView(FStructu
 	return new FD3D11ShaderResourceView(ShaderResourceView,StructuredBuffer);
 }
 
+FShaderResourceViewRHIRef FD3D11DynamicRHI::RHICreateShaderResourceView_RenderThread(
+	class FRHICommandListImmediate& RHICmdList,
+	FStructuredBufferRHIParamRef StructuredBuffer)
+{
+	return RHICreateShaderResourceView(StructuredBuffer);
+}
+
 static void CreateD3D11ShaderResourceViewOnVertexBuffer(ID3D11Device* Direct3DDevice, ID3D11Buffer* Buffer, uint32 Stride, uint8 Format, ID3D11ShaderResourceView** OutSRV)
 {
 	D3D11_BUFFER_DESC BufferDesc;
@@ -234,6 +274,15 @@ FShaderResourceViewRHIRef FD3D11DynamicRHI::RHICreateShaderResourceView(FVertexB
 	return new FD3D11ShaderResourceView(ShaderResourceView,VertexBuffer);
 }
 
+FShaderResourceViewRHIRef FD3D11DynamicRHI::CreateShaderResourceView_RenderThread(
+	class FRHICommandListImmediate& RHICmdList,
+	FVertexBufferRHIParamRef VertexBuffer,
+	uint32 Stride,
+	uint8 Format)
+{
+	return RHICreateShaderResourceView(VertexBuffer, Stride, Format);
+}
+
 void FD3D11DynamicRHI::RHIUpdateShaderResourceView(FShaderResourceViewRHIParamRef SRV, FVertexBufferRHIParamRef VertexBufferRHI, uint32 Stride, uint8 Format)
 {
 	check(SRV);
@@ -298,6 +347,13 @@ FShaderResourceViewRHIRef FD3D11DynamicRHI::RHICreateShaderResourceView(FIndexBu
 	return new FD3D11ShaderResourceView(ShaderResourceView, Buffer);
 }
 
+FShaderResourceViewRHIRef FD3D11DynamicRHI::CreateShaderResourceView_RenderThread(
+	class FRHICommandListImmediate& RHICmdList,
+	FIndexBufferRHIParamRef Buffer)
+{
+	return RHICreateShaderResourceView(Buffer);
+}
+
 void FD3D11DynamicRHI::RHIClearTinyUAV(FUnorderedAccessViewRHIParamRef UnorderedAccessViewRHI, const uint32* Values)
 {
 	FD3D11UnorderedAccessView* UnorderedAccessView = ResourceCast(UnorderedAccessViewRHI);
diff --git a/Engine/Source/Runtime/Windows/D3D11RHI/Private/D3D11UniformBuffer.cpp b/Engine/Source/Runtime/Windows/D3D11RHI/Private/D3D11UniformBuffer.cpp
--- a/Engine/Source/Runtime/Windows/D3D11RHI/Private/D3D11UniformBuffer.cpp
+++ b/Engine/Source/Runtime/Windows/D3D11RHI/Private/D3D11UniformBuffer.cpp
@@ -69,6 +69,7 @@ void UniformBufferBeginFrame()
 				DEC_MEMORY_STAT_BY(STAT_D3D11FreeUniformBufferMemory, PoolEntry.CreatedSize);
 				NumCleaned++;
 				UpdateBufferStats(PoolEntry.Buffer, false);
+				PoolEntry.Buffer.SafeRelease();
 				UniformBufferPool[BucketIndex].RemoveAtSwap(EntryIndex);
 			}
 		}
@@ -95,18 +96,68 @@ void UniformBufferBeginFrame()
 
 static bool IsPoolingEnabled()
 {
-	if (IsRunningRHIInSeparateThread() && IsInRenderingThread() && GRHICommandList.IsRHIThreadActive())
-	{
-		return false; // we can't currently use pooling if the RHI thread is active. 
-	}
 	static const auto CVar = IConsoleManager::Get().FindTConsoleVariableDataInt(TEXT("r.UniformBufferPooling"));
 	int32 CVarValue = CVar->GetValueOnRenderThread();
 	return CVarValue != 0;
 };
 
-FUniformBufferRHIRef FD3D11DynamicRHI::RHICreateUniformBuffer(const void* Contents,const FRHIUniformBufferLayout& Layout,EUniformBufferUsage Usage, EUniformBufferValidation Validation)
+static TRefCountPtr<ID3D11Buffer> CreateAndUpdatePooledUniformBuffer(
+	FD3D11Device* Device,
+	FD3D11DeviceContext* Context,
+	const void* Contents,
+	uint32 NumBytes)
 {
-	check(IsInRenderingThread());
+	TRefCountPtr<ID3D11Buffer> UniformBufferResource;
+
+	// Find the appropriate bucket based on size
+	const uint32 BucketIndex = GetPoolBucketIndex(NumBytes);
+	TArray<FPooledUniformBuffer>& PoolBucket = UniformBufferPool[BucketIndex];
+
+	if (PoolBucket.Num() > 0)
+	{
+		// Reuse the last entry in this size bucket
+		FPooledUniformBuffer FreeBufferEntry = PoolBucket.Pop();
+		check(IsValidRef(FreeBufferEntry.Buffer));
+		checkf(FreeBufferEntry.CreatedSize >= NumBytes, TEXT("%u %u %u %u"), NumBytes, BucketIndex, FreeBufferEntry.CreatedSize, GetPoolBucketSize(NumBytes));
+		DEC_DWORD_STAT(STAT_D3D11NumFreeUniformBuffers);
+		DEC_MEMORY_STAT_BY(STAT_D3D11FreeUniformBufferMemory, FreeBufferEntry.CreatedSize);
+		UniformBufferResource = FreeBufferEntry.Buffer;
+	}
+
+	if (!IsValidRef(UniformBufferResource))
+	{
+		D3D11_BUFFER_DESC Desc;
+		// Allocate based on the bucket size, since this uniform buffer will be reused later
+		Desc.ByteWidth = GetPoolBucketSize(NumBytes);
+		// Use D3D11_USAGE_DYNAMIC, which allows multiple CPU writes for pool reuses
+		// This is method of updating is vastly superior to creating a new constant buffer each time with D3D11_USAGE_IMMUTABLE,
+		// Since that inserts the data into the command buffer which causes GPU flushes
+		Desc.Usage = D3D11_USAGE_DYNAMIC;
+		Desc.BindFlags = D3D11_BIND_CONSTANT_BUFFER;
+		Desc.CPUAccessFlags = D3D11_CPU_ACCESS_WRITE;
+		Desc.MiscFlags = 0;
+		Desc.StructureByteStride = 0;
+
+		VERIFYD3D11RESULT_EX(Device->CreateBuffer(&Desc, NULL, UniformBufferResource.GetInitReference()), Device);
+
+		UpdateBufferStats(UniformBufferResource, true);
+	}
+
+	check(IsValidRef(UniformBufferResource));
+
+	D3D11_MAPPED_SUBRESOURCE MappedSubresource;
+	// Discard previous results since we always do a full update
+	VERIFYD3D11RESULT_EX(Context->Map(UniformBufferResource, 0, D3D11_MAP_WRITE_DISCARD, 0, &MappedSubresource), Device);
+	check(MappedSubresource.RowPitch >= NumBytes);
+	FMemory::Memcpy(MappedSubresource.pData, Contents, NumBytes);
+	Context->Unmap(UniformBufferResource, 0);
+
+	return UniformBufferResource;
+}
+
+FUniformBufferRHIRef FD3D11DynamicRHI::RHICreateUniformBuffer(const void* Contents, const FRHIUniformBufferLayout& Layout, EUniformBufferUsage Usage, EUniformBufferValidation Validation)
+{
+	check(IsInRenderingThread() || IsInRHIThread());
 
 	FD3D11UniformBuffer* NewUniformBuffer = nullptr;
 	const uint32 NumBytes = Layout.ConstantBufferSize;
@@ -119,59 +170,37 @@ FUniformBufferRHIRef FD3D11DynamicRHI::RHICreateUniformBuffer(const void* Conten
 		check(NumBytes < (1 << NumPoolBuckets));
 
 		SCOPE_CYCLE_COUNTER(STAT_D3D11UpdateUniformBufferTime);
+
 		if (IsPoolingEnabled())
 		{
-			TRefCountPtr<ID3D11Buffer> UniformBufferResource;
-			FRingAllocation RingAllocation;
-
-			if (!RingAllocation.IsValid())
+			if (ShouldNotEnqueueRHICommand())
+			{
+				TRefCountPtr<ID3D11Buffer> UniformBufferResource = CreateAndUpdatePooledUniformBuffer(
+					Direct3DDevice.GetReference(),
+					Direct3DDeviceIMContext.GetReference(),
+					Contents,
+					NumBytes);
+				NewUniformBuffer = new FD3D11UniformBuffer(this, Layout, UniformBufferResource, FRingAllocation());
+			}
+			else
 			{
-				// Find the appropriate bucket based on size
-				const uint32 BucketIndex = GetPoolBucketIndex(NumBytes);
-				TArray<FPooledUniformBuffer>& PoolBucket = UniformBufferPool[BucketIndex];
+				NewUniformBuffer = new FD3D11UniformBuffer(this, Layout, nullptr, FRingAllocation());
+				NewUniformBuffer->AddRef();
+				void* CPUContent = FMemory::Malloc(NumBytes);
+				FMemory::Memcpy(CPUContent, Contents, NumBytes);
 
-				if (PoolBucket.Num() > 0)
+				RunOnRHIThread(
+					[NewUniformBuffer, CPUContent, NumBytes]()
 				{
-					// Reuse the last entry in this size bucket
-					FPooledUniformBuffer FreeBufferEntry = PoolBucket.Pop();
-					check(IsValidRef(FreeBufferEntry.Buffer));
-					UniformBufferResource = FreeBufferEntry.Buffer;
-					checkf(FreeBufferEntry.CreatedSize >= NumBytes, TEXT("%u %u %u %u"), NumBytes, BucketIndex, FreeBufferEntry.CreatedSize, GetPoolBucketSize(NumBytes));
-					DEC_DWORD_STAT(STAT_D3D11NumFreeUniformBuffers);
-					DEC_MEMORY_STAT_BY(STAT_D3D11FreeUniformBufferMemory, FreeBufferEntry.CreatedSize);
-				}
-
-				// Nothing usable was found in the free pool, create a new uniform buffer
-				if (!IsValidRef(UniformBufferResource))
-				{
-					D3D11_BUFFER_DESC Desc;
-					// Allocate based on the bucket size, since this uniform buffer will be reused later
-					Desc.ByteWidth = GetPoolBucketSize(NumBytes);
-					// Use D3D11_USAGE_DYNAMIC, which allows multiple CPU writes for pool reuses
-					// This is method of updating is vastly superior to creating a new constant buffer each time with D3D11_USAGE_IMMUTABLE, 
-					// Since that inserts the data into the command buffer which causes GPU flushes
-					Desc.Usage = D3D11_USAGE_DYNAMIC;
-					Desc.BindFlags = D3D11_BIND_CONSTANT_BUFFER;
-					Desc.CPUAccessFlags = D3D11_CPU_ACCESS_WRITE;
-					Desc.MiscFlags = 0;
-					Desc.StructureByteStride = 0;
-
-					VERIFYD3D11RESULT_EX(Direct3DDevice->CreateBuffer(&Desc, NULL, UniformBufferResource.GetInitReference()), Direct3DDevice);
-
-					UpdateBufferStats(UniformBufferResource, true);
-				}
-
-				check(IsValidRef(UniformBufferResource));
-
-				D3D11_MAPPED_SUBRESOURCE MappedSubresource;
-				// Discard previous results since we always do a full update
-				VERIFYD3D11RESULT_EX(Direct3DDeviceIMContext->Map(UniformBufferResource, 0, D3D11_MAP_WRITE_DISCARD, 0, &MappedSubresource), Direct3DDevice);
-				check(MappedSubresource.RowPitch >= NumBytes);
-				FMemory::Memcpy(MappedSubresource.pData, Contents, NumBytes);
-				Direct3DDeviceIMContext->Unmap(UniformBufferResource, 0);
+					NewUniformBuffer->Resource = CreateAndUpdatePooledUniformBuffer(
+						D3D11RHI_DEVICE,
+						D3D11RHI_IMMEDIATE_CONTEXT,
+						CPUContent,
+						NumBytes);
+					NewUniformBuffer->Release();
+					FMemory::Free(CPUContent);
+				});
 			}
-
-			NewUniformBuffer = new FD3D11UniformBuffer(this, Layout, UniformBufferResource, RingAllocation);
 		}
 		else
 		{
@@ -192,6 +221,8 @@ FUniformBufferRHIRef FD3D11DynamicRHI::RHICreateUniformBuffer(const void* Conten
 			VERIFYD3D11RESULT_EX(Direct3DDevice->CreateBuffer(&Desc,&ImmutableData,UniformBufferResource.GetInitReference()), Direct3DDevice);
 
 			NewUniformBuffer = new FD3D11UniformBuffer(this, Layout, UniformBufferResource, FRingAllocation());
+
+			INC_DWORD_STAT(STAT_D3D11NumImmutableUniformBuffers);
 		}
 	}
 	else
diff --git a/Engine/Source/Runtime/Windows/D3D11RHI/Private/D3D11Util.cpp b/Engine/Source/Runtime/Windows/D3D11RHI/Private/D3D11Util.cpp
--- a/Engine/Source/Runtime/Windows/D3D11RHI/Private/D3D11Util.cpp
+++ b/Engine/Source/Runtime/Windows/D3D11RHI/Private/D3D11Util.cpp
@@ -538,6 +538,7 @@ DEFINE_STAT(STAT_D3D11UpdateUniformBufferTime);
 DEFINE_STAT(STAT_D3D11TexturePoolMemory);
 DEFINE_STAT(STAT_D3D11FreeUniformBufferMemory);
 DEFINE_STAT(STAT_D3D11NumFreeUniformBuffers);
+DEFINE_STAT(STAT_D3D11NumImmutableUniformBuffers);
 DEFINE_STAT(STAT_D3D11NumBoundShaderState);
 
 #undef LOCTEXT_NAMESPACE
diff --git a/Engine/Source/Runtime/Windows/D3D11RHI/Private/D3D11VertexBuffer.cpp b/Engine/Source/Runtime/Windows/D3D11RHI/Private/D3D11VertexBuffer.cpp
--- a/Engine/Source/Runtime/Windows/D3D11RHI/Private/D3D11VertexBuffer.cpp
+++ b/Engine/Source/Runtime/Windows/D3D11RHI/Private/D3D11VertexBuffer.cpp
@@ -94,6 +94,15 @@ FVertexBufferRHIRef FD3D11DynamicRHI::RHICreateVertexBuffer(uint32 Size,uint32 I
 	return new FD3D11VertexBuffer(VertexBufferResource,Size,InUsage);
 }
 
+FVertexBufferRHIRef FD3D11DynamicRHI::CreateVertexBuffer_RenderThread(
+	class FRHICommandListImmediate& RHICmdList,
+	uint32 Size,
+	uint32 InUsage,
+	FRHIResourceCreateInfo& CreateInfo)
+{
+	return RHICreateVertexBuffer(Size, InUsage, CreateInfo);
+}
+
 void* FD3D11DynamicRHI::RHILockVertexBuffer(FVertexBufferRHIParamRef VertexBufferRHI,uint32 Offset,uint32 Size,EResourceLockMode LockMode)
 {
 	check(Size > 0);
@@ -160,7 +169,7 @@ void* FD3D11DynamicRHI::RHILockVertexBuffer(FVertexBufferRHIParamRef VertexBuffe
 	}
 
 	// Add the lock to the lock map.
-	OutstandingLocks.Add(LockedKey,LockedData);
+	GetThreadLocalLockTracker().Add(LockedKey,LockedData);
 
 	// Return the offset pointer
 	return (void*)((uint8*)LockedData.GetData() + Offset);
@@ -176,9 +185,10 @@ void FD3D11DynamicRHI::RHIUnlockVertexBuffer(FVertexBufferRHIParamRef VertexBuff
 	const bool bIsDynamic = (Desc.Usage == D3D11_USAGE_DYNAMIC);
 
 	// Find the outstanding lock for this VB.
+	FD3D11LockTracker& OutstandingLocks = GetThreadLocalLockTracker();
 	FD3D11LockedKey LockedKey(VertexBuffer->Resource);
 	FD3D11LockedData* LockedData = OutstandingLocks.Find(LockedKey);
-	check(LockedData);
+	checkf(LockedData, TEXT("Vertex buffer is either not locked or locked on a different thread"));
 
 	if(bIsDynamic)
 	{
diff --git a/Engine/Source/Runtime/Windows/D3D11RHI/Private/Windows/WindowsD3D11Device.cpp b/Engine/Source/Runtime/Windows/D3D11RHI/Private/Windows/WindowsD3D11Device.cpp
--- a/Engine/Source/Runtime/Windows/D3D11RHI/Private/Windows/WindowsD3D11Device.cpp
+++ b/Engine/Source/Runtime/Windows/D3D11RHI/Private/Windows/WindowsD3D11Device.cpp
@@ -1676,6 +1676,10 @@ void FD3D11DynamicRHI::InitD3DDevice()
 
 void FD3D11DynamicRHI::RHIPerFrameRHIFlushComplete()
 {
+	extern void D3D11RHIQueryBatcherPerFrameCleanup();
+
+	D3D11RHIQueryBatcherPerFrameCleanup();
+
 #if NV_AFTERMATH
 	if (GDX11NVAfterMathEnabled)
 	{
diff --git a/Engine/Source/Runtime/Windows/D3D11RHI/Public/D3D11Util.h b/Engine/Source/Runtime/Windows/D3D11RHI/Public/D3D11Util.h
--- a/Engine/Source/Runtime/Windows/D3D11RHI/Public/D3D11Util.h
+++ b/Engine/Source/Runtime/Windows/D3D11RHI/Public/D3D11Util.h
@@ -27,6 +27,9 @@
 
 #endif
 
+#define D3D11RHI_IMMEDIATE_CONTEXT (static_cast<FD3D11DynamicRHI*>(GDynamicRHI)->GetDeviceContext())
+#define D3D11RHI_DEVICE (static_cast<FD3D11DynamicRHI*>(GDynamicRHI)->GetDevice())
+
 /**
  * Checks that the given result isn't a failure.  If it is, the application exits with an appropriate error message.
  * @param	Result - The result code to check.
@@ -261,3 +264,56 @@ private:
 	/** The index of the currently locked sub-buffer. */
 	int32 LockedBufferIndex;
 };
+
+template <
+	typename JobType,
+	typename = TEnableIf<TOr<
+	TIsSame<JobType, TFunction<void()>>,
+	TIsSame<JobType, TFunction<void()>&>>::Value>>
+	class TD3D11RHIGenericCommand final : public FRHICommand<TD3D11RHIGenericCommand<JobType>>
+{
+public:
+	// InRHIJob is supposed to be called on RHIT (don't capture things that can become outdated here)
+	TD3D11RHIGenericCommand(JobType&& InRHIJob)
+		: RHIJob(Forward<JobType>(InRHIJob))
+	{}
+
+	void Execute(FRHICommandListBase& RHICmdList)
+	{
+		RHIJob();
+	}
+
+private:
+	JobType RHIJob;
+};
+
+template<typename JobType>
+inline void RunOnRHIThread(JobType&& InRHIJob)
+{
+	//check(IsInRenderingThread());
+	typedef TD3D11RHIGenericCommand<JobType> CmdType;
+	FRHICommandListImmediate& RHICmdList = FRHICommandListExecutor::GetImmediateCommandList();
+	new (RHICmdList.AllocCommand<CmdType>()) CmdType(Forward<JobType>(InRHIJob));
+}
+
+inline bool ShouldNotEnqueueRHICommand()
+{
+	FRHICommandListImmediate& RHICmdList = FRHICommandListExecutor::GetImmediateCommandList();
+	return RHICmdList.Bypass() || (IsRunningRHIInSeparateThread() && IsInRHIThread()) || (!IsRunningRHIInSeparateThread() && IsInRenderingThread());
+}
+
+inline void D3D11StallRHIThread()
+{
+	if (IsRunningRHIInSeparateThread() && IsInRenderingThread() && GRHICommandList.IsRHIThreadActive())
+	{
+		FRHICommandListExecutor::GetImmediateCommandList().StallRHIThread();
+	}
+}
+
+inline void D3D11UnstallRHIThread()
+{
+	if (IsRunningRHIInSeparateThread() && IsInRenderingThread() && FRHICommandListExecutor::GetImmediateCommandList().IsStalled())
+	{
+		FRHICommandListExecutor::GetImmediateCommandList().UnStallRHIThread();
+	}
+}
\ No newline at end of file
