UE4.22 - Fix clang
#rnx

diff --git a/Engine/Source/Runtime/Windows/D3D11RHI/Private/D3D11Commands.cpp b/Engine/Source/Runtime/Windows/D3D11RHI/Private/D3D11Commands.cpp
--- a/Engine/Source/Runtime/Windows/D3D11RHI/Private/D3D11Commands.cpp
+++ b/Engine/Source/Runtime/Windows/D3D11RHI/Private/D3D11Commands.cpp
@@ -2157,11 +2157,13 @@ void FD3D11DynamicRHI::RHIFlushComputeShaderCache()
 
 FD3D11DynamicRHI::FD3D11LockTracker& FD3D11DynamicRHI::GetThreadLocalLockTracker()
 {
+#if EXPERIMENTAL_D3D11_RHITHREAD
 	if (IsInRHIThread())
 	{
 		return LockTrackers[1];
 	}
 	else
+#endif
 	{
 		return LockTrackers[0];
 	}
