missing file for dx11 RHI thread.
Original change desc:

RHI Thread for DirectX11(Change by Intel)
Experimental, disabled by default for nowl.
define EXPERIMENTAL_D3D11_RHITHREAD to 1 in D3D11RHIPrivate.h to enable.

diff --git a/Engine/Source/Runtime/Windows/D3D11RHI/Private/D3D11Device.cpp b/Engine/Source/Runtime/Windows/D3D11RHI/Private/D3D11Device.cpp
--- a/Engine/Source/Runtime/Windows/D3D11RHI/Private/D3D11Device.cpp
+++ b/Engine/Source/Runtime/Windows/D3D11RHI/Private/D3D11Device.cpp
@@ -209,6 +209,7 @@ FD3D11DynamicRHI::FD3D11DynamicRHI(IDXGIFactory1* InDXGIFactory1,D3D_FEATURE_LEV
 		GMaxCubeTextureDimensions = D3D11_REQ_TEXTURECUBE_DIMENSION;
 		GMaxTextureArrayLayers = D3D11_REQ_TEXTURE2D_ARRAY_AXIS_DIMENSION;
 		GRHISupportsMSAADepthSampleAccess = true;
+		GRHISupportsRHIThread = !!EXPERIMENTAL_D3D11_RHITHREAD;
 	}
 	else if (FeatureLevel >= D3D_FEATURE_LEVEL_10_0)
 	{
@@ -548,6 +549,8 @@ void FD3D11DynamicRHI::CleanupD3DDevice()
 
 		ReleasePooledUniformBuffers();
 		ReleasePooledTextures();
+		ReleaseCachedQueries();
+
 
 		// Clean up the AMD extensions and shut down the AMD AGS utility library
 		if (AmdAgsContext != NULL)
