[LLM] fix for PC LLM builds.

diff --git a/Engine/Source/Runtime/D3D12RHI/Private/D3D12Stats.cpp b/Engine/Source/Runtime/D3D12RHI/Private/D3D12Stats.cpp
--- a/Engine/Source/Runtime/D3D12RHI/Private/D3D12Stats.cpp
+++ b/Engine/Source/Runtime/D3D12RHI/Private/D3D12Stats.cpp
@@ -340,7 +340,7 @@ void UpdateBufferStats(FD3D12ResourceLocation* ResourceLocation, bool bAllocatin
 		// released and the tracking could be re-enabled in the buddy allocator.
 		// The problem is that the releasing of resources happens in a generic way
 		// (see FD3D12ResourceLocation) 
-		LLM_SCOPED_PAUSE_TRACKING_WITH_ENUM_AND_AMOUNT(ELLMTag::VertexData, RequestedSize, ELLMTracker::Default, ELLMAllocType::None);
+		LLM_SCOPED_PAUSE_TRACKING_WITH_ENUM_AND_AMOUNT(ELLMTag::Meshes, RequestedSize, ELLMTracker::Default, ELLMAllocType::None);
 		LLM_SCOPED_PAUSE_TRACKING_WITH_ENUM_AND_AMOUNT(ELLMTag::GraphicsPlatform, RequestedSize, ELLMTracker::Platform, ELLMAllocType::None);
 #endif
 	}
@@ -365,7 +365,7 @@ void UpdateBufferStats(FD3D12ResourceLocation* ResourceLocation, bool bAllocatin
 
 #if PLATFORM_WINDOWS
 		// this is a work-around on Windows. See comment above.
-		LLM_SCOPED_PAUSE_TRACKING_WITH_ENUM_AND_AMOUNT(ELLMTag::VertexData, -(int64)RequestedSize, ELLMTracker::Default, ELLMAllocType::None);
+		LLM_SCOPED_PAUSE_TRACKING_WITH_ENUM_AND_AMOUNT(ELLMTag::Meshes, -(int64)RequestedSize, ELLMTracker::Default, ELLMAllocType::None);
 		LLM_SCOPED_PAUSE_TRACKING_WITH_ENUM_AND_AMOUNT(ELLMTag::GraphicsPlatform, -(int64)RequestedSize, ELLMTracker::Platform, ELLMAllocType::None);
 #endif
 	}
diff --git a/Engine/Source/Runtime/Windows/D3D11RHI/Private/D3D11RHI.cpp b/Engine/Source/Runtime/Windows/D3D11RHI/Private/D3D11RHI.cpp
--- a/Engine/Source/Runtime/Windows/D3D11RHI/Private/D3D11RHI.cpp
+++ b/Engine/Source/Runtime/Windows/D3D11RHI/Private/D3D11RHI.cpp
@@ -777,7 +777,7 @@ void UpdateBufferStats(TRefCountPtr<ID3D11Buffer> Buffer, bool bAllocating)
 		// this is a work-around on Windows. Due to the fact that there is no way
 		// to hook the actual d3d allocations we can't track the memory in the normal way.
 		// Instead we simply tell LLM the size of these resources.
-		LLM_SCOPED_PAUSE_TRACKING_WITH_ENUM_AND_AMOUNT(ELLMTag::VertexData, Desc.ByteWidth, ELLMTracker::Default, ELLMAllocType::None);
+		LLM_SCOPED_PAUSE_TRACKING_WITH_ENUM_AND_AMOUNT(ELLMTag::Meshes, Desc.ByteWidth, ELLMTracker::Default, ELLMAllocType::None);
 		LLM_SCOPED_PAUSE_TRACKING_WITH_ENUM_AND_AMOUNT(ELLMTag::GraphicsPlatform, Desc.ByteWidth, ELLMTracker::Platform, ELLMAllocType::None);
 #endif
 	}
@@ -804,7 +804,7 @@ void UpdateBufferStats(TRefCountPtr<ID3D11Buffer> Buffer, bool bAllocating)
 		// this is a work-around on Windows. Due to the fact that there is no way
 		// to hook the actual d3d allocations we can't track the memory in the normal way.
 		// Instead we simply tell LLM the size of these resources.
-		LLM_SCOPED_PAUSE_TRACKING_WITH_ENUM_AND_AMOUNT(ELLMTag::VertexData, -(int64)Desc.ByteWidth, ELLMTracker::Default, ELLMAllocType::None);
+		LLM_SCOPED_PAUSE_TRACKING_WITH_ENUM_AND_AMOUNT(ELLMTag::Meshes, -(int64)Desc.ByteWidth, ELLMTracker::Default, ELLMAllocType::None);
 		LLM_SCOPED_PAUSE_TRACKING_WITH_ENUM_AND_AMOUNT(ELLMTag::GraphicsPlatform, -(int64)Desc.ByteWidth, ELLMTracker::Platform, ELLMAllocType::None);
 #endif
 	}
