General improvements to the FHighlightRecorder class.
Also, it now works in PIE too.

diff --git a/Engine/Source/Runtime/GameplayMediaEncoder/Private/GameplayMediaEncoder.cpp b/Engine/Source/Runtime/GameplayMediaEncoder/Private/GameplayMediaEncoder.cpp
--- a/Engine/Source/Runtime/GameplayMediaEncoder/Private/GameplayMediaEncoder.cpp
+++ b/Engine/Source/Runtime/GameplayMediaEncoder/Private/GameplayMediaEncoder.cpp
@@ -288,21 +288,24 @@ bool FGameplayMediaEncoder::Start()
 	//
 	// subscribe to engine delegates for audio output and back buffer
 	//
-	if (UGameEngine* GameEngine = Cast<UGameEngine>(GEngine))
+	if (UEngine* Engine = GEngine)
 	{
-		FAudioDevice* AudioDevice = GameEngine->GetMainAudioDevice();
+		FAudioDevice* AudioDevice = Engine->GetMainAudioDevice();
 		if (AudioDevice)
 		{
 			AudioDevice->RegisterSubmixBufferListener(this);
 		}
 
 		VideoEncoder->Start();
-		FSlateRenderer::FOnBackBufferReadyToPresent OnBackBufferReadyDelegate;
-		OnBackBufferReadyDelegate.AddRaw(this, &FGameplayMediaEncoder::OnBackBufferReady);
-		FSlateApplication::Get().GetRenderer()->OnBackBufferReadyToPresent() = OnBackBufferReadyDelegate;
-	}
+		FSlateApplication::Get().GetRenderer()->OnBackBufferReadyToPresent().AddRaw(this, &FGameplayMediaEncoder::OnBackBufferReady);
 
-	return true;
+		return true;
+	}
+	else
+	{
+		UE_LOG(GameplayMediaEncoder, Error, TEXT("Failed to register audio and/or video callbacks"));
+		return false;
+	}
 }
 
 void FGameplayMediaEncoder::Stop()
diff --git a/Engine/Source/Runtime/Windows/WindowsPlatformFeatures/Private/HighlightRecorder.cpp b/Engine/Source/Runtime/Windows/WindowsPlatformFeatures/Private/HighlightRecorder.cpp
--- a/Engine/Source/Runtime/Windows/WindowsPlatformFeatures/Private/HighlightRecorder.cpp
+++ b/Engine/Source/Runtime/Windows/WindowsPlatformFeatures/Private/HighlightRecorder.cpp
@@ -116,6 +116,12 @@ void FHighlightRecorder::Stop()
 	FGameplayMediaEncoder::Get()->UnregisterListener(this);
 	State = EState::Stopped;
 
+	if (BackgroundSaving)
+	{
+		BackgroundSaving->Join();
+		BackgroundSaving.Reset();
+	}
+
 	UE_LOG(HighlightRecorder, Log, TEXT("recording stopped"));
 }
 
@@ -178,7 +184,7 @@ bool FHighlightRecorder::SaveHighlight(const TCHAR* Filename, FDoneCallback InDo
 
 	{
 		CSV_SCOPED_TIMING_STAT(WindowsVideoRecordingSystem, HighlightRecorder_SaveThreadCreation);
-		BackgroundProcessor.Reset(new FThread(TEXT("Highlight Saving"), [this, LocalFilename, MaxDurationSecs]()
+		BackgroundSaving.Reset(new FThread(TEXT("Highlight Saving"), [this, LocalFilename, MaxDurationSecs]()
 		{
 			SaveHighlightInBackground(LocalFilename, MaxDurationSecs);
 		}));
@@ -187,6 +193,7 @@ bool FHighlightRecorder::SaveHighlight(const TCHAR* Filename, FDoneCallback InDo
 	return true;
 }
 
+
 // the bool result is solely for convenience (CHECK_HR) and is ignored as it's a thread function and
 // nobody checks it's result. Actual result is notified by the callback.
 bool FHighlightRecorder::SaveHighlightInBackground(const FString& Filename, double MaxDurationSecs)
@@ -195,69 +202,80 @@ bool FHighlightRecorder::SaveHighlightInBackground(const FString& Filename, doub
 
 	double T0 = FPlatformTime::Seconds();
 
-	bool bRes = true;
+	bool bRes = SaveHighlightInBackgroundImpl(Filename, MaxDurationSecs);
 
+	double PassedSecs = FPlatformTime::Seconds() - T0;
+	UE_LOG(HighlightRecorder, Log, TEXT("saving to %s %s, took %.3f secs"), *Filename, bRes ? TEXT("succeeded") : TEXT("failed"), PassedSecs);
+
+	DoneCallback(bRes);
+	bSaving = false;
+
+	return bRes;
+}
+
+bool FHighlightRecorder::SaveHighlightInBackgroundImpl(const FString& Filename, double MaxDurationSecs)
+{
 	TArray<FGameplayMediaEncoderSample> Samples = RingBuffer.GetCopy();
 
-	do // once
+	if (Samples.Num()==0)
 	{
-		if (!InitialiseMp4Writer(Filename))
-		{
-			bRes = false;
-			break;
-		}
+		UE_LOG(HighlightRecorder, Error, TEXT("no samples to save to .mp4"));
+		return false;
+	}
+
+	int FirstSampleIndex;
+	FTimespan StartTime;
+	if (!GetSavingStart(Samples, FTimespan::FromSeconds(MaxDurationSecs), FirstSampleIndex, StartTime))
+	{
+		return false;
+	}
 
-		int SampleIndex;
-		FTimespan StartTime;
-		if (!GetSavingStart(Samples, FTimespan::FromSeconds(MaxDurationSecs), SampleIndex, StartTime))
+	// Check if we have audio, so that if we don't, we don't create an audio track
+	bool bHasAudio = false;
+	for (int Idx = FirstSampleIndex; Idx != Samples.Num(); ++Idx)
+	{
+		if (Samples[Idx].GetType() == EMediaType::Audio)
 		{
-			bRes = false;
+			bHasAudio = true;
 			break;
 		}
+	}
 
-		checkf(Samples[SampleIndex].IsVideoKeyFrame(), TEXT("t %.3f d %.3f"), Samples[SampleIndex].GetTime().GetTotalSeconds(), Samples[SampleIndex].GetDuration().GetTotalSeconds());
+	if (!InitialiseMp4Writer(Filename, bHasAudio))
+	{
+		return false;
+	}
 
-		if (SampleIndex == Samples.Num())
-		{
-			UE_LOG(HighlightRecorder, Error, TEXT("no samples to save to .mp4"));
-			bRes = false;
-			break;
-		}
+	checkf(Samples[FirstSampleIndex].IsVideoKeyFrame(), TEXT("t %.3f d %.3f"), Samples[FirstSampleIndex].GetTime().GetTotalSeconds(), Samples[FirstSampleIndex].GetDuration().GetTotalSeconds());
 
-		UE_LOG(HighlightRecorder, Verbose, TEXT("writting %d samples to .mp4, %.3f s, starting from %.3f s, index %d"), Samples.Num() - SampleIndex, (Samples.Last().GetTime() - StartTime + Samples.Last().GetDuration()).GetTotalSeconds(), StartTime.GetTotalSeconds(), SampleIndex);
+	if (FirstSampleIndex == Samples.Num())
+	{
+		UE_LOG(HighlightRecorder, Error, TEXT("no samples to save to .mp4"));
+		return false;
+	}
 
-		// get samples starting from `StartTime` and push them into Mp4Writer
-		for (; SampleIndex != Samples.Num() && !bStopSaving; ++SampleIndex)
-		{
-			FGameplayMediaEncoderSample& Sample = Samples[SampleIndex];
-			Sample.SetTime(Sample.GetTime() - StartTime);
-			if (!Mp4Writer->Write(Sample))
-			{
-				bRes = false;
-				break;
-			}
-		}
-	} while (false);
+	UE_LOG(HighlightRecorder, Verbose, TEXT("writting %d samples to .mp4, %.3f s, starting from %.3f s, index %d"), Samples.Num() - FirstSampleIndex, (Samples.Last().GetTime() - StartTime + Samples.Last().GetDuration()).GetTotalSeconds(), StartTime.GetTotalSeconds(), FirstSampleIndex);
 
-	if (bRes)
+	// get samples starting from `StartTime` and push them into Mp4Writer
+	for (int Idx = FirstSampleIndex; Idx != Samples.Num(); ++Idx)
 	{
-		if (!Mp4Writer->Finalize())
+		FGameplayMediaEncoderSample& Sample = Samples[Idx];
+		Sample.SetTime(Sample.GetTime() - StartTime);
+		if (!Mp4Writer->Write(Sample, (Sample.GetType()==EMediaType::Video && bHasAudio) ? 1 : 0))
 		{
-			bRes = false;
+			return false;
 		}
 	}
 
-	double PassedSecs = FPlatformTime::Seconds() - T0;
-	UE_LOG(HighlightRecorder, Log, TEXT("saving to %s %s, took %.3f msecs"), *Filename, bRes ? TEXT("succeeded") : TEXT("failed"), PassedSecs);
-
-	bSaving = false;
-
-	DoneCallback(bRes);
+	if (!Mp4Writer->Finalize())
+	{
+		return false;
+	}
 
-	return bRes;
+	return true;
 }
 
-bool FHighlightRecorder::InitialiseMp4Writer(const FString& Filename)
+bool FHighlightRecorder::InitialiseMp4Writer(const FString& Filename, bool bHasAudio)
 {
 	FString VideoCaptureDir = FPaths::VideoCaptureDir();
 	auto& PlatformFile = FPlatformFileManager::Get().GetPlatformFile();
@@ -287,15 +305,12 @@ bool FHighlightRecorder::InitialiseMp4Writer(const FString& Filename)
 	}
 
 	DWORD StreamIndex;
-	if (!Mp4Writer->CreateStream(AudioType, StreamIndex))
+	if (bHasAudio)
 	{
-		return false;
-	}
-
-	if (StreamIndex != static_cast<decltype(StreamIndex)>(EMediaType::Audio))
-	{
-		UE_LOG(HighlightRecorder, Error, TEXT("Invalid audio stream index: %d"), StreamIndex);
-		return false;
+		if (!Mp4Writer->CreateStream(AudioType, StreamIndex))
+		{
+			return false;
+		}
 	}
 
 	TRefCountPtr<IMFMediaType> VideoType;
@@ -309,12 +324,6 @@ bool FHighlightRecorder::InitialiseMp4Writer(const FString& Filename)
 		return false;
 	}
 
-	if (StreamIndex != static_cast<decltype(StreamIndex)>(EMediaType::Video))
-	{
-		UE_LOG(HighlightRecorder, Error, TEXT("Invalid video stream index: %d"), StreamIndex);
-		return false;
-	}
-
 	if (!Mp4Writer->Start())
 	{
 		return false;
diff --git a/Engine/Source/Runtime/Windows/WindowsPlatformFeatures/Private/HighlightRecorder.h b/Engine/Source/Runtime/Windows/WindowsPlatformFeatures/Private/HighlightRecorder.h
--- a/Engine/Source/Runtime/Windows/WindowsPlatformFeatures/Private/HighlightRecorder.h
+++ b/Engine/Source/Runtime/Windows/WindowsPlatformFeatures/Private/HighlightRecorder.h
@@ -74,7 +74,8 @@ public:
 
 private:
 	bool SaveHighlightInBackground(const FString& Filename, double MaxDurationSecs);
-	bool InitialiseMp4Writer(const FString& Filename);
+	bool SaveHighlightInBackgroundImpl(const FString& Filename, double MaxDurationSecs);
+	bool InitialiseMp4Writer(const FString& Filename, bool bHasAudio);
 	bool GetSavingStart(const TArray<FGameplayMediaEncoderSample>& Samples, FTimespan MaxDuration, int& OutStartIndex, FTimespan& OutStartTime) const;
 
 	// takes into account if we've been paused and shifts current time back to compensate paused state
@@ -104,10 +105,9 @@ private:
 	// for how long recording has been paused since it's started
 	FTimespan TotalPausedDuration = 0;
 
-	TUniquePtr<FThread> BackgroundProcessor;
+	TUniquePtr<FThread> BackgroundSaving;
 	FDoneCallback DoneCallback;
 	FThreadSafeBool bSaving = false;
-	FThreadSafeBool bStopSaving = false;
 
 #pragma region testing
 public:
@@ -136,6 +136,11 @@ public:
 
 	static void StopCmd()
 	{
+		if (!CheckSingleton())
+		{
+			return;
+		}
+
 		Get()->Stop();
 		delete Singleton;
 		Singleton = nullptr;
@@ -143,16 +148,31 @@ public:
 
 	static void PauseCmd()
 	{
+		if (!CheckSingleton())
+		{
+			return;
+		}
+
 		Get()->Pause(true);
 	}
 
 	static void ResumeCmd()
 	{
+		if (!CheckSingleton())
+		{
+			return;
+		}
+
 		Get()->Pause(false);
 	}
 
 	static void SaveCmd(const TArray<FString>& Args, UWorld*, FOutputDevice& Output)
 	{
+		if (!CheckSingleton())
+		{
+			return;
+		}
+
 		if (Args.Num() > 2)
 		{
 			Output.Logf(ELogVerbosity::Error, TEXT("0-2 parameters expected: Save [filename=\"test.mp4\"] [max_duration_secs= ring buffer duration]"));
@@ -182,6 +202,20 @@ public:
 	}
 
 private:
+
+	static bool CheckSingleton()
+	{
+		if (Singleton)
+		{
+			return true;
+		}
+		else
+		{
+			UE_LOG(HighlightRecorder, Error, TEXT("HighlightRecorder not initialized."));
+			return false;
+		}
+	}
+
 	static FHighlightRecorder* Get()
 	{
 		check(Singleton);
diff --git a/Engine/Source/Runtime/Windows/WindowsPlatformFeatures/Private/WmfMp4Writer.cpp b/Engine/Source/Runtime/Windows/WindowsPlatformFeatures/Private/WmfMp4Writer.cpp
--- a/Engine/Source/Runtime/Windows/WindowsPlatformFeatures/Private/WmfMp4Writer.cpp
+++ b/Engine/Source/Runtime/Windows/WindowsPlatformFeatures/Private/WmfMp4Writer.cpp
@@ -37,9 +37,9 @@ bool FWmfMp4Writer::Start()
 	return true;
 }
 
-bool FWmfMp4Writer::Write(const FGameplayMediaEncoderSample& Sample)
+bool FWmfMp4Writer::Write(const FGameplayMediaEncoderSample& Sample, DWORD StreamIndex)
 {
-	CHECK_HR(Writer->WriteSample(static_cast<DWORD>(Sample.GetType()), const_cast<IMFSample*>(Sample.GetSample())));
+	CHECK_HR(Writer->WriteSample(StreamIndex, const_cast<IMFSample*>(Sample.GetSample())));
 
 	UE_LOG(MP4, VeryVerbose, TEXT("stream #%d: time %.3f, duration %.3f%s"), static_cast<int>(Sample.GetType()), Sample.GetTime().GetTotalSeconds(), Sample.GetDuration().GetTotalSeconds(), Sample.IsVideoKeyFrame() ? TEXT(", key-frame") : TEXT(""));
 
diff --git a/Engine/Source/Runtime/Windows/WindowsPlatformFeatures/Private/WmfMp4Writer.h b/Engine/Source/Runtime/Windows/WindowsPlatformFeatures/Private/WmfMp4Writer.h
--- a/Engine/Source/Runtime/Windows/WindowsPlatformFeatures/Private/WmfMp4Writer.h
+++ b/Engine/Source/Runtime/Windows/WindowsPlatformFeatures/Private/WmfMp4Writer.h
@@ -13,7 +13,7 @@ public:
 	bool Initialize(const TCHAR* Filename);
 	bool CreateStream(IMFMediaType* StreamType, DWORD& StreamIndex);
 	bool Start();
-	bool Write(const FGameplayMediaEncoderSample& Sample);
+	bool Write(const FGameplayMediaEncoderSample& Sample, DWORD StreamIndex);
 	bool Finalize();
 
 private:
