Fixed issue where scaling of R and B were swapped in  FD3DFloatR11G11B10 to linear conversion. (UDN479620)

diff --git a/Engine/Source/Runtime/Windows/D3D11RHI/Private/D3D11RenderTarget.cpp b/Engine/Source/Runtime/Windows/D3D11RHI/Private/D3D11RenderTarget.cpp
--- a/Engine/Source/Runtime/Windows/D3D11RHI/Private/D3D11RenderTarget.cpp
+++ b/Engine/Source/Runtime/Windows/D3D11RHI/Private/D3D11RenderTarget.cpp
@@ -461,13 +461,13 @@ struct FD3DFloatR11G11B10
 
 		Result[0].Components.Sign = 0;
 		Result[0].Components.Exponent = R_Exponent - 15 + 127;
-		Result[0].Components.Mantissa = FMath::Min<uint32>(FMath::FloorToInt((float)R_Mantissa / 32.0f * 8388608.0f),(1 << 23) - 1);
+		Result[0].Components.Mantissa = FMath::Min<uint32>(FMath::FloorToInt((float)R_Mantissa / 64.0f * 8388608.0f),(1 << 23) - 1);
 		Result[1].Components.Sign = 0;
 		Result[1].Components.Exponent = G_Exponent - 15 + 127;
 		Result[1].Components.Mantissa = FMath::Min<uint32>(FMath::FloorToInt((float)G_Mantissa / 64.0f * 8388608.0f),(1 << 23) - 1);
 		Result[2].Components.Sign = 0;
 		Result[2].Components.Exponent = B_Exponent - 15 + 127;
-		Result[2].Components.Mantissa = FMath::Min<uint32>(FMath::FloorToInt((float)B_Mantissa / 64.0f * 8388608.0f),(1 << 23) - 1);
+		Result[2].Components.Mantissa = FMath::Min<uint32>(FMath::FloorToInt((float)B_Mantissa / 32.0f * 8388608.0f),(1 << 23) - 1);
 
 		return FLinearColor(Result[0].Float, Result[1].Float, Result[2].Float);
 	}
