Fix for first line being repeated in screenshots of float32 types.

diff --git a/Engine/Source/Runtime/D3D12RHI/Private/D3D12RenderTarget.cpp b/Engine/Source/Runtime/D3D12RHI/Private/D3D12RenderTarget.cpp
--- a/Engine/Source/Runtime/D3D12RHI/Private/D3D12RenderTarget.cpp
+++ b/Engine/Source/Runtime/D3D12RHI/Private/D3D12RenderTarget.cpp
@@ -996,7 +996,8 @@ static void ConvertRAWSurfaceDataToFColor(DXGI_FORMAT Format, uint32 Width, uint
 
 		for (uint32 Y = 0; Y < Height; Y++)
 		{
-			float* SrcPtr = (float*)In;
+			float* SrcPtr = (float*)(In + Y * SrcPitch);
+
 			FColor* DestPtr = Out + Y * Width;
 
 			for (uint32 X = 0; X < Width; X++)
diff --git a/Engine/Source/Runtime/Windows/D3D11RHI/Private/D3D11RenderTarget.cpp b/Engine/Source/Runtime/Windows/D3D11RHI/Private/D3D11RenderTarget.cpp
--- a/Engine/Source/Runtime/Windows/D3D11RHI/Private/D3D11RenderTarget.cpp
+++ b/Engine/Source/Runtime/Windows/D3D11RHI/Private/D3D11RenderTarget.cpp
@@ -794,7 +794,7 @@ static void ConvertRAWSurfaceDataToFColor(DXGI_FORMAT Format, uint32 Width, uint
 		
 		for(uint32 Y = 0; Y < Height; Y++)
 		{
-			float* SrcPtr = (float*)In;
+			float* SrcPtr = (float*)(In + Y * SrcPitch);
 			FColor* DestPtr = Out + Y * Width;
 
 			for(uint32 X = 0; X < Width; X++)
