Fix for failing to init audio device when no audio devices are available.

Re-introduce the appropriate macro define that was accidentally set to 0.

diff --git a/Engine/Source/Runtime/Windows/XAudio2/Private/XAudio2Support.h b/Engine/Source/Runtime/Windows/XAudio2/Private/XAudio2Support.h
--- a/Engine/Source/Runtime/Windows/XAudio2/Private/XAudio2Support.h
+++ b/Engine/Source/Runtime/Windows/XAudio2/Private/XAudio2Support.h
@@ -10,7 +10,7 @@
 	#define XAUDIO_SUPPORTS_XMA2WAVEFORMATEX	0
 #endif	//XAUDIO_SUPPORTS_XMA2WAVEFORMATEX
 #ifndef XAUDIO_SUPPORTS_DEVICE_DETAILS
-	#define XAUDIO_SUPPORTS_DEVICE_DETAILS		0
+	#define XAUDIO_SUPPORTS_DEVICE_DETAILS		1
 #endif	//XAUDIO_SUPPORTS_DEVICE_DETAILS
 #ifndef XAUDIO2_SUPPORTS_MUSIC
 	#define XAUDIO2_SUPPORTS_MUSIC				0
