fix for VB/IB lock with RLM_ReadOnly and Offset != 0

diff --git a/Engine/Source/Runtime/Windows/D3D11RHI/Private/D3D11IndexBuffer.cpp b/Engine/Source/Runtime/Windows/D3D11RHI/Private/D3D11IndexBuffer.cpp
--- a/Engine/Source/Runtime/Windows/D3D11RHI/Private/D3D11IndexBuffer.cpp
+++ b/Engine/Source/Runtime/Windows/D3D11RHI/Private/D3D11IndexBuffer.cpp
@@ -107,7 +107,7 @@ void* FD3D11DynamicRHI::RHILockIndexBuffer(FRHIIndexBuffer* IndexBufferRHI,uint3
 			// If the static buffer is being locked for reading, create a staging buffer.
 			D3D11_BUFFER_DESC StagingBufferDesc;
 			ZeroMemory( &StagingBufferDesc, sizeof( D3D11_BUFFER_DESC ) );
-			StagingBufferDesc.ByteWidth = Size;
+			StagingBufferDesc.ByteWidth = Desc.ByteWidth;
 			StagingBufferDesc.Usage = D3D11_USAGE_STAGING;
 			StagingBufferDesc.BindFlags = 0;
 			StagingBufferDesc.CPUAccessFlags = D3D11_CPU_ACCESS_READ;
diff --git a/Engine/Source/Runtime/Windows/D3D11RHI/Private/D3D11VertexBuffer.cpp b/Engine/Source/Runtime/Windows/D3D11RHI/Private/D3D11VertexBuffer.cpp
--- a/Engine/Source/Runtime/Windows/D3D11RHI/Private/D3D11VertexBuffer.cpp
+++ b/Engine/Source/Runtime/Windows/D3D11RHI/Private/D3D11VertexBuffer.cpp
@@ -163,7 +163,7 @@ void* FD3D11DynamicRHI::RHILockVertexBuffer(FRHIVertexBuffer* VertexBufferRHI,ui
 			// Copy the contents of the vertex buffer to the staging buffer.
 			D3D11_BOX SourceBox;
 			SourceBox.left = Offset;
-			SourceBox.right = Size;
+			SourceBox.right = Offset + Size;
 			SourceBox.top = SourceBox.front = 0;
 			SourceBox.bottom = SourceBox.back = 1;
 			Direct3DDeviceIMContext->CopySubresourceRegion(StagingVertexBuffer,0,0,0,0,VertexBuffer->Resource,0,&SourceBox);
@@ -173,6 +173,7 @@ void* FD3D11DynamicRHI::RHILockVertexBuffer(FRHIVertexBuffer* VertexBufferRHI,ui
 			VERIFYD3D11RESULT_EX(Direct3DDeviceIMContext->Map(StagingVertexBuffer,0,D3D11_MAP_READ,0,&MappedSubresource), Direct3DDevice);
 			LockedData.SetData(MappedSubresource.pData);
 			LockedData.Pitch = MappedSubresource.RowPitch;
+			Offset = 0;
 		}
 		else
 		{
