Updated min feature level error message to be a little more generic, calling out project requirements rather than 10.0.
Added support for min and max feature level via ini.

diff --git a/Engine/Source/Runtime/RHI/Private/Windows/WindowsDynamicRHI.cpp b/Engine/Source/Runtime/RHI/Private/Windows/WindowsDynamicRHI.cpp
--- a/Engine/Source/Runtime/RHI/Private/Windows/WindowsDynamicRHI.cpp
+++ b/Engine/Source/Runtime/RHI/Private/Windows/WindowsDynamicRHI.cpp
@@ -212,7 +212,7 @@ static IDynamicRHIModule* LoadDynamicRHIModule(ERHIFeatureLevel::Type& DesiredFe
 
 		if (!DynamicRHIModule->IsSupported())
 		{
-			FMessageDialog::Open(EAppMsgType::Ok, NSLOCTEXT("WindowsDynamicRHI", "RequiredDX11Feature", "DX11 feature level 10.0 is required to run the engine."));
+			FMessageDialog::Open(EAppMsgType::Ok, NSLOCTEXT("WindowsDynamicRHI", "RequiredDX11Feature", "Failed to find a graphics adapter with the minimum required DX11 feature level."));
 			FPlatformMisc::RequestExit(1);
 			DynamicRHIModule = NULL;
 		}
diff --git a/Engine/Source/Runtime/Windows/D3D11RHI/Private/Windows/WindowsD3D11Device.cpp b/Engine/Source/Runtime/Windows/D3D11RHI/Private/Windows/WindowsD3D11Device.cpp
--- a/Engine/Source/Runtime/Windows/D3D11RHI/Private/Windows/WindowsD3D11Device.cpp
+++ b/Engine/Source/Runtime/Windows/D3D11RHI/Private/Windows/WindowsD3D11Device.cpp
@@ -128,10 +128,17 @@ static FAutoConsoleVariableRef CVarDX11NumGPUs(
  */
 namespace RHIConsoleVariables
 {
-	int32 FeatureSetLimit = -1;
-	static FAutoConsoleVariableRef CVarFeatureSetLimit(
+	int32 MinFeatureSetLimit = -1;
+	static FAutoConsoleVariableRef CVarMinFeatureSetLimit(
+		TEXT("RHI.MinFeatureSetLimit"),
+		MinFeatureSetLimit,
+		TEXT("If set to 11, limit D3D RHI to D3D11 and 11.1 feature levels, disallowing 10. Otherwise, it will use default. Changing this at run-time has no effect. (default is -1)")
+		);
+
+	int32 MaxFeatureSetLimit = -1;
+	static FAutoConsoleVariableRef CVarMaxFeatureSetLimit(
 		TEXT("RHI.FeatureSetLimit"),
-		FeatureSetLimit,
+		MaxFeatureSetLimit,
 		TEXT("If set to 10, limit D3D RHI to D3D10 feature level. Otherwise, it will use default. Changing this at run-time has no effect. (default is -1)")
 		);
 };
@@ -237,20 +244,37 @@ static void SafeCreateDXGIFactory(IDXGIFactory1** DXGIFactory1)
 #endif	//!D3D11_CUSTOM_VIEWPORT_CONSTRUCTOR
 }
 
+/**
+ * Returns the lowest D3D feature level we are allowed to created based on
+ * command line parameters.
+ */
+static D3D_FEATURE_LEVEL GetMinAllowedD3DFeatureLevel()
+{
+	// Default to 10.0
+	D3D_FEATURE_LEVEL AllowedFeatureLevel = D3D_FEATURE_LEVEL_10_0;
+
+	if (RHIConsoleVariables::MinFeatureSetLimit == 11)
+	{
+		AllowedFeatureLevel = D3D_FEATURE_LEVEL_11_0;
+	}
+
+	return AllowedFeatureLevel;
+}
+
 /**
  * Returns the highest D3D feature level we are allowed to created based on
  * command line parameters.
  */
-static D3D_FEATURE_LEVEL GetAllowedD3DFeatureLevel()
+static D3D_FEATURE_LEVEL GetMaxAllowedD3DFeatureLevel()
 {
-	// Default to D3D11 
+	// Default to 11.0
 	D3D_FEATURE_LEVEL AllowedFeatureLevel = D3D_FEATURE_LEVEL_11_0;
 
 	// Use a feature level 10 if specified on the command line.
 	if(FParse::Param(FCommandLine::Get(),TEXT("d3d10")) || 
 		FParse::Param(FCommandLine::Get(),TEXT("dx10")) ||
 		FParse::Param(FCommandLine::Get(),TEXT("sm4")) ||
-		RHIConsoleVariables::FeatureSetLimit == 10)
+		RHIConsoleVariables::MaxFeatureSetLimit == 10)
 	{
 		AllowedFeatureLevel = D3D_FEATURE_LEVEL_10_0;
 	}
@@ -327,10 +351,10 @@ static void SafeTestForOptionalGraphicsTools(IDXGIAdapter* Adapter, D3D_FEATURE_
  * Attempts to create a D3D11 device for the adapter using at most MaxFeatureLevel.
  * If creation is successful, true is returned and the supported feature level is set in OutFeatureLevel.
  */
-static bool SafeTestD3D11CreateDevice(IDXGIAdapter* Adapter,D3D_FEATURE_LEVEL MaxFeatureLevel,D3D_FEATURE_LEVEL* OutFeatureLevel)
+static bool SafeTestD3D11CreateDevice(IDXGIAdapter* Adapter,D3D_FEATURE_LEVEL MinFeatureLevel,D3D_FEATURE_LEVEL MaxFeatureLevel,D3D_FEATURE_LEVEL* OutFeatureLevel)
 {
-	ID3D11Device* D3DDevice = NULL;
-	ID3D11DeviceContext* D3DDeviceContext = NULL;
+	ID3D11Device* D3DDevice = nullptr;
+	ID3D11DeviceContext* D3DDeviceContext = nullptr;
 	uint32 DeviceFlags = D3D11_CREATE_DEVICE_SINGLETHREADED;
 	// Use a debug device if specified on the command line.
 	if(D3D11RHI_ShouldCreateWithD3DDebug())
@@ -348,9 +372,12 @@ static bool SafeTestD3D11CreateDevice(IDXGIAdapter* Adapter,D3D_FEATURE_LEVEL Ma
 		D3D_FEATURE_LEVEL_11_0,
 		D3D_FEATURE_LEVEL_10_0
 	};
-
+	
+	// Trim to allowed feature levels
 	int32 FirstAllowedFeatureLevel = 0;
 	int32 NumAllowedFeatureLevels = ARRAY_COUNT(RequestedFeatureLevels);
+	int32 LastAllowedFeatureLevel = NumAllowedFeatureLevels - 1;
+	
 	while (FirstAllowedFeatureLevel < NumAllowedFeatureLevels)
 	{
 		if (RequestedFeatureLevels[FirstAllowedFeatureLevel] == MaxFeatureLevel)
@@ -359,9 +386,18 @@ static bool SafeTestD3D11CreateDevice(IDXGIAdapter* Adapter,D3D_FEATURE_LEVEL Ma
 		}
 		FirstAllowedFeatureLevel++;
 	}
-	NumAllowedFeatureLevels -= FirstAllowedFeatureLevel;
 
-	if (NumAllowedFeatureLevels == 0)
+	while (LastAllowedFeatureLevel > 0)
+	{
+		if (RequestedFeatureLevels[LastAllowedFeatureLevel] >= MinFeatureLevel)
+		{
+			break;
+		}
+		LastAllowedFeatureLevel--;
+	}
+	
+	NumAllowedFeatureLevels = LastAllowedFeatureLevel - FirstAllowedFeatureLevel + 1;
+	if (MaxFeatureLevel < MinFeatureLevel || NumAllowedFeatureLevels <= 0)
 	{
 		return false;
 	}
@@ -376,7 +412,7 @@ static bool SafeTestD3D11CreateDevice(IDXGIAdapter* Adapter,D3D_FEATURE_LEVEL Ma
 		if(SUCCEEDED(D3D11CreateDevice(
 			Adapter,
 			D3D_DRIVER_TYPE_UNKNOWN,
-			NULL,
+			nullptr,
 			DeviceFlags,
 			&RequestedFeatureLevels[FirstAllowedFeatureLevel],
 			NumAllowedFeatureLevels,
@@ -811,7 +847,11 @@ void FD3D11DynamicRHIModule::FindAdapter()
 	const bool bFavorNonIntegrated = CVarExplicitAdapterValue == -1;
 
 	TRefCountPtr<IDXGIAdapter> TempAdapter;
-	D3D_FEATURE_LEVEL MaxAllowedFeatureLevel = GetAllowedD3DFeatureLevel();
+	D3D_FEATURE_LEVEL MinAllowedFeatureLevel = GetMinAllowedD3DFeatureLevel();
+	D3D_FEATURE_LEVEL MaxAllowedFeatureLevel = GetMaxAllowedD3DFeatureLevel();
+
+	UE_LOG(LogD3D11RHI, Log, TEXT("D3D11 min allowed feature level: %s"), GetFeatureLevelString(MinAllowedFeatureLevel));
+	UE_LOG(LogD3D11RHI, Log, TEXT("D3D11 max allowed feature level: %s"), GetFeatureLevelString(MaxAllowedFeatureLevel));
 
 	FD3D11Adapter FirstWithoutIntegratedAdapter;
 	FD3D11Adapter FirstAdapter;
@@ -837,7 +877,7 @@ void FD3D11DynamicRHIModule::FindAdapter()
 		if(TempAdapter)
 		{
 			D3D_FEATURE_LEVEL ActualFeatureLevel = (D3D_FEATURE_LEVEL)0;
-			if(SafeTestD3D11CreateDevice(TempAdapter,MaxAllowedFeatureLevel,&ActualFeatureLevel))
+			if(SafeTestD3D11CreateDevice(TempAdapter,MinAllowedFeatureLevel,MaxAllowedFeatureLevel,&ActualFeatureLevel))
 			{
 				// Log some information about the available D3D11 adapters.
 				VERIFYD3D11RESULT(TempAdapter->GetDesc(&AdapterDesc));
