Re-submitting some changes from dev-rendering

diff --git a/Engine/Source/Runtime/Windows/D3D11RHI/Private/D3D11Util.cpp b/Engine/Source/Runtime/Windows/D3D11RHI/Private/D3D11Util.cpp
--- a/Engine/Source/Runtime/Windows/D3D11RHI/Private/D3D11Util.cpp
+++ b/Engine/Source/Runtime/Windows/D3D11RHI/Private/D3D11Util.cpp
@@ -249,6 +249,15 @@ void VerifyD3D11Result(HRESULT D3DResult,const ANSICHAR* Code,const ANSICHAR* Fi
 	UE_LOG(LogD3D11RHI, Fatal,TEXT("%s failed \n at %s:%u \n with error %s"),ANSI_TO_TCHAR(Code),ANSI_TO_TCHAR(Filename),Line,*ErrorString);
 }
 
+void VerifyD3D11ResultNoExit(HRESULT D3DResult, const ANSICHAR* Code, const ANSICHAR* Filename, uint32 Line, ID3D11Device* Device)
+{
+	check(FAILED(D3DResult));
+
+	const FString& ErrorString = GetD3D11ErrorString(D3DResult, Device);
+
+	UE_LOG(LogD3D11RHI, Error, TEXT("%s failed \n at %s:%u \n with error %s Error Code List: https://docs.microsoft.com/en-us/windows/desktop/direct3ddxgi/dxgi-error"), ANSI_TO_TCHAR(Code), ANSI_TO_TCHAR(Filename), Line, *ErrorString);
+}
+
 void VerifyD3D11ShaderResult(FRHIShader* Shader, HRESULT D3DResult, const ANSICHAR* Code, const ANSICHAR* Filename, uint32 Line, ID3D11Device* Device)
 {
 	check(FAILED(D3DResult));
diff --git a/Engine/Source/Runtime/Windows/D3D11RHI/Private/Windows/WindowsD3D11Device.cpp b/Engine/Source/Runtime/Windows/D3D11RHI/Private/Windows/WindowsD3D11Device.cpp
--- a/Engine/Source/Runtime/Windows/D3D11RHI/Private/Windows/WindowsD3D11Device.cpp
+++ b/Engine/Source/Runtime/Windows/D3D11RHI/Private/Windows/WindowsD3D11Device.cpp
@@ -290,63 +290,6 @@ static D3D_FEATURE_LEVEL GetMaxAllowedD3DFeatureLevel()
 	return AllowedFeatureLevel;
 }
 
-/**
-* Attempts to create a D3D11 device to test if optional debugging features are installed.
-*/
-static void SafeTestForOptionalGraphicsTools(IDXGIAdapter* Adapter, D3D_FEATURE_LEVEL& FeatureLevel, uint32& DeviceFlags, D3D_FEATURE_LEVEL& OutFeatureLevel, int32 NumAllowedFeatureLevels)
-{
-	if((DeviceFlags & D3D11_CREATE_DEVICE_DEBUG) != D3D11_CREATE_DEVICE_DEBUG)
-	{
-		return;
-	}
-
-	ID3D11Device* D3DDevice = NULL;
-	ID3D11DeviceContext* D3DDeviceContext = NULL;
-
-#if PLATFORM_HOLOLENS
-	static bool bIsWin10 = true;
-#else
-	static bool bIsWin10 = FWindowsPlatformMisc::VerifyWindowsVersion(10, 0);
-#endif
-	static bool bPrintErrorMessage = true;
-
-	// Test device creation.
-	HRESULT Result = D3D11CreateDevice(
-		Adapter,
-		D3D_DRIVER_TYPE_UNKNOWN,
-		NULL,
-		DeviceFlags,
-		& FeatureLevel,
-		NumAllowedFeatureLevels,
-		D3D11_SDK_VERSION,
-		&D3DDevice,
-		&OutFeatureLevel,
-		&D3DDeviceContext
-	);
-
-	// If Optional Graphics tools were enabled then return true.
-	if (SUCCEEDED(Result))
-	{
-		D3DDevice->Release();
-		D3DDeviceContext->Release();
-		return;
-	}
-	// If creation failed due to missing Optional Features in Windows 10.
-	else if (DXGI_ERROR_SDK_COMPONENT_MISSING == Result && bIsWin10)
-	{
-		// Only print the debug error message once and don't clutter the log.
-		if (bPrintErrorMessage) 
-		{
-			UE_LOG(LogD3D11RHI, Warning,
-				TEXT("-d3ddebug was ignored as optional Graphics Tools still need to be installed. Install them through the Manage Optional Features in windows."));
-			bPrintErrorMessage = false;
-		}
-
-		// Remove the debug option from the device flag.
-		DeviceFlags = DeviceFlags & (~D3D11_CREATE_DEVICE_DEBUG);
-	}
-}
-
 /**
  * Attempts to create a D3D11 device for the adapter using at most MaxFeatureLevel.
  * If creation is successful, true is returned and the supported feature level is set in OutFeatureLevel.
@@ -402,14 +345,12 @@ static bool SafeTestD3D11CreateDevice(IDXGIAdapter* Adapter,D3D_FEATURE_LEVEL Mi
 		return false;
 	}
 
-	SafeTestForOptionalGraphicsTools(Adapter, RequestedFeatureLevels[FirstAllowedFeatureLevel], DeviceFlags, *OutFeatureLevel, NumAllowedFeatureLevels);
-
 	__try
 	{
 		// We don't want software renderer. Ideally we specify D3D_DRIVER_TYPE_HARDWARE on creation but
 		// when we specify an adapter we need to specify D3D_DRIVER_TYPE_UNKNOWN (otherwise the call fails).
 		// We cannot check the device type later (seems this is missing functionality in D3D).
-		if(SUCCEEDED(D3D11CreateDevice(
+		HRESULT Result = D3D11CreateDevice(
 			Adapter,
 			D3D_DRIVER_TYPE_UNKNOWN,
 			nullptr,
@@ -419,13 +360,25 @@ static bool SafeTestD3D11CreateDevice(IDXGIAdapter* Adapter,D3D_FEATURE_LEVEL Mi
 			D3D11_SDK_VERSION,
 			&D3DDevice,
 			OutFeatureLevel,
-			&D3DDeviceContext
-		))) 
+			&D3DDeviceContext);
+
+		if (SUCCEEDED(Result))
 		{
 			D3DDevice->Release();
 			D3DDeviceContext->Release();
 			return true;
 		}
+
+		// Log any reason for failure to create test device. Extra debug help.
+		VERIFYD3D11RESULT_NOEXIT(Result);
+
+		bool bIsWin10 = FWindowsPlatformMisc::VerifyWindowsVersion(10, 0);
+
+		// Fatal error on 0x887A002D
+		if (DXGI_ERROR_SDK_COMPONENT_MISSING == Result && bIsWin10)
+		{
+			UE_LOG(LogD3D11RHI, Fatal, TEXT("-d3ddebug was used but optional Graphics Tools were not found. Install them through the Manage Optional Features in windows. See: https://docs.microsoft.com/en-us/windows/uwp/gaming/use-the-directx-runtime-and-visual-studio-graphics-diagnostic-features"));
+		}
 	}
 	__except(IsDelayLoadException(GetExceptionInformation()))
 	{
@@ -1612,8 +1565,6 @@ void FD3D11DynamicRHI::InitD3DDevice()
 		
 		if (!bDeviceCreated)
 		{
-			SafeTestForOptionalGraphicsTools(Adapter, FeatureLevel, DeviceFlags, ActualFeatureLevel, 1);
-			
 			// Creating the Direct3D device.
 			VERIFYD3D11RESULT(D3D11CreateDevice(
 				Adapter,
diff --git a/Engine/Source/Runtime/Windows/D3D11RHI/Public/D3D11Util.h b/Engine/Source/Runtime/Windows/D3D11RHI/Public/D3D11Util.h
--- a/Engine/Source/Runtime/Windows/D3D11RHI/Public/D3D11Util.h
+++ b/Engine/Source/Runtime/Windows/D3D11RHI/Public/D3D11Util.h
@@ -30,6 +30,15 @@
 #define D3D11RHI_IMMEDIATE_CONTEXT	(GD3D11RHI->GetDeviceContext())
 #define D3D11RHI_DEVICE				(GD3D11RHI->GetDevice())
 
+/**
+ * Checks that the given result isn't a failure.  If it is, the application does not exit and only logs an appropriate error message.
+ * @param	Result - The result code to check.
+ * @param	Code - The code which yielded the result.
+ * @param	Filename - The filename of the source file containing Code.
+ * @param	Line - The line number of Code within Filename.
+ */
+extern D3D11RHI_API void VerifyD3D11ResultNoExit(HRESULT Result, const ANSICHAR* Code, const ANSICHAR* Filename, uint32 Line, ID3D11Device* Device);
+
 /**
  * Checks that the given result isn't a failure.  If it is, the application exits with an appropriate error message.
  * @param	Result - The result code to check.
@@ -72,7 +81,7 @@ extern D3D11RHI_API void VerifyD3D11ResizeViewportResult(HRESULT D3DResult, cons
 #define VERIFYD3D11RESULT_EX(x, Device)	{HRESULT hr = x; if (FAILED(hr)) { VerifyD3D11Result(hr,#x,__FILE__,__LINE__, Device); }}
 #define VERIFYD3D11RESULT(x)			{HRESULT hr = x; if (FAILED(hr)) { VerifyD3D11Result(hr,#x,__FILE__,__LINE__, 0); }}
 #define VERIFYD3D11SHADERRESULT(Result, Shader, Device) {HRESULT hr = (Result); if (FAILED(hr)) { VerifyD3D11ShaderResult(Shader, hr, #Result,__FILE__,__LINE__, Device); }}
-
+#define VERIFYD3D11RESULT_NOEXIT(x)		{HRESULT hr = x; if (FAILED(hr)) { VerifyD3D11ResultNoExit(hr,#x,__FILE__,__LINE__, 0); }}
 #define VERIFYD3D11CREATETEXTURERESULT(x,SizeX,SizeY,SizeZ,Format,NumMips,Flags,Usage,CPUAccessFlags,MiscFlags,SampleCount,SampleQuality,SubResPtr,SubResPitch,SubResSlicePitch,Device) {HRESULT hr = x; if (FAILED(hr)) { VerifyD3D11CreateTextureResult(hr,#x,__FILE__,__LINE__,SizeX,SizeY,SizeZ,Format,NumMips,Flags,Usage,CPUAccessFlags,MiscFlags,SampleCount,SampleQuality,SubResPtr,SubResPitch,SubResSlicePitch,Device); }}
 #define VERIFYD3D11RESIZEVIEWPORTRESULT(x,SizeX,SizeY,Format, Device) {HRESULT hr = x; if (FAILED(hr)) { VerifyD3D11ResizeViewportResult(hr,#x,__FILE__,__LINE__,SizeX,SizeY,Format, Device); }}
 /**
