Add missing include directive to fix HoloLens compilation.

#rnx

diff --git a/Engine/Source/Runtime/Windows/D3D11RHI/Private/Windows/WindowsD3D11Device.cpp b/Engine/Source/Runtime/Windows/D3D11RHI/Private/Windows/WindowsD3D11Device.cpp
--- a/Engine/Source/Runtime/Windows/D3D11RHI/Private/Windows/WindowsD3D11Device.cpp
+++ b/Engine/Source/Runtime/Windows/D3D11RHI/Private/Windows/WindowsD3D11Device.cpp
@@ -7,6 +7,7 @@
 #include "D3D11RHIPrivate.h"
 #include "Misc/CommandLine.h"
 #include "Misc/EngineVersion.h"
+#include "Windows/WindowsPlatformMisc.h"
 #include "Windows/AllowWindowsPlatformTypes.h"
 	#include <delayimp.h>
 	#if !PLATFORM_HOLOLENS
