Another attempt to fix Hololens compile error.

#rnx

diff --git a/Engine/Source/Runtime/Windows/D3D11RHI/Private/Windows/WindowsD3D11Device.cpp b/Engine/Source/Runtime/Windows/D3D11RHI/Private/Windows/WindowsD3D11Device.cpp
--- a/Engine/Source/Runtime/Windows/D3D11RHI/Private/Windows/WindowsD3D11Device.cpp
+++ b/Engine/Source/Runtime/Windows/D3D11RHI/Private/Windows/WindowsD3D11Device.cpp
@@ -7,7 +7,6 @@
 #include "D3D11RHIPrivate.h"
 #include "Misc/CommandLine.h"
 #include "Misc/EngineVersion.h"
-#include "Windows/WindowsPlatformMisc.h"
 #include "Windows/AllowWindowsPlatformTypes.h"
 	#include <delayimp.h>
 	#if !PLATFORM_HOLOLENS
@@ -373,7 +372,11 @@ static bool SafeTestD3D11CreateDevice(IDXGIAdapter* Adapter,D3D_FEATURE_LEVEL Mi
 		// Log any reason for failure to create test device. Extra debug help.
 		VERIFYD3D11RESULT_NOEXIT(Result);
 
+#if PLATFORM_HOLOLENS
+		bool bIsWin10 = true;
+#else
 		bool bIsWin10 = FWindowsPlatformMisc::VerifyWindowsVersion(10, 0);
+#endif
 
 		// Fatal error on 0x887A002D
 		if (DXGI_ERROR_SDK_COMPONENT_MISSING == Result && bIsWin10)
