fix UE-76115. Raise timeout for fetching timer queries

diff --git a/Engine/Source/Runtime/Windows/D3D11RHI/Private/D3D11Query.cpp b/Engine/Source/Runtime/Windows/D3D11RHI/Private/D3D11Query.cpp
--- a/Engine/Source/Runtime/Windows/D3D11RHI/Private/D3D11Query.cpp
+++ b/Engine/Source/Runtime/Windows/D3D11RHI/Private/D3D11Query.cpp
@@ -298,14 +298,21 @@ bool FD3D11DynamicRHI::GetQueryData(ID3D11Query* Query, void* Data, SIZE_T DataS
 		SCOPE_CYCLE_COUNTER( STAT_RenderQueryResultTime );
 		uint32 IdleStart = FPlatformTime::Cycles();
 		double StartTime = FPlatformTime::Seconds();
-		do 
+		double TimeoutWarningLimit = 5.0;
+		// timer queries are used for Benchmarks which can stall a bit more
+		double TimeoutValue = (QueryType == RQT_AbsoluteTime) ? 30.0 : 0.5;
+
+		do
 		{
 			SAFE_GET_QUERY_DATA
+			float DeltaTime = FPlatformTime::Seconds() - StartTime;
+			if(DeltaTime > TimeoutWarningLimit)
+			{
+				TimeoutWarningLimit += 5.0;
+				UE_LOG(LogD3D11RHI, Log, TEXT("GetQueryData is taking a very long time (%.1f s)"), DeltaTime);
+			}
 
-			// timer queries are used for Benchmarks which can stall a bit more
-			double TimeoutValue = (QueryType == RQT_AbsoluteTime) ? 2.0 : 0.5;
-
-			if((FPlatformTime::Seconds() - StartTime) > TimeoutValue)
+			if(DeltaTime > TimeoutValue)
 			{
 				UE_LOG(LogD3D11RHI, Log, TEXT("Timed out while waiting for GPU to catch up. (%.1f s)"), TimeoutValue);
 				return false;
