Fix scissor rect ensures. Also add scissor validation to DX11.

diff --git a/Engine/Source/Runtime/SlateRHIRenderer/Private/SlateRHIRenderingPolicy.cpp b/Engine/Source/Runtime/SlateRHIRenderer/Private/SlateRHIRenderingPolicy.cpp
--- a/Engine/Source/Runtime/SlateRHIRenderer/Private/SlateRHIRenderingPolicy.cpp
+++ b/Engine/Source/Runtime/SlateRHIRenderer/Private/SlateRHIRenderingPolicy.cpp
@@ -340,12 +340,15 @@ static bool UpdateScissorRect(
 
 				const FSlateClippingZone& ScissorRect = ClipState.ScissorRect.GetValue();
 
-				const FVector2D TopLeft = ScissorRect.TopLeft + ViewTranslation2D;
-				const FVector2D BottomRight = ScissorRect.BottomRight + ViewTranslation2D;
+				const FIntPoint SizeXY = BackBuffer.GetSizeXY();
+				const FVector2D ViewSize((float) SizeXY.X, (float) SizeXY.Y);
 
+				// Clamp scissor rect to BackBuffer size
+				const FVector2D TopLeft     = FMath::Min(FMath::Max(ScissorRect.TopLeft     + ViewTranslation2D, FVector2D(0.0f, 0.0f)), ViewSize);
+				const FVector2D BottomRight = FMath::Min(FMath::Max(ScissorRect.BottomRight + ViewTranslation2D, FVector2D(0.0f, 0.0f)), ViewSize);
+				
 				if (bSwitchVerticalAxis)
 				{
-					const FIntPoint ViewSize = BackBuffer.GetSizeXY();
 					const int32 MinY = (ViewSize.Y - BottomRight.Y);
 					const int32 MaxY = (ViewSize.Y - TopLeft.Y);
 					RHICmdList.SetScissorRect(true, TopLeft.X, MinY, BottomRight.X, MaxY);
diff --git a/Engine/Source/Runtime/Windows/D3D11RHI/Private/D3D11Commands.cpp b/Engine/Source/Runtime/Windows/D3D11RHI/Private/D3D11Commands.cpp
--- a/Engine/Source/Runtime/Windows/D3D11RHI/Private/D3D11Commands.cpp
+++ b/Engine/Source/Runtime/Windows/D3D11RHI/Private/D3D11Commands.cpp
@@ -254,26 +254,38 @@ void FD3D11DynamicRHI::RHISetViewport(uint32 MinX,uint32 MinY,float MinZ,uint32
 	}
 }
 
+static void ValidateScissorRect(const D3D11_VIEWPORT& Viewport, const D3D11_RECT& ScissorRect)
+{
+	ensure(ScissorRect.left   >= (LONG)Viewport.TopLeftX);
+	ensure(ScissorRect.top    >= (LONG)Viewport.TopLeftY);
+	ensure(ScissorRect.right  <= (LONG)Viewport.TopLeftX + (LONG)Viewport.Width);
+	ensure(ScissorRect.bottom <= (LONG)Viewport.TopLeftY + (LONG)Viewport.Height);
+	ensure(ScissorRect.left <= ScissorRect.right && ScissorRect.top <= ScissorRect.bottom);
+}
+
 void FD3D11DynamicRHI::RHISetScissorRect(bool bEnable,uint32 MinX,uint32 MinY,uint32 MaxX,uint32 MaxY)
 {
-	if(bEnable)
+	D3D11_VIEWPORT Viewport;
+	StateCache.GetViewport(&Viewport);
+
+	D3D11_RECT ScissorRect;
+	if (bEnable)
 	{
-		D3D11_RECT ScissorRect;
-		ScissorRect.left = MinX;
-		ScissorRect.right = MaxX;
-		ScissorRect.top = MinY;
+		ScissorRect.left   = MinX;
+		ScissorRect.top    = MinY;
+		ScissorRect.right  = MaxX;
 		ScissorRect.bottom = MaxY;
-		Direct3DDeviceIMContext->RSSetScissorRects(1,&ScissorRect);
 	}
 	else
 	{
-		D3D11_RECT ScissorRect;
-		ScissorRect.left = 0;
-		ScissorRect.right = GetMax2DTextureDimension();
-		ScissorRect.top = 0;
-		ScissorRect.bottom = GetMax2DTextureDimension();
-		Direct3DDeviceIMContext->RSSetScissorRects(1,&ScissorRect);
+		ScissorRect.left   = (LONG) Viewport.TopLeftX;
+		ScissorRect.top    = (LONG) Viewport.TopLeftY;
+		ScissorRect.right  = (LONG) Viewport.TopLeftX + (LONG) Viewport.Width;
+		ScissorRect.bottom = (LONG) Viewport.TopLeftY + (LONG) Viewport.Height;
 	}
+
+	ValidateScissorRect(Viewport, ScissorRect);
+	Direct3DDeviceIMContext->RSSetScissorRects(1, &ScissorRect);
 }
 
 /**
