fix for null vertexbuffer when creating an SRV on a vertexbuffer

diff --git a/Engine/Source/Runtime/Windows/D3D11RHI/Private/D3D11UAV.cpp b/Engine/Source/Runtime/Windows/D3D11RHI/Private/D3D11UAV.cpp
--- a/Engine/Source/Runtime/Windows/D3D11RHI/Private/D3D11UAV.cpp
+++ b/Engine/Source/Runtime/Windows/D3D11RHI/Private/D3D11UAV.cpp
@@ -260,13 +260,11 @@ static void CreateD3D11ShaderResourceViewOnBuffer(ID3D11Device* Direct3DDevice,
 
 FShaderResourceViewRHIRef FD3D11DynamicRHI::RHICreateShaderResourceView(FRHIVertexBuffer* VertexBufferRHI, uint32 Stride, uint8 Format)
 {
-	if (!VertexBufferRHI)
+	FD3D11VertexBuffer* VertexBuffer = ResourceCast(VertexBufferRHI);
+	if (!VertexBufferRHI || !VertexBuffer->Resource)
 	{
 		return new FD3D11ShaderResourceView(nullptr, nullptr);
 	}
-	FD3D11VertexBuffer* VertexBuffer = ResourceCast(VertexBufferRHI);
-	check(VertexBuffer);
-	check(VertexBuffer->Resource);
 
 	TRefCountPtr<ID3D11ShaderResourceView> ShaderResourceView;
 	CreateD3D11ShaderResourceViewOnBuffer(Direct3DDevice, VertexBuffer->Resource, Stride, Format, ShaderResourceView.GetInitReference());
