fix for UE-76115

diff --git a/Engine/Source/Runtime/Windows/D3D11RHI/Private/D3D11Query.cpp b/Engine/Source/Runtime/Windows/D3D11RHI/Private/D3D11Query.cpp
--- a/Engine/Source/Runtime/Windows/D3D11RHI/Private/D3D11Query.cpp
+++ b/Engine/Source/Runtime/Windows/D3D11RHI/Private/D3D11Query.cpp
@@ -305,6 +305,14 @@ bool FD3D11DynamicRHI::GetQueryData(ID3D11Query* Query, void* Data, SIZE_T DataS
 		do
 		{
 			SAFE_GET_QUERY_DATA
+
+			if(Result == S_OK)
+			{
+				return true;
+			}
+
+
+
 			float DeltaTime = FPlatformTime::Seconds() - StartTime;
 			if(DeltaTime > TimeoutWarningLimit)
 			{
