If there are no samples in a video recording session, the session should not be allowed to start

#rnx

diff --git a/Engine/Source/Runtime/Windows/WindowsPlatformFeatures/Private/HighlightRecorder.cpp b/Engine/Source/Runtime/Windows/WindowsPlatformFeatures/Private/HighlightRecorder.cpp
--- a/Engine/Source/Runtime/Windows/WindowsPlatformFeatures/Private/HighlightRecorder.cpp
+++ b/Engine/Source/Runtime/Windows/WindowsPlatformFeatures/Private/HighlightRecorder.cpp
@@ -347,7 +347,8 @@ bool FHighlightRecorder::GetSavingStart(const TArray<FGameplayMediaEncoderSample
 	if (Samples.Num() == 0)
 	{
 		StartIndex = 0;
-		return true;
+		UE_LOG(HighlightRecorder, Error, TEXT("No samples to write to .mp4, max duration: %.3f"), MaxDuration.GetTotalSeconds());
+		return false;
 	}
 
 	FTimespan FirstTimestamp = Samples[0].GetTime();
@@ -375,7 +376,7 @@ bool FHighlightRecorder::GetSavingStart(const TArray<FGameplayMediaEncoderSample
 
 	if (!bFound)
 	{
-		UE_LOG(HighlightRecorder, Error, TEXT("No samples to write to .mp4, max duration: %.3f"), MaxDuration.GetTotalSeconds());
+		UE_LOG(HighlightRecorder, Error, TEXT("No samples found to write to .mp4, max duration: %.3f"), MaxDuration.GetTotalSeconds());
 		return false;
 	}
 
