check for device removed when timing out in GetQueryData. This should hopefully classify the remaining crashes as gpu hangs.

diff --git a/Engine/Source/Runtime/Windows/D3D11RHI/Private/D3D11Query.cpp b/Engine/Source/Runtime/Windows/D3D11RHI/Private/D3D11Query.cpp
--- a/Engine/Source/Runtime/Windows/D3D11RHI/Private/D3D11Query.cpp
+++ b/Engine/Source/Runtime/Windows/D3D11RHI/Private/D3D11Query.cpp
@@ -322,7 +322,11 @@ bool FD3D11DynamicRHI::GetQueryData(ID3D11Query* Query, void* Data, SIZE_T DataS
 
 			if(DeltaTime > TimeoutValue)
 			{
-				UE_LOG(LogD3D11RHI, Log, TEXT("Timed out while waiting for GPU to catch up. (%.1f s)"), TimeoutValue);
+				UE_LOG(LogD3D11RHI, Log, TEXT("Timed out while waiting for GPU to catch up. (%.1f s) (ErrorCode %08x)"), TimeoutValue, (uint32)Result);
+				if(FAILED(Result))
+				{
+					VERIFYD3D11RESULT_EX(Result, Direct3DDevice);
+				}
 				return false;
 			}
 		} while ( Result == S_FALSE );
