Fixed a typo.

diff --git a/Engine/Source/Runtime/Windows/AudioMixerXAudio2/Private/AudioMixerPlatformXAudio2.cpp b/Engine/Source/Runtime/Windows/AudioMixerXAudio2/Private/AudioMixerPlatformXAudio2.cpp
--- a/Engine/Source/Runtime/Windows/AudioMixerXAudio2/Private/AudioMixerPlatformXAudio2.cpp
+++ b/Engine/Source/Runtime/Windows/AudioMixerXAudio2/Private/AudioMixerPlatformXAudio2.cpp
@@ -120,7 +120,7 @@ namespace Audio
 			case E_INVALIDARG:								return TEXT("E_INVALIDARG");
 			case E_OUTOFMEMORY:								return TEXT("E_OUTOFMEMORY");
 #endif
-			default:										return TEXT("UKNOWN");
+			default:										return TEXT("UNKNOWN");
 		}
 	}
 
