XAudio2 error logging improvements.

diff --git a/Engine/Source/Runtime/Windows/AudioMixerXAudio2/Private/AudioMixerPlatformXAudio2.cpp b/Engine/Source/Runtime/Windows/AudioMixerXAudio2/Private/AudioMixerPlatformXAudio2.cpp
--- a/Engine/Source/Runtime/Windows/AudioMixerXAudio2/Private/AudioMixerPlatformXAudio2.cpp
+++ b/Engine/Source/Runtime/Windows/AudioMixerXAudio2/Private/AudioMixerPlatformXAudio2.cpp
@@ -33,21 +33,21 @@
 #include "Logging/LogMacros.h"
 
 // Macro to check result code for XAudio2 failure, get the string version, log, and goto a cleanup
-#define XAUDIO2_CLEANUP_ON_FAIL(Result)						\
-	if (FAILED(Result))										\
-	{														\
-		const TCHAR* ErrorString = GetErrorString(Result);	\
-		AUDIO_PLATFORM_ERROR(ErrorString);					\
-		goto Cleanup;										\
+#define XAUDIO2_CLEANUP_ON_FAIL(Result)																\
+	if (FAILED(Result))																				\
+	{																								\
+		FString ErrorString = FString::Printf(TEXT("0x%X: %s"), Result, GetErrorString(Result));	\
+		AUDIO_PLATFORM_ERROR(*ErrorString);															\
+		goto Cleanup;																				\
 	}
 
 // Macro to check result for XAudio2 failure, get string version, log, and return false
-#define XAUDIO2_RETURN_ON_FAIL(Result)						\
-	if (FAILED(Result))										\
-	{														\
-		const TCHAR* ErrorString = GetErrorString(Result);	\
-		AUDIO_PLATFORM_ERROR(ErrorString);					\
-		return false;										\
+#define XAUDIO2_RETURN_ON_FAIL(Result)																\
+	if (FAILED(Result))																				\
+	{																								\
+		FString ErrorString = FString::Printf(TEXT("0x%X: %s"), Result, GetErrorString(Result));	\
+		AUDIO_PLATFORM_ERROR(*ErrorString);															\
+		return false;																				\
 	}
 
 
