UE-79137 HoloLens - No engine audio plays in emulator
-Hololens now iterates through audio devices starting at zero trying to find one that it can successfully open a stream with.  The hololens emulator device 0 cannot, but device 1 can.  Hololens device has only one device.

diff --git a/Engine/Source/Runtime/Windows/AudioMixerXAudio2/Private/AudioMixerPlatformXAudio2.cpp b/Engine/Source/Runtime/Windows/AudioMixerXAudio2/Private/AudioMixerPlatformXAudio2.cpp
--- a/Engine/Source/Runtime/Windows/AudioMixerXAudio2/Private/AudioMixerPlatformXAudio2.cpp
+++ b/Engine/Source/Runtime/Windows/AudioMixerXAudio2/Private/AudioMixerPlatformXAudio2.cpp
@@ -540,12 +540,13 @@ namespace Audio
 
 		OpenStreamParams = Params;
 
+#if !PLATFORM_HOLOLENS
 		// On windows, default device index is 0
 		if (Params.OutputDeviceIndex == AUDIO_MIXER_DEFAULT_DEVICE_INDEX)
 		{
 			OpenStreamParams.OutputDeviceIndex = 0;
 		}
-
+#endif
 		AudioStreamInfo.Reset();
 
 		AudioStreamInfo.OutputDeviceIndex = OpenStreamParams.OutputDeviceIndex;
@@ -558,10 +559,41 @@ namespace Audio
 
 		if (GetNumOutputDevices(NumOutputDevices) && NumOutputDevices > 0)
 		{
+#if PLATFORM_HOLOLENS
+			// On windows, default device index is 0
+			// But if that device cannot be configured try to find one that can be.  This happens in the hololens emulator.
+			if (AudioStreamInfo.OutputDeviceIndex == AUDIO_MIXER_DEFAULT_DEVICE_INDEX)
+			{
+				bool bFoundUsefulDevice = false;
+				for (uint32 i = 0; i < NumOutputDevices; ++i)
+				{
+					OpenStreamParams.OutputDeviceIndex = i;
+					AudioStreamInfo.OutputDeviceIndex = OpenStreamParams.OutputDeviceIndex;
+					if (GetOutputDeviceInfo(AudioStreamInfo.OutputDeviceIndex, AudioStreamInfo.DeviceInfo))
+					{
+						bFoundUsefulDevice = true;
+						break;
+					}
+				}
+				if (!bFoundUsefulDevice)
+				{
+					return false;
+				}
+			}
+			else
+			{
+				if (!GetOutputDeviceInfo(AudioStreamInfo.OutputDeviceIndex, AudioStreamInfo.DeviceInfo))
+				{
+					return false;
+				}
+			}
+#else
 			if (!GetOutputDeviceInfo(AudioStreamInfo.OutputDeviceIndex, AudioStreamInfo.DeviceInfo))
 			{
 				return false;
 			}
+#endif
+
 
 			// Store the device ID here in case it is removed. We can switch back if the device comes back.
 			if (Params.bRestoreIfRemoved)
