Add support for DXGI_SWAP_EFFECT_FLIP_DISCARD & DXGI_SWAP_CHAIN_FLAG_ALLOW_TEARING, when running on windows version that supports it
fixes a rare crash when calling present

diff --git a/Engine/Source/Runtime/Windows/D3D11RHI/Private/D3D11Viewport.cpp b/Engine/Source/Runtime/Windows/D3D11RHI/Private/D3D11Viewport.cpp
--- a/Engine/Source/Runtime/Windows/D3D11RHI/Private/D3D11Viewport.cpp
+++ b/Engine/Source/Runtime/Windows/D3D11RHI/Private/D3D11Viewport.cpp
@@ -85,6 +85,7 @@ namespace RHIConsoleVariables
 };
 
 extern void D3D11TextureAllocated2D( FD3D11Texture2D& Texture );
+extern uint32 D3D11GetSwapChainFlags();
 
 /**
  * Creates a FD3D11Surface to represent a swap chain's back buffer.
@@ -259,12 +260,8 @@ void FD3D11Viewport::Resize(uint32 InSizeX, uint32 InSizeY, bool bInIsFullscreen
 		{
 			// Resize the swap chain.
 			DXGI_FORMAT RenderTargetFormat = GetRenderTargetFormat(PixelFormat);
-#if PLATFORM_HOLOLENS
 			// Resize all existing buffers, don't change count
-			VERIFYD3D11RESIZEVIEWPORTRESULT(SwapChain->ResizeBuffers(0, SizeX, SizeY, RenderTargetFormat, DXGI_SWAP_CHAIN_FLAG_ALLOW_MODE_SWITCH), SizeX, SizeY, RenderTargetFormat, D3DRHI->GetDevice());
-#else
-			VERIFYD3D11RESIZEVIEWPORTRESULT(SwapChain->ResizeBuffers(BackBufferCount, SizeX, SizeY, RenderTargetFormat, DXGI_SWAP_CHAIN_FLAG_ALLOW_MODE_SWITCH), SizeX, SizeY, RenderTargetFormat, D3DRHI->GetDevice());
-#endif
+			VERIFYD3D11RESIZEVIEWPORTRESULT(SwapChain->ResizeBuffers(0, SizeX, SizeY, RenderTargetFormat, D3D11GetSwapChainFlags()), SizeX, SizeY, RenderTargetFormat, D3DRHI->GetDevice());
 
 			if (bInIsFullscreen)
 			{
@@ -372,7 +369,8 @@ bool FD3D11Viewport::PresentChecked(int32 SyncInterval)
 	FThreadHeartBeat::Get().PresentFrame();
 
 	VERIFYD3D11RESULT_EX(Result, D3DRHI->GetDevice());
-	
+	D3DRHI->GetDeviceContext()->OMSetRenderTargets(0,0,0);
+
 	return bNeedNativePresent;
 }
 
diff --git a/Engine/Source/Runtime/Windows/D3D11RHI/Private/Windows/WindowsD3D11Viewport.cpp b/Engine/Source/Runtime/Windows/D3D11RHI/Private/Windows/WindowsD3D11Viewport.cpp
--- a/Engine/Source/Runtime/Windows/D3D11RHI/Private/Windows/WindowsD3D11Viewport.cpp
+++ b/Engine/Source/Runtime/Windows/D3D11RHI/Private/Windows/WindowsD3D11Viewport.cpp
@@ -13,6 +13,30 @@
 
 #include "dxgi1_6.h"
 
+#ifndef DXGI_PRESENT_ALLOW_TEARING
+#define DXGI_PRESENT_ALLOW_TEARING          0x00000200UL
+#define DXGI_SWAP_CHAIN_FLAG_ALLOW_TEARING  2048
+#endif
+
+
+
+static bool GSwapFlagsInitialized = false;
+static DXGI_SWAP_EFFECT GSwapEffect = DXGI_SWAP_EFFECT_DISCARD;
+static uint32 GSwapChainFlags = DXGI_SWAP_CHAIN_FLAG_ALLOW_MODE_SWITCH;
+static uint32 GSwapChainBufferCount = 1;
+
+uint32 D3D11GetSwapChainFlags()
+{
+	return GSwapChainFlags;
+}
+
+static int32 GD3D11UseAllowTearing = 1;
+static FAutoConsoleVariableRef CVarD3DUseAllowTearing(
+	TEXT("r.D3D11.UseAllowTearing"),
+	GD3D11UseAllowTearing,
+	TEXT("Enable new dxgi flip mode with d3d11"),
+	ECVF_RenderThreadSafe| ECVF_ReadOnly
+);
 
 
 FD3D11Viewport::FD3D11Viewport(FD3D11DynamicRHI* InD3DRHI,HWND InWindowHandle,uint32 InSizeX,uint32 InSizeY,bool bInIsFullscreen, EPixelFormat InPreferredPixelFormat):
@@ -42,6 +66,27 @@ FD3D11Viewport::FD3D11Viewport(FD3D11DynamicRHI* InD3DRHI,HWND InWindowHandle,ui
 	TRefCountPtr<IDXGIDevice> DXGIDevice;
 	VERIFYD3D11RESULT_EX(D3DRHI->GetDevice()->QueryInterface(IID_IDXGIDevice, (void**)DXGIDevice.GetInitReference()), D3DRHI->GetDevice());
 
+	if(!GSwapFlagsInitialized)
+	{
+		IDXGIFactory1* Factory1 = D3DRHI->GetFactory();
+		TRefCountPtr<IDXGIFactory5> Factory5;
+
+		if(GD3D11UseAllowTearing)
+		{
+			if (S_OK == Factory1->QueryInterface(__uuidof(IDXGIFactory5), (void**)Factory5.GetInitReference()))
+			{
+				UINT AllowTearing = 0;
+				if (S_OK == Factory5->CheckFeatureSupport(DXGI_FEATURE_PRESENT_ALLOW_TEARING, &AllowTearing, sizeof(UINT)) && AllowTearing != 0)
+				{
+					GSwapEffect = DXGI_SWAP_EFFECT_FLIP_DISCARD;
+					GSwapChainFlags |= DXGI_SWAP_CHAIN_FLAG_ALLOW_TEARING;
+					GSwapChainBufferCount = 2;
+				}
+			}
+		}
+		GSwapFlagsInitialized = true;
+	}
+	uint32 BufferCount = GSwapChainBufferCount;
 	// If requested, keep a handle to a DXGIOutput so we can force that display on fullscreen swap
 	uint32 DisplayIndex = D3DRHI->GetHDRDetectedDisplayIndex();
 	bForcedFullscreenDisplay = FParse::Value(FCommandLine::Get(), TEXT("FullscreenDisplay="), DisplayIndex);
@@ -96,7 +141,7 @@ FD3D11Viewport::FD3D11Viewport(FD3D11DynamicRHI* InD3DRHI,HWND InWindowHandle,ui
 				SwapChainDesc1.BufferCount = 2;
 				SwapChainDesc1.Scaling = DXGI_SCALING_NONE;
 				SwapChainDesc1.SwapEffect = DXGI_SWAP_EFFECT_FLIP_DISCARD;
-				SwapChainDesc1.Flags = DXGI_SWAP_CHAIN_FLAG_ALLOW_MODE_SWITCH;
+				SwapChainDesc1.Flags = GSwapChainFlags;
 
 				IDXGISwapChain1* SwapChain1 = nullptr;
 				VERIFYD3D11RESULT_EX((Factory2->CreateSwapChainForHwnd(D3DRHI->GetDevice(), WindowHandle, &SwapChainDesc1, nullptr, nullptr, &SwapChain1)), D3DRHI->GetDevice());
@@ -127,9 +172,9 @@ FD3D11Viewport::FD3D11Viewport(FD3D11DynamicRHI* InD3DRHI,HWND InWindowHandle,ui
 
 			// Needed for HDR
 			BackBufferCount = 2;
-			SwapChainDesc.SwapEffect = DXGI_SWAP_EFFECT_FLIP_DISCARD;
+			SwapChainDesc.SwapEffect = GSwapEffect;
 			SwapChainDesc.BufferCount = BackBufferCount;
-			SwapChainDesc.Flags = DXGI_SWAP_CHAIN_FLAG_ALLOW_MODE_SWITCH;
+			SwapChainDesc.Flags = GSwapChainFlags;
 
 			IDXGISwapChain1* SwapChain1 = nullptr;
 			IDXGIFactory2* Factory2 = (IDXGIFactory2*)D3DRHI->GetFactory();
@@ -144,7 +189,7 @@ FD3D11Viewport::FD3D11Viewport(FD3D11DynamicRHI* InD3DRHI,HWND InWindowHandle,ui
 
 		if (SwapChain == nullptr)
 		{
-			BackBufferCount = 1;
+			BackBufferCount = GSwapChainBufferCount;
 
 			// Create the swapchain.
 			DXGI_SWAP_CHAIN_DESC SwapChainDesc;
@@ -160,8 +205,8 @@ FD3D11Viewport::FD3D11Viewport(FD3D11DynamicRHI* InD3DRHI,HWND InWindowHandle,ui
 			SwapChainDesc.OutputWindow = WindowHandle;
 			SwapChainDesc.Windowed = !bIsFullscreen;
 			// DXGI_SWAP_EFFECT_DISCARD / DXGI_SWAP_EFFECT_SEQUENTIAL
-			SwapChainDesc.SwapEffect = DXGI_SWAP_EFFECT_DISCARD;
-			SwapChainDesc.Flags = DXGI_SWAP_CHAIN_FLAG_ALLOW_MODE_SWITCH;
+			SwapChainDesc.SwapEffect = GSwapEffect;
+			SwapChainDesc.Flags = GSwapChainFlags;
 			VERIFYD3D11RESULT_EX(D3DRHI->GetFactory()->CreateSwapChain(DXGIDevice, &SwapChainDesc, SwapChain.GetInitReference()), D3DRHI->GetDevice());
 		}
 
