Fix broken hololens build with missing swapchain flags accessor. HololensD3D11Viewport implementation needed a version of this function (as well as the global to track the flags).

diff --git a/Engine/Source/Runtime/Windows/D3D11RHI/Private/HoloLens/HoloLensD3D11Viewport.cpp b/Engine/Source/Runtime/Windows/D3D11RHI/Private/HoloLens/HoloLensD3D11Viewport.cpp
--- a/Engine/Source/Runtime/Windows/D3D11RHI/Private/HoloLens/HoloLensD3D11Viewport.cpp
+++ b/Engine/Source/Runtime/Windows/D3D11RHI/Private/HoloLens/HoloLensD3D11Viewport.cpp
@@ -14,6 +14,13 @@ using namespace Windows::Foundation;
 using namespace Windows::ApplicationModel::Core;
 using namespace Windows::UI::Core;
 
+static uint32 GSwapChainFlags = DXGI_SWAP_CHAIN_FLAG_ALLOW_MODE_SWITCH;
+
+uint32 D3D11GetSwapChainFlags()
+{
+	return GSwapChainFlags;
+}
+
 DXGI_FORMAT GetSupportedSwapChainBufferFormat(DXGI_FORMAT InPreferredDXGIFormat)
 {
 	switch (InPreferredDXGIFormat)
@@ -114,7 +121,7 @@ FD3D11Viewport::FD3D11Viewport(FD3D11DynamicRHI* InD3DRHI,HWND InWindowHandle,ui
 				SwapChainDesc1.BufferCount = 2;
 				SwapChainDesc1.Scaling = DXGI_SCALING_NONE;
 				SwapChainDesc1.SwapEffect = DXGI_SWAP_EFFECT_FLIP_DISCARD;
-				SwapChainDesc1.Flags = DXGI_SWAP_CHAIN_FLAG_ALLOW_MODE_SWITCH;
+				SwapChainDesc1.Flags = GSwapChainFlags = DXGI_SWAP_CHAIN_FLAG_ALLOW_MODE_SWITCH;
 
 				IDXGISwapChain1* SwapChain1 = nullptr;
 				VERIFYD3D11RESULT_EX((Factory2->CreateSwapChainForCoreWindow(
@@ -158,7 +165,7 @@ FD3D11Viewport::FD3D11Viewport(FD3D11DynamicRHI* InD3DRHI,HWND InWindowHandle,ui
 			swapChainDesc.Scaling = DXGI_SCALING_NONE;
 			swapChainDesc.AlphaMode = DXGI_ALPHA_MODE_IGNORE;
 			swapChainDesc.SwapEffect = DXGI_SWAP_EFFECT_FLIP_SEQUENTIAL; // we recommend using this swap effect for all applications
-			swapChainDesc.Flags = 0;
+			swapChainDesc.Flags = GSwapChainFlags = 0;
 
 			IDXGISwapChain1* CreatedSwapChain = NULL;
 
