Fix for Resizing an editor window causes text to flicker between the original and resize locations

diff --git a/Engine/Source/Runtime/Windows/D3D11RHI/Private/Windows/WindowsD3D11Viewport.cpp b/Engine/Source/Runtime/Windows/D3D11RHI/Private/Windows/WindowsD3D11Viewport.cpp
--- a/Engine/Source/Runtime/Windows/D3D11RHI/Private/Windows/WindowsD3D11Viewport.cpp
+++ b/Engine/Source/Runtime/Windows/D3D11RHI/Private/Windows/WindowsD3D11Viewport.cpp
@@ -175,6 +175,7 @@ FD3D11Viewport::FD3D11Viewport(FD3D11DynamicRHI* InD3DRHI,HWND InWindowHandle,ui
 			SwapChainDesc.SwapEffect = GSwapEffect;
 			SwapChainDesc.BufferCount = BackBufferCount;
 			SwapChainDesc.Flags = GSwapChainFlags;
+			SwapChainDesc.Scaling = DXGI_SCALING_NONE;
 
 			IDXGISwapChain1* SwapChain1 = nullptr;
 			IDXGIFactory2* Factory2 = (IDXGIFactory2*)D3DRHI->GetFactory();
