Remove some dead code

diff --git a/Engine/Source/Runtime/D3D12RHI/Private/D3D12CommandContext.cpp b/Engine/Source/Runtime/D3D12RHI/Private/D3D12CommandContext.cpp
--- a/Engine/Source/Runtime/D3D12RHI/Private/D3D12CommandContext.cpp
+++ b/Engine/Source/Runtime/D3D12RHI/Private/D3D12CommandContext.cpp
@@ -64,8 +64,6 @@ FD3D12CommandContext::FD3D12CommandContext(FD3D12Device* InParent, FD3D12SubAllo
 #if PLATFORM_SUPPORTS_VIRTUAL_TEXTURES
 	bNeedFlushTextureCache(false),
 #endif
-	DynamicVB(InParent),
-	DynamicIB(InParent),
 	VSConstantBuffer(InParent, ConstantsAllocator),
 	HSConstantBuffer(InParent, ConstantsAllocator),
 	DSConstantBuffer(InParent, ConstantsAllocator),
diff --git a/Engine/Source/Runtime/D3D12RHI/Private/D3D12CommandContext.h b/Engine/Source/Runtime/D3D12RHI/Private/D3D12CommandContext.h
--- a/Engine/Source/Runtime/D3D12RHI/Private/D3D12CommandContext.h
+++ b/Engine/Source/Runtime/D3D12RHI/Private/D3D12CommandContext.h
@@ -184,15 +184,11 @@ public:
 	uint32 numCopies;
 	uint32 otherWorkCounter;
 
-	bool HasDoneWork()
+	bool HasDoneWork() const
 	{
 		return (numDraws + numDispatches + numClears + numBarriers + numCopies + otherWorkCounter) > 0;
 	}
 
-	/** Dynamic vertex and index buffers. */
-	FD3D12DynamicBuffer DynamicVB;
-	FD3D12DynamicBuffer DynamicIB;
-
 	/** Constant buffers for Set*ShaderParameter calls. */
 	FD3D12ConstantBuffer VSConstantBuffer;
 	FD3D12ConstantBuffer HSConstantBuffer;
diff --git a/Engine/Source/Runtime/D3D12RHI/Private/D3D12Resources.cpp b/Engine/Source/Runtime/D3D12RHI/Private/D3D12Resources.cpp
--- a/Engine/Source/Runtime/D3D12RHI/Private/D3D12Resources.cpp
+++ b/Engine/Source/Runtime/D3D12RHI/Private/D3D12Resources.cpp
@@ -606,29 +606,3 @@ void FD3D12ResourceLocation::SetResource(FD3D12Resource* Value)
 	UnderlyingResource = Value;
 	ResidencyHandle = UnderlyingResource->GetResidencyHandle();
 }
-
-/////////////////////////////////////////////////////////////////////
-//	FD3D12 Dynamic Buffer
-/////////////////////////////////////////////////////////////////////
-
-FD3D12DynamicBuffer::FD3D12DynamicBuffer(FD3D12Device* InParent)
-	: FD3D12DeviceChild(InParent)
-	, ResourceLocation(InParent)
-{
-}
-
-FD3D12DynamicBuffer::~FD3D12DynamicBuffer()
-{
-}
-
-void* FD3D12DynamicBuffer::Lock(uint32 Size)
-{
-	FD3D12Adapter* Adapter = GetParentDevice()->GetParentAdapter();
-
-	return 	Adapter->GetUploadHeapAllocator(GetParentDevice()->GetGPUIndex()).AllocUploadResource(Size, D3D12_TEXTURE_DATA_PLACEMENT_ALIGNMENT, ResourceLocation);
-}
-
-FD3D12ResourceLocation* FD3D12DynamicBuffer::Unlock()
-{
-	return &ResourceLocation;
-}
diff --git a/Engine/Source/Runtime/D3D12RHI/Public/D3D12Resources.h b/Engine/Source/Runtime/D3D12RHI/Public/D3D12Resources.h
--- a/Engine/Source/Runtime/D3D12RHI/Public/D3D12Resources.h
+++ b/Engine/Source/Runtime/D3D12RHI/Public/D3D12Resources.h
@@ -981,28 +981,6 @@ private:
 	TArray<D3D12_RESOURCE_BARRIER> Barriers;
 };
 
-/**
-* Class for managing dynamic buffers (Used for DrawUp).
-*/
-class FD3D12DynamicBuffer : public FD3D12DeviceChild
-{
-public:
-	/** Initialization constructor. */
-	FD3D12DynamicBuffer(FD3D12Device* InParent);
-	/** Destructor. */
-	~FD3D12DynamicBuffer();
-
-	/** Locks the buffer returning at least Size bytes. */
-	void* Lock(uint32 Size);
-	/** Unlocks the buffer returning the underlying D3D12 buffer to use as a resource. */
-	FD3D12ResourceLocation* Unlock();
-
-	void ReleaseResourceLocation() { ResourceLocation.Clear(); }
-
-private:
-	FD3D12ResourceLocation ResourceLocation;
-};
-
 class FD3D12StagingBuffer : public FRHIStagingBuffer
 {
 	friend class FD3D12CommandContext;
diff --git a/Engine/Source/Runtime/Windows/D3D11RHI/Private/D3D11Device.cpp b/Engine/Source/Runtime/Windows/D3D11RHI/Private/D3D11Device.cpp
--- a/Engine/Source/Runtime/Windows/D3D11RHI/Private/D3D11Device.cpp
+++ b/Engine/Source/Runtime/Windows/D3D11RHI/Private/D3D11Device.cpp
@@ -227,12 +227,6 @@ FD3D11DynamicRHI::FD3D11DynamicRHI(IDXGIFactory1* InDXGIFactory1,D3D_FEATURE_LEV
 	// Initialize the constant buffers.
 	InitConstantBuffers();
 
-	// Create the dynamic vertex and index buffers used for Draw[Indexed]PrimitiveUP.
-	uint32 DynamicVBSizes[] = {128,1024,64*1024,1024*1024,0};
-	DynamicVB = new FD3D11DynamicBuffer(this,D3D11_BIND_VERTEX_BUFFER,DynamicVBSizes);
-	uint32 DynamicIBSizes[] = {128,1024,64*1024,1024*1024,0};
-	DynamicIB = new FD3D11DynamicBuffer(this,D3D11_BIND_INDEX_BUFFER,DynamicIBSizes);
-
 	for (int32 Frequency = 0; Frequency < SF_NumStandardFrequencies; ++Frequency)
 	{
 		DirtyUniformBuffers[Frequency] = 0;
@@ -503,10 +497,6 @@ void FD3D11DynamicRHI::CleanupD3DDevice()
 		extern void EmptyD3DSamplerStateCache();
 		EmptyD3DSamplerStateCache();
 
-		// release our dynamic VB and IB buffers
-		DynamicVB = NULL;
-		DynamicIB = NULL;
-
 		// Release references to bound uniform buffers.
 		for (int32 Frequency = 0; Frequency < SF_NumStandardFrequencies; ++Frequency)
 		{
diff --git a/Engine/Source/Runtime/Windows/D3D11RHI/Private/D3D11RHIPrivate.h b/Engine/Source/Runtime/Windows/D3D11RHI/Private/D3D11RHIPrivate.h
--- a/Engine/Source/Runtime/Windows/D3D11RHI/Private/D3D11RHIPrivate.h
+++ b/Engine/Source/Runtime/Windows/D3D11RHI/Private/D3D11RHIPrivate.h
@@ -878,10 +878,6 @@ protected:
 	/** Set to true when the current shading setup uses tessellation */
 	bool bUsingTessellation;
 
-	/** Dynamic vertex and index buffers. */
-	TRefCountPtr<FD3D11DynamicBuffer> DynamicVB;
-	TRefCountPtr<FD3D11DynamicBuffer> DynamicIB;
-
 	/** A list of all D3D constant buffers RHIs that have been created. */
 	TArray<TRefCountPtr<FD3D11ConstantBuffer> > VSConstantBuffers;
 	TArray<TRefCountPtr<FD3D11ConstantBuffer> > HSConstantBuffers;
diff --git a/Engine/Source/Runtime/Windows/D3D11RHI/Private/D3D11Util.cpp b/Engine/Source/Runtime/Windows/D3D11RHI/Private/D3D11Util.cpp
--- a/Engine/Source/Runtime/Windows/D3D11RHI/Private/D3D11Util.cpp
+++ b/Engine/Source/Runtime/Windows/D3D11RHI/Private/D3D11Util.cpp
@@ -435,96 +435,6 @@ FD3D11BoundRenderTargets::~FD3D11BoundRenderTargets()
 	}
 }
 
-FD3D11DynamicBuffer::FD3D11DynamicBuffer(FD3D11DynamicRHI* InD3DRHI, D3D11_BIND_FLAG InBindFlags, uint32* InBufferSizes)
-	: D3DRHI(InD3DRHI)
-	, BindFlags(InBindFlags)
-	, LockedBufferIndex(-1)
-{
-	while (BufferSizes.Num() < MAX_BUFFER_SIZES && *InBufferSizes > 0)
-	{
-		uint32 Size = *InBufferSizes++;
-		BufferSizes.Add(Size);
-	}
-	check(*InBufferSizes == 0);
-	InitResource();
-}
-
-FD3D11DynamicBuffer::~FD3D11DynamicBuffer()
-{
-	ReleaseResource();
-}
-
-void FD3D11DynamicBuffer::InitRHI()
-{
-	D3D11_BUFFER_DESC Desc;
-	ZeroMemory( &Desc, sizeof( D3D11_BUFFER_DESC ) );
-	Desc.Usage = D3D11_USAGE_DYNAMIC;
-	Desc.BindFlags = BindFlags;
-	Desc.CPUAccessFlags = D3D11_CPU_ACCESS_WRITE;
-	Desc.MiscFlags = 0;
-
-	while (Buffers.Num() < BufferSizes.Num())
-	{
-		TRefCountPtr<ID3D11Buffer> Buffer;
-		Desc.ByteWidth = BufferSizes[Buffers.Num()];
-		VERIFYD3D11RESULT_EX(D3DRHI->GetDevice()->CreateBuffer(&Desc,NULL,Buffer.GetInitReference()), D3DRHI->GetDevice());
-		UpdateBufferStats(Buffer,true);
-		Buffers.Add(Buffer);
-	}
-}
-
-void FD3D11DynamicBuffer::ReleaseRHI()
-{
-	for (int32 i = 0; i < Buffers.Num(); ++i)
-	{
-		UpdateBufferStats(Buffers[i],false);
-	}
-	Buffers.Empty();
-}
-
-void* FD3D11DynamicBuffer::Lock(uint32 Size)
-{
-	check(LockedBufferIndex == -1 && Buffers.Num() > 0);
-	
-	int32 BufferIndex = 0;
-	int32 NumBuffers = Buffers.Num();
-	while (BufferIndex < NumBuffers && BufferSizes[BufferIndex] < Size)
-	{
-		BufferIndex++;
-	}
-	if (BufferIndex == NumBuffers)
-	{
-		BufferIndex--;
-
-		TRefCountPtr<ID3D11Buffer> Buffer;
-		D3D11_BUFFER_DESC Desc;
-		ZeroMemory( &Desc, sizeof( D3D11_BUFFER_DESC ) );
-		Desc.Usage = D3D11_USAGE_DYNAMIC;
-		Desc.BindFlags = BindFlags;
-		Desc.CPUAccessFlags = D3D11_CPU_ACCESS_WRITE;
-		Desc.MiscFlags = 0;
-		Desc.ByteWidth = Size;
-		VERIFYD3D11RESULT_EX(D3DRHI->GetDevice()->CreateBuffer(&Desc,NULL,Buffer.GetInitReference()), D3DRHI->GetDevice());
-		UpdateBufferStats(Buffers[BufferIndex],false);
-		UpdateBufferStats(Buffer,true);
-		Buffers[BufferIndex] = Buffer;
-		BufferSizes[BufferIndex] = Size;
-	}
-
-	LockedBufferIndex = BufferIndex;
-	D3D11_MAPPED_SUBRESOURCE MappedSubresource;
-	VERIFYD3D11RESULT_EX(D3DRHI->GetDeviceContext()->Map(Buffers[BufferIndex],0,D3D11_MAP_WRITE_DISCARD,0,&MappedSubresource), D3DRHI->GetDevice());
-	return MappedSubresource.pData;
-}
-
-ID3D11Buffer* FD3D11DynamicBuffer::Unlock()
-{
-	check(LockedBufferIndex != -1);
-	ID3D11Buffer* LockedBuffer = Buffers[LockedBufferIndex];
-	D3DRHI->GetDeviceContext()->Unmap(LockedBuffer,0);
-	LockedBufferIndex = -1;
-	return LockedBuffer;
-}
 
 //
 // Stat declarations.
diff --git a/Engine/Source/Runtime/Windows/D3D11RHI/Public/D3D11Util.h b/Engine/Source/Runtime/Windows/D3D11RHI/Public/D3D11Util.h
--- a/Engine/Source/Runtime/Windows/D3D11RHI/Public/D3D11Util.h
+++ b/Engine/Source/Runtime/Windows/D3D11RHI/Public/D3D11Util.h
@@ -244,42 +244,6 @@ private:
 	int32 NumActiveTargets;
 };
 
-/**
- * Class for managing dynamic buffers.
- */
-class FD3D11DynamicBuffer : public FRenderResource, public FRefCountedObject
-{
-public:
-	/** Initialization constructor. */
-	FD3D11DynamicBuffer(class FD3D11DynamicRHI* InD3DRHI, D3D11_BIND_FLAG InBindFlags, uint32* InBufferSizes);
-	/** Destructor. */
-	~FD3D11DynamicBuffer();
-
-	/** Locks the buffer returning at least Size bytes. */
-	void* Lock(uint32 Size);
-	/** Unlocks the buffer returning the underlying D3D11 buffer to use as a resource. */
-	ID3D11Buffer* Unlock();
-
-	//~ Begin FRenderResource Interface.
-	virtual void InitRHI() override;
-	virtual void ReleaseRHI() override;
-	// End FRenderResource interface.
-
-private:
-	/** The maximum number of sub-buffers supported. */
-	enum { MAX_BUFFER_SIZES = 4 };
-	/** The size of each sub-buffer. */
-	TArray<uint32,TFixedAllocator<MAX_BUFFER_SIZES> > BufferSizes;
-	/** The sub-buffers. */
-	TArray<TRefCountPtr<ID3D11Buffer>,TFixedAllocator<MAX_BUFFER_SIZES> > Buffers;
-	/** The D3D11 RHI to that owns this dynamic buffer. */
-	class FD3D11DynamicRHI* D3DRHI;
-	/** Bind flags to use when creating sub-buffers. */
-	D3D11_BIND_FLAG BindFlags;
-	/** The index of the currently locked sub-buffer. */
-	int32 LockedBufferIndex;
-};
-
 template <
 	typename JobType,
 	typename = TEnableIf<TOr<
