Merging fix from Dev-PlatformRender-Staging to UE-81433

diff --git a/Engine/Source/Runtime/Windows/D3D11RHI/Private/D3D11Shaders.cpp b/Engine/Source/Runtime/Windows/D3D11RHI/Private/D3D11Shaders.cpp
--- a/Engine/Source/Runtime/Windows/D3D11RHI/Private/D3D11Shaders.cpp
+++ b/Engine/Source/Runtime/Windows/D3D11RHI/Private/D3D11Shaders.cpp
@@ -6,7 +6,10 @@
 
 #include "D3D11RHIPrivate.h"
 #include "Serialization/MemoryReader.h"
+
+#if !PLATFORM_HOLOLENS
 #include "nvapi.h"
+#endif
 
 template <typename TShaderType>
 static inline void ReadShaderOptionalData(FShaderCodeReader& InShaderCode, TShaderType& OutShader)
@@ -43,6 +46,7 @@ static inline void ReadShaderOptionalData(FShaderCodeReader& InShaderCode, TShad
 
 static void ApplyVendorExtensions(ID3D11Device* Direct3DDevice, EShaderFrequency Frequency, const TArray<FShaderCodeVendorExtension>& VendorExtensions)
 {
+#if !PLATFORM_HOLOLENS
 	for (int32 ExtensionIndex = 0; ExtensionIndex < VendorExtensions.Num(); ++ExtensionIndex)
 	{
 		const FShaderCodeVendorExtension& Extension = VendorExtensions[ExtensionIndex];
@@ -63,10 +67,12 @@ static void ApplyVendorExtensions(ID3D11Device* Direct3DDevice, EShaderFrequency
 			// TODO: [REDACTED_PROFILE_URL]/intel-graphics-compiler/blob/master/inc/IntelExtensions.hlsl
 		}
 	}
+#endif
 }
 
 static void ResetVendorExtensions(ID3D11Device* Direct3DDevice, EShaderFrequency Frequency, const TArray<FShaderCodeVendorExtension>& VendorExtensions)
 {
+#if !PLATFORM_HOLOLENS
 	for (int32 ExtensionIndex = 0; ExtensionIndex < VendorExtensions.Num(); ++ExtensionIndex)
 	{
 		const FShaderCodeVendorExtension& Extension = VendorExtensions[ExtensionIndex];
@@ -84,6 +90,7 @@ static void ResetVendorExtensions(ID3D11Device* Direct3DDevice, EShaderFrequency
 		{
 		}
 	}
+#endif
 }
 
 FVertexShaderRHIRef FD3D11DynamicRHI::RHICreateVertexShader(const TArray<uint8>& Code)
