Merging using Dev-VR=>Main

Adding HDR Monitor Status function to HL2

diff --git a/Engine/Source/Runtime/Windows/D3D11RHI/Private/HoloLens/HoloLensD3D11Viewport.cpp b/Engine/Source/Runtime/Windows/D3D11RHI/Private/HoloLens/HoloLensD3D11Viewport.cpp
--- a/Engine/Source/Runtime/Windows/D3D11RHI/Private/HoloLens/HoloLensD3D11Viewport.cpp
+++ b/Engine/Source/Runtime/Windows/D3D11RHI/Private/HoloLens/HoloLensD3D11Viewport.cpp
@@ -193,6 +193,13 @@ FD3D11Viewport::FD3D11Viewport(FD3D11DynamicRHI* InD3DRHI,HWND InWindowHandle,ui
 	BeginInitResource(&FrameSyncEvent);
 }
 
+
+// When a window has moved or resized we need to check whether it is on a HDR monitor or not. Set the correct color space of the monitor
+void FD3D11Viewport::CheckHDRMonitorStatus()
+{
+	PixelColorSpace = EColorSpaceAndEOTF::ERec709_sRGB;
+}
+
 void FD3D11Viewport::ConditionalResetSwapChain(bool bIgnoreFocus)
 {
 	if(!bIsValid)
