Improve XAudio2Device logging, reduce log spam by adding log once like behavior

diff --git a/Engine/Source/Runtime/AudioMixerCore/Public/AudioMixer.h b/Engine/Source/Runtime/AudioMixerCore/Public/AudioMixer.h
--- a/Engine/Source/Runtime/AudioMixerCore/Public/AudioMixer.h
+++ b/Engine/Source/Runtime/AudioMixerCore/Public/AudioMixer.h
@@ -494,8 +494,19 @@ namespace Audio
 		/** Is called when an error is generated. */
 		inline void OnAudioMixerPlatformError(const FString& ErrorDetails, const FString& FileName, int32 LineNumber)
 		{
+#if !NO_LOGGING
+			// Log once on these errors to avoid Spam.
+			static FCriticalSection Cs;
+			static TSet<uint32> LogHistory;
+			FScopeLock Lock(&Cs);
 			LastError = FString::Printf(TEXT("Audio Platform Device Error: %s (File %s, Line %d)"), *ErrorDetails, *FileName, LineNumber);
-			UE_LOG(LogAudioMixer, Error, TEXT("%s"), *LastError);
+			uint32 Hash = GetTypeHash(LastError);
+			if (!LogHistory.Contains(Hash))
+			{
+				UE_LOG(LogAudioMixer, Error, TEXT("%s"), *LastError);
+				LogHistory.Add(Hash);
+			}
+#endif //!NO_LOGGING
 		}
 
 		/** Start generating audio from our mixer. */
diff --git a/Engine/Source/Runtime/Windows/AudioMixerXAudio2/Private/AudioMixerPlatformXAudio2.cpp b/Engine/Source/Runtime/Windows/AudioMixerXAudio2/Private/AudioMixerPlatformXAudio2.cpp
--- a/Engine/Source/Runtime/Windows/AudioMixerXAudio2/Private/AudioMixerPlatformXAudio2.cpp
+++ b/Engine/Source/Runtime/Windows/AudioMixerXAudio2/Private/AudioMixerPlatformXAudio2.cpp
@@ -33,21 +33,21 @@
 #include "Logging/LogMacros.h"
 
 // Macro to check result code for XAudio2 failure, get the string version, log, and goto a cleanup
-#define XAUDIO2_CLEANUP_ON_FAIL(Result)																\
-	if (FAILED(Result))																				\
-	{																								\
-		FString ErrorString = FString::Printf(TEXT("0x%X: %s"), Result, GetErrorString(Result));	\
-		AUDIO_PLATFORM_ERROR(*ErrorString);															\
-		goto Cleanup;																				\
+#define XAUDIO2_CLEANUP_ON_FAIL(Result)																					\
+	if (FAILED(Result))																									\
+	{																													\
+		FString ErrorString = FString::Printf(TEXT("%s -> 0x%X: %s"), TEXT( #Result ), Result, *GetErrorString(Result));\
+		AUDIO_PLATFORM_ERROR(*ErrorString);																				\
+		goto Cleanup;																									\
 	}
 
 // Macro to check result for XAudio2 failure, get string version, log, and return false
-#define XAUDIO2_RETURN_ON_FAIL(Result)																\
-	if (FAILED(Result))																				\
-	{																								\
-		FString ErrorString = FString::Printf(TEXT("0x%X: %s"), Result, GetErrorString(Result));	\
-		AUDIO_PLATFORM_ERROR(*ErrorString);															\
-		return false;																				\
+#define XAUDIO2_RETURN_ON_FAIL(Result)																					\
+	if (FAILED(Result))																									\
+	{																													\
+		FString ErrorString = FString::Printf(TEXT("%s -> 0x%X: %s"), TEXT( #Result ), Result, *GetErrorString(Result));\
+		AUDIO_PLATFORM_ERROR(*ErrorString);																				\
+		return false;																									\
 	}
 
 
@@ -104,8 +104,8 @@ namespace Audio
 	{
 	}
 
-	const TCHAR* FMixerPlatformXAudio2::GetErrorString(HRESULT Result)
-	{
+	FString FMixerPlatformXAudio2::GetErrorString(HRESULT Result)
+	{		
 		switch (Result)
 		{
 			case HRESULT(XAUDIO2_E_INVALID_CALL):			return TEXT("XAUDIO2_E_INVALID_CALL");
@@ -120,7 +120,18 @@ namespace Audio
 			case E_INVALIDARG:								return TEXT("E_INVALIDARG");
 			case E_OUTOFMEMORY:								return TEXT("E_OUTOFMEMORY");
 #endif
-			default:										return TEXT("UNKNOWN");
+			case 0xe000020b:								return TEXT("ERROR_NO_SUCH_DEVINST");
+			default:										
+			{
+				// We don't know this error, ask this system if it does.
+				TCHAR Buffer[1024] = { 0 }; 
+				FString Msg(FPlatformMisc::GetSystemErrorMessage(Buffer, UE_ARRAY_COUNT(Buffer), Result));
+				
+				// Anything to return? Otherwise "UNKNOWN"
+				return !Msg.IsEmpty() ? 
+					   Msg : 
+					   TEXT("UNKNOWN");
+			}
 		}
 	}
 
diff --git a/Engine/Source/Runtime/Windows/AudioMixerXAudio2/Private/AudioMixerPlatformXAudio2.h b/Engine/Source/Runtime/Windows/AudioMixerXAudio2/Private/AudioMixerPlatformXAudio2.h
--- a/Engine/Source/Runtime/Windows/AudioMixerXAudio2/Private/AudioMixerPlatformXAudio2.h
+++ b/Engine/Source/Runtime/Windows/AudioMixerXAudio2/Private/AudioMixerPlatformXAudio2.h
@@ -84,7 +84,7 @@ namespace Audio
 
 	private:
 
-		const TCHAR* GetErrorString(HRESULT Result);
+		FString GetErrorString(HRESULT Result);
 		bool AllowDeviceSwap();
 
 		// Used to teardown and reinitialize XAudio2.
