Reverting change in d3d11texture to avoid static code analysis issues

diff --git a/Engine/Source/Runtime/Windows/D3D11RHI/Private/D3D11Texture.cpp b/Engine/Source/Runtime/Windows/D3D11RHI/Private/D3D11Texture.cpp
--- a/Engine/Source/Runtime/Windows/D3D11RHI/Private/D3D11Texture.cpp
+++ b/Engine/Source/Runtime/Windows/D3D11RHI/Private/D3D11Texture.cpp
@@ -2158,12 +2158,12 @@ void FD3D11DynamicRHI::RHICopySubTextureRegion(FTexture2DRHIParamRef SourceTextu
 		SourceBox.Max.Y -= Delta;
 	}
 
-	uint32 DestinationOffsetX = 0;
-	uint32 DestinationOffsetY = 0;
-	uint32 SourceStartX = SourceBox.Min.X;
-	uint32 SourceEndX = SourceBox.Max.X;
-	uint32 SourceStartY = SourceBox.Min.Y;
-	uint32 SourceEndY = SourceBox.Max.Y;
+	int32 DestinationOffsetX = 0;
+	int32 DestinationOffsetY = 0;
+	int32 SourceStartX = SourceBox.Min.X;
+	int32 SourceEndX = SourceBox.Max.X;
+	int32 SourceStartY = SourceBox.Min.Y;
+	int32 SourceEndY = SourceBox.Max.Y;
 	//If the source box is not fitting on the left bottom side, offset the result so the destination pixel match the expectation
 	if (SourceStartX < 0)
 	{
