Fixed the Xbox audio errors due to the HoloLens merge

diff --git a/Engine/Source/Runtime/Windows/XAudio2/Private/XAudio2Device.cpp b/Engine/Source/Runtime/Windows/XAudio2/Private/XAudio2Device.cpp
--- a/Engine/Source/Runtime/Windows/XAudio2/Private/XAudio2Device.cpp
+++ b/Engine/Source/Runtime/Windows/XAudio2/Private/XAudio2Device.cpp
@@ -19,7 +19,7 @@
 #include "XAudio2Effects.h"
 #include "Interfaces/IAudioFormat.h"
 #include "HAL/PlatformAffinity.h"
-#if PLATFORM_WINDOWS
+#if PLATFORM_WINDOWS || PLATFORM_HOLOLENS || PLATFORM_XBOXONE
 #include "Windows/WindowsHWrapper.h"
 #include "Windows/AllowWindowsPlatformTypes.h"
 #include "Windows/AllowWindowsPlatformAtomics.h"
diff --git a/Engine/Source/Runtime/Windows/XAudio2/Private/XAudio2Support.h b/Engine/Source/Runtime/Windows/XAudio2/Private/XAudio2Support.h
--- a/Engine/Source/Runtime/Windows/XAudio2/Private/XAudio2Support.h
+++ b/Engine/Source/Runtime/Windows/XAudio2/Private/XAudio2Support.h
@@ -6,7 +6,6 @@
 
 #pragma once
 
-// @ATG_CHANGE :  BEGIN HoloLens support
 #ifndef XAUDIO_SUPPORTS_XMA2WAVEFORMATEX
 	#define XAUDIO_SUPPORTS_XMA2WAVEFORMATEX	0
 #endif	//XAUDIO_SUPPORTS_XMA2WAVEFORMATEX
@@ -27,7 +26,7 @@
 #include "XAudio2Device.h"
 #include "AudioDecompress.h"
 #include "AudioEffect.h"
-#if PLATFORM_WINDOWS || PLATFORM_HOLOLENS
+#if PLATFORM_WINDOWS || PLATFORM_HOLOLENS || PLATFORM_XBOXONE
 #include "Windows/WindowsHWrapper.h"
 #include "Windows/AllowWindowsPlatformTypes.h"
 #include "Windows/AllowWindowsPlatformAtomics.h"
