Collision particles failed in any HMD using the swapchain pattern (e.g. Oculus, and soon OpenXR and SteamVR, etc). Turns out, this was due to how the texture aliasing worked. Both the D3D11 and D3D12 RHIs implementation of the aliasing copied the BaseShaderResource member from the source texture, which is incorrect, since the state caching logic uses the BaseShaderResource pointer to track the currently bound resource. By aliasing this, the logic breaks.

In this specific bug, the broken logic was caused by the FD3D11DynamicRHI::RHISetRenderTargets function failing to unset the SRV for the depth buffer because the ConditionalClearShaderResource call would pass the pointer to the source aliased texture, but the bound pointer was to the aliased version. Both had the same SRV, so the unbind would skip, but then D3D would hit the OM resource hazard logic and force-unbind the SRV. Once we hit this the state cache was out of sync with D3Ds state, and later when we went to bind the depth as an SRV the state manager would assume it was already bound when it had been evicted.

By correctly aliasing by not copying the BaseShaderResource pointer, state tracking is assured, and we successfully unbind and avoid the OM resource hazard path and keep our state cache current.

diff --git a/Engine/Source/Runtime/D3D12RHI/Private/D3D12Texture.h b/Engine/Source/Runtime/D3D12RHI/Private/D3D12Texture.h
--- a/Engine/Source/Runtime/D3D12RHI/Private/D3D12Texture.h
+++ b/Engine/Source/Runtime/D3D12RHI/Private/D3D12Texture.h
@@ -158,7 +158,9 @@ public:
 		// Alias the location, will perform an addref underneath
 		FD3D12ResourceLocation::Alias(ResourceLocation, Texture->ResourceLocation);
 
-		BaseShaderResource = Texture->BaseShaderResource;
+		// Do not copy the BaseShaderResource from the source texture (this is initialized correctly here, and is used for
+		// state caching logic).
+
 		ShaderResourceView = Texture->ShaderResourceView;
 
 		for (uint32 Index = 0; Index < FExclusiveDepthStencil::MaxIndex; Index++)
diff --git a/Engine/Source/Runtime/HeadMountedDisplay/Private/XRSwapChain.cpp b/Engine/Source/Runtime/HeadMountedDisplay/Private/XRSwapChain.cpp
--- a/Engine/Source/Runtime/HeadMountedDisplay/Private/XRSwapChain.cpp
+++ b/Engine/Source/Runtime/HeadMountedDisplay/Private/XRSwapChain.cpp
@@ -65,6 +65,11 @@ FXRSwapChain::FXRSwapChain(FTextureRHIParamRef InRHITexture, TArray<FTextureRHIR
 	, RHITextureSwapChain(InRHITextureSwapChain)
 	, SwapChainIndex_RHIThread(0)
 {
+	RHITexture->SetName(TEXT("XRSwapChainAliasedTexture"));
+	for (int ChainElement = 0; ChainElement < RHITextureSwapChain.Num(); ++ChainElement)
+	{
+		RHITextureSwapChain[ChainElement]->SetName(FName(*FString::Printf(TEXT("XRSwapChainBackingTex%d"), ChainElement)));
+	}
 }
 
 
diff --git a/Engine/Source/Runtime/Windows/D3D11RHI/Public/D3D11Resources.h b/Engine/Source/Runtime/Windows/D3D11RHI/Public/D3D11Resources.h
--- a/Engine/Source/Runtime/Windows/D3D11RHI/Public/D3D11Resources.h
+++ b/Engine/Source/Runtime/Windows/D3D11RHI/Public/D3D11Resources.h
@@ -307,8 +307,9 @@ public:
 		check(RTVArraySize == Texture->RTVArraySize);
 		check(NumDepthStencilViews == Texture->NumDepthStencilViews);
 
+		// Do not copy the BaseShaderResource from the source texture (this is initialized correctly here, and is used for
+		// state caching logic).
 		Resource = Texture->Resource;
-		BaseShaderResource = Texture->BaseShaderResource;
 		ShaderResourceView = Texture->ShaderResourceView;
 		RenderTargetViews = Texture->RenderTargetViews;
 
