HoloLens x64 build fixes
-Found a couple things that didn't quite work for hololens x64 (ie for the emulator)

diff --git a/Engine/Source/Runtime/Windows/D3D11RHI/Private/D3D11Device.cpp b/Engine/Source/Runtime/Windows/D3D11RHI/Private/D3D11Device.cpp
--- a/Engine/Source/Runtime/Windows/D3D11RHI/Private/D3D11Device.cpp
+++ b/Engine/Source/Runtime/Windows/D3D11RHI/Private/D3D11Device.cpp
@@ -554,7 +554,7 @@ void FD3D11DynamicRHI::CleanupD3DDevice()
 		ReleaseCachedQueries();
 
 
-#if !PLATFORM_CPU_ARM_FAMILY
+#ifdef AMD_AGS_API
 		// Clean up the AMD extensions and shut down the AMD AGS utility library
 		if (AmdAgsContext != NULL)
 		{
@@ -573,7 +573,7 @@ void FD3D11DynamicRHI::CleanupD3DDevice()
 		}
 #endif // INTEL_METRICSDISCOVERY
 
-#endif
+#endif //AMD_AGS_API
 
 		// When running with D3D debug, clear state and flush the device to get rid of spurious live objects in D3D11's report.
 		if (D3D11RHI_ShouldCreateWithD3DDebug())
diff --git a/Engine/Source/Runtime/Windows/D3D11RHI/Private/D3D11Texture.cpp b/Engine/Source/Runtime/Windows/D3D11RHI/Private/D3D11Texture.cpp
--- a/Engine/Source/Runtime/Windows/D3D11RHI/Private/D3D11Texture.cpp
+++ b/Engine/Source/Runtime/Windows/D3D11RHI/Private/D3D11Texture.cpp
@@ -436,7 +436,6 @@ void ReturnPooledTexture2D(int32 MipCount, EPixelFormat PixelFormat, ID3D11Textu
 #endif // #if USE_TEXTURE_POOLING
 }
 
-#if WITH_D3DX_LIBS || PLATFORM_CPU_ARM_FAMILY
 DXGI_FORMAT FD3D11DynamicRHI::GetPlatformTextureResourceFormat(DXGI_FORMAT InFormat, uint32 InFlags)
 {
 	// DX 11 Shared textures must be B8G8R8A8_UNORM
@@ -446,7 +445,6 @@ DXGI_FORMAT FD3D11DynamicRHI::GetPlatformTextureResourceFormat(DXGI_FORMAT InFor
 	}
 	return InFormat;
 }
-#endif	
 
 /** If true, guard texture creates with SEH to log more information about a driver crash we are seeing during texture streaming. */
 #define GUARDED_TEXTURE_CREATES (PLATFORM_WINDOWS && !(UE_BUILD_SHIPPING || UE_BUILD_TEST))
