Improve swapchain texture aliasing pattern
Provide an API to create an aliased texture directly from a source texture, rather than create an empty texture and then discard it via the current aliasing API (which is how XR swapchains have been created thus far). This implementation is provided for D3D11, D3D12, Vulkan and OpenGL. The default method will assert on an RHI that doesn't support this method.

Also includes validation work (Add RHI validation to RHIAliasTextureResources)

diff --git a/Engine/Plugins/Runtime/Oculus/OculusVR/Source/OculusHMD/Private/OculusHMD_CustomPresent.cpp b/Engine/Plugins/Runtime/Oculus/OculusVR/Source/OculusHMD/Private/OculusHMD_CustomPresent.cpp
--- a/Engine/Plugins/Runtime/Oculus/OculusVR/Source/OculusHMD/Private/OculusHMD_CustomPresent.cpp
+++ b/Engine/Plugins/Runtime/Oculus/OculusVR/Source/OculusHMD/Private/OculusHMD_CustomPresent.cpp
@@ -280,12 +280,7 @@ FXRSwapChainPtr FCustomPresent::CreateSwapChain_RenderThread(uint32 InSizeX, uin
 {
 	CheckInRenderThread();
 
-	// @todo: Once things settle down, the chain target will be created below in the CreateXRSwapChain call, via an RHI "CreateAliasedTexture" call.
 	FTextureRHIRef RHITexture;
-	{
-		RHITexture = CreateTexture_RenderThread(InSizeX, InSizeY, InFormat, InBinding, InNumMips, InNumSamples, InNumSamplesTileMem, InResourceType, InTextures[0], InTexCreateFlags);
-	}
-
 	TArray<FTextureRHIRef> RHITextureSwapChain;
 	{
 		for (int32 TextureIndex = 0; TextureIndex < InTextures.Num(); ++TextureIndex)
@@ -294,6 +289,8 @@ FXRSwapChainPtr FCustomPresent::CreateSwapChain_RenderThread(uint32 InSizeX, uin
 		}
 	}
 
+	RHITexture = GDynamicRHI->RHICreateAliasedTexture(RHITextureSwapChain[0]);
+
 	return CreateXRSwapChain(MoveTemp(RHITextureSwapChain), RHITexture);
 }
 
diff --git a/Engine/Source/Runtime/D3D12RHI/Private/D3D12RHIPrivate.h b/Engine/Source/Runtime/D3D12RHI/Private/D3D12RHIPrivate.h
--- a/Engine/Source/Runtime/D3D12RHI/Private/D3D12RHIPrivate.h
+++ b/Engine/Source/Runtime/D3D12RHI/Private/D3D12RHIPrivate.h
@@ -347,6 +347,7 @@ public:
 	virtual uint32 RHIGetViewportNextPresentGPUIndex(FRHIViewport* Viewport) final override;
 	virtual FTexture2DRHIRef RHIGetViewportBackBuffer(FRHIViewport* Viewport) final override;
 	virtual void RHIAliasTextureResources(FRHITexture* DestTexture, FRHITexture* SrcTexture) final override;
+	virtual FTextureRHIRef RHICreateAliasedTexture(FRHITexture* SourceTexture) final override;
 	virtual void RHIAdvanceFrameForGetViewportBackBuffer(FRHIViewport* Viewport) final override;
 	virtual void RHIAcquireThreadOwnership() final override;
 	virtual void RHIReleaseThreadOwnership() final override;
@@ -938,6 +939,9 @@ protected:
 
 	FD3D12Texture3D* CreateD3D12Texture3D(class FRHICommandListImmediate* RHICmdList, uint32 SizeX, uint32 SizeY, uint32 SizeZ, EPixelFormat Format, uint32 NumMips, uint32 Flags, FRHIResourceCreateInfo& CreateInfo);
 
+	template<typename BaseResourceType>
+	TD3D12Texture2D<BaseResourceType>* CreateAliasedD3D12Texture2D(TD3D12Texture2D<BaseResourceType>* SourceTexture);
+
 	/**
 	 * Gets the best supported MSAA settings from the provided MSAA count to check against.
 	 *
diff --git a/Engine/Source/Runtime/D3D12RHI/Private/D3D12Texture.cpp b/Engine/Source/Runtime/D3D12RHI/Private/D3D12Texture.cpp
--- a/Engine/Source/Runtime/D3D12RHI/Private/D3D12Texture.cpp
+++ b/Engine/Source/Runtime/D3D12RHI/Private/D3D12Texture.cpp
@@ -2851,6 +2851,99 @@ void FD3D12DynamicRHI::RHIAliasTextureResources(FRHITexture* DestTextureRHI, FRH
 	}
 }
 
+template<typename BaseResourceType>
+TD3D12Texture2D<BaseResourceType>* FD3D12DynamicRHI::CreateAliasedD3D12Texture2D(TD3D12Texture2D<BaseResourceType>* SourceTexture)
+{
+	FD3D12Adapter* Adapter = &GetAdapter();
+
+	D3D12_RESOURCE_DESC TextureDesc = SourceTexture->GetResource()->GetDesc();
+	TextureDesc.Alignment = 0;
+
+	uint32 SizeX = TextureDesc.Width;
+	uint32 SizeY = TextureDesc.Height;
+	uint32 SizeZ = 1;
+	uint32 ArraySize = TextureDesc.DepthOrArraySize;
+	uint32 NumMips = TextureDesc.MipLevels;
+	uint32 NumSamples = TextureDesc.SampleDesc.Count;
+
+	check(TextureDesc.Dimension == D3D12_RESOURCE_DIMENSION_TEXTURE2D);
+	check(ArraySize == 1);
+
+	//TODO: Somehow Oculus is creating a Render Target with 4k alignment with ovr_GetTextureSwapChainBufferDX
+	//      This is invalid and causes our size calculation to fail. Oculus SDK bug?
+	if (TextureDesc.Flags & D3D12_RESOURCE_FLAG_ALLOW_RENDER_TARGET)
+	{
+		TextureDesc.Alignment = D3D12_DEFAULT_RESOURCE_PLACEMENT_ALIGNMENT;
+	}
+
+	SCOPE_CYCLE_COUNTER(STAT_D3D12CreateTextureTime);
+
+	FD3D12Device* Device = Adapter->GetDevice(0);
+
+	const bool bSRGB = (SourceTexture->GetFlags() & TexCreate_SRGB) != 0;
+
+	const DXGI_FORMAT PlatformResourceFormat = TextureDesc.Format;
+	const DXGI_FORMAT PlatformShaderResourceFormat = FindShaderResourceDXGIFormat(PlatformResourceFormat, bSRGB);
+	const DXGI_FORMAT PlatformRenderTargetFormat = FindShaderResourceDXGIFormat(PlatformResourceFormat, bSRGB);
+
+	TD3D12Texture2D<BaseResourceType>* Texture2D = new TD3D12Texture2D<BaseResourceType>(Device, SizeX, SizeY, SizeZ, NumMips, NumSamples, SourceTexture->GetFormat(), false, SourceTexture->GetFlags(), SourceTexture->GetClearBinding());
+
+	// Set up the texture bind flags.
+	bool bCreateRTV = (TextureDesc.Flags & D3D12_RESOURCE_FLAG_ALLOW_RENDER_TARGET) != 0;
+	bool bCreateDSV = (TextureDesc.Flags & D3D12_RESOURCE_FLAG_ALLOW_DEPTH_STENCIL) != 0;
+
+	const D3D12_RESOURCE_STATES State = D3D12_RESOURCE_STATE_COMMON;
+
+	if (bCreateRTV)
+	{
+		Texture2D->SetCreatedRTVsPerSlice(false, NumMips);
+		Texture2D->SetNumRenderTargetViews(NumMips);
+
+		// Create a render target view for each array index and mip index.
+		// These are null because we'll be aliasing them shortly.
+		for (uint32 ArrayIndex = 0; ArrayIndex < ArraySize; ArrayIndex++)
+		{
+			for (uint32 MipIndex = 0; MipIndex < NumMips; MipIndex++)
+			{
+				Texture2D->SetRenderTargetViewIndex(nullptr, ArrayIndex * NumMips + MipIndex);
+			}
+		}
+	}
+
+	if (bCreateDSV)
+	{
+		// Create a depth-stencil-view for the texture.
+		for (uint32 AccessType = 0; AccessType < FExclusiveDepthStencil::MaxIndex; ++AccessType)
+		{
+			Texture2D->SetDepthStencilView(nullptr, AccessType);
+		}
+	}
+
+	RHIAliasTextureResources(Texture2D, SourceTexture);
+
+	return Texture2D;
+}
+
+FTextureRHIRef FD3D12DynamicRHI::RHICreateAliasedTexture(FRHITexture* SourceTextureRHI)
+{
+	FD3D12TextureBase* SourceTexture = GetD3D12TextureFromRHITexture(SourceTextureRHI);
+	if (SourceTextureRHI->GetTexture2D() != nullptr)
+	{
+		return CreateAliasedD3D12Texture2D<FD3D12BaseTexture2D>(static_cast<FD3D12Texture2D*>(SourceTextureRHI->GetTexture2D()));
+	}
+	else if (SourceTextureRHI->GetTexture2DArray() != nullptr)
+	{
+		return CreateAliasedD3D12Texture2D<FD3D12BaseTexture2DArray>(static_cast<FD3D12Texture2DArray*>(SourceTextureRHI->GetTexture2DArray()));
+	}
+	else if (SourceTextureRHI->GetTextureCube() != nullptr)
+	{
+		return CreateAliasedD3D12Texture2D<FD3D12BaseTextureCube>(static_cast<FD3D12TextureCube*>(SourceTextureRHI->GetTextureCube()));
+	}
+
+	UE_LOG(LogD3D12RHI, Error, TEXT("Currently FD3D12DynamicRHI::RHICreateAliasedTexture only supports 2D, 2D Array and Cube textures."));
+	return nullptr;
+}
+
 void FD3D12DynamicRHI::RHICopySubTextureRegion(FRHITexture2D* SourceTextureRHI, FRHITexture2D* DestTextureRHI, FBox2D SourceBox, FBox2D DestinationBox)
 {
 	FD3D12TextureBase* SourceTexture = GetD3D12TextureFromRHITexture(SourceTextureRHI);
diff --git a/Engine/Source/Runtime/OpenGLDrv/Private/OpenGLTexture.cpp b/Engine/Source/Runtime/OpenGLDrv/Private/OpenGLTexture.cpp
--- a/Engine/Source/Runtime/OpenGLDrv/Private/OpenGLTexture.cpp
+++ b/Engine/Source/Runtime/OpenGLDrv/Private/OpenGLTexture.cpp
@@ -2726,6 +2726,75 @@ void FOpenGLDynamicRHI::RHIAliasTextureResources(FRHITexture* DestRHITexture, FR
 	}
 }
 
+FRHITexture* FOpenGLDynamicRHI::CreateTexture2DAliased(FOpenGLTexture2D* SourceTexture)
+{
+	SCOPE_CYCLE_COUNTER(STAT_OpenGLCreateTextureTime);
+
+	FRHITexture* Result = new FOpenGLTexture2D(this, 0, SourceTexture->Target, -1, SourceTexture->GetSizeX(), SourceTexture->GetSizeY(), 0, 
+		SourceTexture->GetNumMips(), SourceTexture->GetNumSamples(), SourceTexture->GetNumSamplesTileMem(), 1, SourceTexture->GetFormat(), 
+		false, false, SourceTexture->GetFlags(), nullptr, SourceTexture->GetClearBinding());
+
+	RHIAliasTextureResources(Result, SourceTexture);
+
+	return Result;
+}
+
+FRHITexture* FOpenGLDynamicRHI::CreateTexture2DArrayAliased(FOpenGLTexture2DArray* SourceTexture)
+{
+	SCOPE_CYCLE_COUNTER(STAT_OpenGLCreateTextureTime);
+
+	FRHITexture* Result = new FOpenGLTexture2DArray(this, 0, SourceTexture->Target, -1, SourceTexture->GetSizeX(), SourceTexture->GetSizeY(), SourceTexture->GetSizeZ(),
+		SourceTexture->GetNumMips(), SourceTexture->GetNumSamples(), 1 /* aka check(InNumSamplesTileMem == 1) in OpenGLResource.h FOpenGLBaseTexture2DArray constructor */,
+		1, SourceTexture->GetFormat(), false, false, SourceTexture->GetFlags(), nullptr, SourceTexture->GetClearBinding());
+
+	RHIAliasTextureResources(Result, SourceTexture);
+
+	return Result;
+}
+
+FRHITexture* FOpenGLDynamicRHI::CreateTextureCubeAliased(FOpenGLTextureCube* SourceTexture)
+{
+	SCOPE_CYCLE_COUNTER(STAT_OpenGLCreateTextureTime);
+
+	FRHITexture* Result = new FOpenGLTextureCube(this, 0, SourceTexture->Target, -1, SourceTexture->GetSizeX(), SourceTexture->GetSizeY(), SourceTexture->GetSizeZ(),
+		SourceTexture->GetNumMips(), SourceTexture->GetNumSamples(), 1 /* OpenGL currently doesn't support multisample cube textures, per OpenGLResource.h FOpenGLBaseTextureCube */,
+		1, SourceTexture->GetFormat(), true, false, SourceTexture->GetFlags(), nullptr, SourceTexture->GetClearBinding());
+
+	RHIAliasTextureResources(Result, SourceTexture);
+
+	return Result;
+}
+
+FTextureRHIRef FOpenGLDynamicRHI::RHICreateAliasedTexture(FRHITexture* SourceTexture)
+{
+	FTextureRHIRef AliasedTexture;
+	if (SourceTexture->GetTexture2D() != nullptr)
+	{
+		AliasedTexture = CreateTexture2DAliased(static_cast<FOpenGLTexture2D*>(SourceTexture));
+	}
+	else if (SourceTexture->GetTexture2DArray() != nullptr)
+	{
+		AliasedTexture = CreateTexture2DArrayAliased(static_cast<FOpenGLTexture2DArray*>(SourceTexture));
+	}
+	else if (SourceTexture->GetTextureCube() != nullptr)
+	{
+		AliasedTexture = CreateTextureCubeAliased(static_cast<FOpenGLTextureCube*>(SourceTexture));
+	}
+	else
+	{
+		UE_LOG(LogRHI, Error, TEXT("Currently FOpenGLDynamicRHI::RHICreateAliasedTexture only supports 2D, 2D Array and Cube textures."));
+	}
+
+	if (AliasedTexture)
+	{
+		FOpenGLTextureBase* BaseTexture = GetOpenGLTextureFromRHITexture(AliasedTexture);
+		// Init memory size to zero, since we're aliased.
+		BaseTexture->SetMemorySize(0);
+	}
+
+	return AliasedTexture;
+}
+
 void* FOpenGLDynamicRHI::LockTexture2D_RenderThread(class FRHICommandListImmediate& RHICmdList, FRHITexture2D* Texture, uint32 MipIndex, EResourceLockMode LockMode, uint32& DestStride, bool bLockWithinMiptail, bool bNeedsDefaultRHIFlush)
 {
 	check(IsInRenderingThread());
diff --git a/Engine/Source/Runtime/OpenGLDrv/Public/OpenGLDrv.h b/Engine/Source/Runtime/OpenGLDrv/Public/OpenGLDrv.h
--- a/Engine/Source/Runtime/OpenGLDrv/Public/OpenGLDrv.h
+++ b/Engine/Source/Runtime/OpenGLDrv/Public/OpenGLDrv.h
@@ -422,6 +422,7 @@ public:
 	virtual bool RHIGetRenderQueryResult(FRHIRenderQuery* RenderQuery, uint64& OutResult, bool bWait) final override;
 	virtual FTexture2DRHIRef RHIGetViewportBackBuffer(FRHIViewport* Viewport) final override;
 	virtual void RHIAliasTextureResources(FRHITexture* DestTexture, FRHITexture* SrcTexture) final override;
+	virtual FTextureRHIRef RHICreateAliasedTexture(FRHITexture* SourceTexture) final override;
 	virtual void RHIAdvanceFrameForGetViewportBackBuffer(FRHIViewport* Viewport) final override;
 	virtual void RHIAcquireThreadOwnership() final override;
 	virtual void RHIReleaseThreadOwnership() final override;
@@ -639,6 +640,10 @@ public:
 	FRHITexture* CreateOpenGLTexture(uint32 SizeX, uint32 SizeY, bool CubeTexture, bool ArrayTexture, bool bIsExternal, uint8 Format, uint32 NumMips, uint32 NumSamples, uint32 ArraySize, uint32 Flags, const FClearValueBinding& InClearValue, FResourceBulkDataInterface* BulkData = NULL);
 	
 	FRHITexture* CreateOpenGLRHITextureOnly(const uint32 SizeX, const uint32 SizeY, const bool bCubeTexture, const bool bArrayTexture, const bool bIsExternal, uint8& Format, uint32& NumMips, uint32& NumSamples, const uint32 ArraySize, uint32& Flags, const FClearValueBinding& InClearValue, FResourceBulkDataInterface* BulkData = nullptr);
+	FRHITexture* CreateTexture2DAliased(FOpenGLTexture2D* SourceTexture);
+	FRHITexture* CreateTexture2DArrayAliased(FOpenGLTexture2DArray* SourceTexture);
+	FRHITexture* CreateTextureCubeAliased(FOpenGLTextureCube* SourceTexture);
+
 	void InitializeGLTexture(FRHITexture* Texture, uint32 SizeX, const uint32 SizeY, const bool bCubeTexture, const bool bArrayTexture, const bool bIsExternal, const uint8 Format, const uint32 NumMips, const uint32 NumSamples, const uint32 ArraySize, const uint32 Flags, const FClearValueBinding& InClearValue, FResourceBulkDataInterface* BulkData = nullptr);
 
 	void* GetOpenGLCurrentContextHandle();
diff --git a/Engine/Source/Runtime/RHI/Public/DynamicRHI.h b/Engine/Source/Runtime/RHI/Public/DynamicRHI.h
--- a/Engine/Source/Runtime/RHI/Public/DynamicRHI.h
+++ b/Engine/Source/Runtime/RHI/Public/DynamicRHI.h
@@ -954,6 +954,12 @@ public:
 		checkNoEntry();
 	}
 
+	virtual FTextureRHIRef RHICreateAliasedTexture(FRHITexture* SourceTexture)
+	{
+		checkNoEntry();
+		return nullptr;
+	}
+
 	// Only relevant with an RHI thread, this advances the backbuffer for the purpose of GetViewportBackBuffer
 	// FlushType: Thread safe
 	virtual void RHIAdvanceFrameForGetViewportBackBuffer(FRHIViewport* Viewport) = 0;
diff --git a/Engine/Source/Runtime/RHI/Public/RHIValidation.h b/Engine/Source/Runtime/RHI/Public/RHIValidation.h
--- a/Engine/Source/Runtime/RHI/Public/RHIValidation.h
+++ b/Engine/Source/Runtime/RHI/Public/RHIValidation.h
@@ -981,6 +981,21 @@ public:
 		return RHI->RHIGetFMaskTexture(SourceTextureRHI);
 	}
 
+	virtual void RHIAliasTextureResources(FRHITexture* DestTexture, FRHITexture* SourceTexture) override final
+	{
+		// Source and target need to be valid objects.
+		check(DestTexture && SourceTexture);
+		// Source texture must have been created (i.e. have a native resource backing).
+		check(SourceTexture->GetNativeResource() != nullptr);
+		RHI->RHIAliasTextureResources(DestTexture, SourceTexture);
+	}
+
+	virtual FTextureRHIRef RHICreateAliasedTexture(FRHITexture* SourceTexture) override final
+	{
+		check(SourceTexture);
+		return RHI->RHICreateAliasedTexture(SourceTexture);
+	}
+
 	// Only relevant with an RHI thread, this advances the backbuffer for the purpose of GetViewportBackBuffer
 	// FlushType: Thread safe
 	virtual void RHIAdvanceFrameForGetViewportBackBuffer(FRHIViewport* Viewport) override final
diff --git a/Engine/Source/Runtime/VulkanRHI/Private/VulkanRHI.cpp b/Engine/Source/Runtime/VulkanRHI/Private/VulkanRHI.cpp
--- a/Engine/Source/Runtime/VulkanRHI/Private/VulkanRHI.cpp
+++ b/Engine/Source/Runtime/VulkanRHI/Private/VulkanRHI.cpp
@@ -1038,6 +1038,29 @@ void FVulkanDynamicRHI::RHIAliasTextureResources(FRHITexture* DestTextureRHI, FR
 	}
 }
 
+FTextureRHIRef FVulkanDynamicRHI::RHICreateAliasedTexture(FRHITexture* SourceTexture)
+{
+	FTextureRHIRef AliasedTexture;
+	if (SourceTexture->GetTexture2D() != nullptr)
+	{
+		AliasedTexture = new FVulkanTexture2D(static_cast<FVulkanTexture2D*>(SourceTexture));
+	}
+	else if (SourceTexture->GetTexture2DArray() != nullptr)
+	{
+		AliasedTexture = new FVulkanTexture2DArray(static_cast<FVulkanTexture2DArray*>(SourceTexture));
+	}
+	else if (SourceTexture->GetTextureCube() != nullptr)
+	{
+		AliasedTexture = new FVulkanTextureCube(static_cast<FVulkanTextureCube*>(SourceTexture));
+	}
+	else
+	{
+		UE_LOG(LogRHI, Error, TEXT("Currently FOpenGLDynamicRHI::RHICreateAliasedTexture only supports 2D, 2D Array and Cube textures."));
+	}
+
+	return AliasedTexture;
+}
+
 void FVulkanDynamicRHI::RHICopySubTextureRegion(FRHITexture2D* SourceTexture, FRHITexture2D* DestinationTexture, FBox2D SourceBox, FBox2D DestinationBox)
 {
 	FRHICopyTextureInfo CopyInfo;
diff --git a/Engine/Source/Runtime/VulkanRHI/Private/VulkanResources.h b/Engine/Source/Runtime/VulkanRHI/Private/VulkanResources.h
--- a/Engine/Source/Runtime/VulkanRHI/Private/VulkanResources.h
+++ b/Engine/Source/Runtime/VulkanRHI/Private/VulkanResources.h
@@ -488,6 +488,9 @@ struct FVulkanTextureBase : public FVulkanBaseShaderResource
 	FVulkanTextureBase(FVulkanDevice& Device, VkImageViewType ResourceType, EPixelFormat Format, uint32 SizeX, uint32 SizeY, uint32 SizeZ, uint32 ArraySize, uint32 NumMips, uint32 NumSamples, VkImage InImage, VkDeviceMemory InMem, uint32 UEFlags, const FRHIResourceCreateInfo& CreateInfo = FRHIResourceCreateInfo());
 	FVulkanTextureBase(FVulkanDevice& Device, VkImageViewType ResourceType, EPixelFormat Format, uint32 SizeX, uint32 SizeY, uint32 SizeZ, uint32 ArraySize, uint32 NumMips, uint32 NumSamples, VkImage InImage, VkDeviceMemory InMem, FSamplerYcbcrConversionInitializer& ConversionInitializer, uint32 UEFlags, const FRHIResourceCreateInfo& CreateInfo = FRHIResourceCreateInfo());
 
+	// Aliasing constructor.
+	FVulkanTextureBase(const FVulkanTextureBase* SrcTexture, VkImageViewType ResourceType, uint32 SizeX, uint32 SizeY, uint32 sizeZ);
+
 	virtual ~FVulkanTextureBase();
 
 	void AliasTextureResources(const FVulkanTextureBase* SrcTexture);
@@ -514,6 +517,9 @@ public:
 	FVulkanTexture2D(FVulkanDevice& Device, EPixelFormat Format, uint32 SizeX, uint32 SizeY, uint32 NumMips, uint32 NumSamples, VkImage Image, uint32 UEFlags, const FRHIResourceCreateInfo& CreateInfo);
 	FVulkanTexture2D(FVulkanDevice& Device, EPixelFormat Format, uint32 SizeX, uint32 SizeY, uint32 NumMips, uint32 NumSamples, VkImage Image, struct FSamplerYcbcrConversionInitializer& ConversionInitializer, uint32 UEFlags, const FRHIResourceCreateInfo& CreateInfo);
 
+	// Aliasing constructor
+	FVulkanTexture2D(const FVulkanTexture2D* SrcTexture);
+
 	virtual ~FVulkanTexture2D();
 
 	// IRefCountedObject interface.
@@ -548,6 +554,10 @@ public:
 	// Constructor, just calls base and Surface constructor
 	FVulkanTexture2DArray(FVulkanDevice& Device, EPixelFormat Format, uint32 SizeX, uint32 SizeY, uint32 ArraySize, uint32 NumMips, uint32 NumSamples, uint32 Flags, FResourceBulkDataInterface* BulkData, const FClearValueBinding& InClearValue);
 	FVulkanTexture2DArray(FVulkanDevice& Device, EPixelFormat Format, uint32 SizeX, uint32 SizeY, uint32 ArraySize, uint32 NumMips, uint32 NumSamples, VkImage Image, uint32 Flags, FResourceBulkDataInterface* BulkData, const FClearValueBinding& InClearValue);
+
+	// Aliasing constructor
+	FVulkanTexture2DArray(const FVulkanTexture2DArray* SrcTexture);
+
 		
 	// IRefCountedObject interface.
 	virtual uint32 AddRef() const override final
@@ -612,6 +622,10 @@ class FVulkanTextureCube : public FRHITextureCube, public FVulkanTextureBase
 public:
 	FVulkanTextureCube(FVulkanDevice& Device, EPixelFormat Format, uint32 Size, bool bArray, uint32 ArraySize, uint32 NumMips, uint32 Flags, FResourceBulkDataInterface* BulkData, const FClearValueBinding& InClearValue);
 	FVulkanTextureCube(FVulkanDevice& Device, EPixelFormat Format, uint32 Size, bool bArray, uint32 ArraySize, uint32 NumMips, VkImage Image, uint32 Flags, FResourceBulkDataInterface* BulkData, const FClearValueBinding& InClearValue);
+
+	// Aliasing constructor
+	FVulkanTextureCube(const FVulkanTextureCube* SrcTexture);
+
 	virtual ~FVulkanTextureCube();
 
 	// IRefCountedObject interface.
diff --git a/Engine/Source/Runtime/VulkanRHI/Private/VulkanTexture.cpp b/Engine/Source/Runtime/VulkanRHI/Private/VulkanTexture.cpp
--- a/Engine/Source/Runtime/VulkanRHI/Private/VulkanTexture.cpp
+++ b/Engine/Source/Runtime/VulkanRHI/Private/VulkanTexture.cpp
@@ -1749,6 +1749,23 @@ FVulkanTextureBase::FVulkanTextureBase(FVulkanDevice& Device, VkImageViewType Re
 	InsertInitialImageLayout(Device, InImage, VK_IMAGE_LAYOUT_SHADER_READ_ONLY_OPTIMAL);
 }
 
+FVulkanTextureBase::FVulkanTextureBase(const FVulkanTextureBase* SrcTexture, VkImageViewType ResourceType, uint32 SizeX, uint32 SizeY, uint32 SizeZ)
+	: Surface(*SrcTexture->Surface.Device, ResourceType, SrcTexture->Surface.PixelFormat, SizeX, SizeY, SizeZ, SrcTexture->Surface.NumArrayLevels, SrcTexture->Surface.NumMips, SrcTexture->Surface.NumSamples, SrcTexture->Surface.Image, SrcTexture->Surface.UEFlags, FRHIResourceCreateInfo())
+	, PartialView(nullptr)
+	, bIsAliased(true)
+{
+	if (Surface.FullAspectMask == Surface.PartialAspectMask)
+	{
+		PartialView = &DefaultView;
+	}
+	else
+	{
+		PartialView = new FVulkanTextureView;
+		// Skip create, since we're aliasing.
+	}
+	AliasTextureResources(SrcTexture);
+}
+
 FVulkanTextureBase::~FVulkanTextureBase()
 {
 	DestroyViews();
@@ -1811,6 +1828,12 @@ FVulkanTexture2D::FVulkanTexture2D(FVulkanDevice& Device, EPixelFormat Format, u
 {
 }
 
+FVulkanTexture2D::FVulkanTexture2D(const FVulkanTexture2D* SrcTexture)
+	: FRHITexture2D(SrcTexture->GetSizeX(), SrcTexture->GetSizeY(), SrcTexture->GetNumMips(), SrcTexture->GetNumSamples(), SrcTexture->GetFormat(), SrcTexture->GetFlags(), SrcTexture->GetClearBinding())
+	, FVulkanTextureBase(static_cast<const FVulkanTextureBase*>(SrcTexture), VK_IMAGE_VIEW_TYPE_2D, SrcTexture->GetSizeX(), SrcTexture->GetSizeY(), 1)
+{
+}
+
 FVulkanTexture2D::~FVulkanTexture2D()
 {
 	if ((Surface.UEFlags & (TexCreate_DepthStencilTargetable | TexCreate_RenderTargetable)) != 0)
@@ -1831,6 +1854,11 @@ FVulkanTexture2DArray::FVulkanTexture2DArray(FVulkanDevice& Device, EPixelFormat
 {
 }
 
+FVulkanTexture2DArray::FVulkanTexture2DArray(const FVulkanTexture2DArray* SrcTexture)
+	: FRHITexture2DArray(SrcTexture->GetSizeX(), SrcTexture->GetSizeY(), SrcTexture->GetSizeZ(), SrcTexture->GetNumMips(), SrcTexture->GetNumSamples(), SrcTexture->GetFormat(), SrcTexture->GetFlags(), SrcTexture->GetClearBinding())
+	, FVulkanTextureBase(static_cast<const FVulkanTextureBase*>(SrcTexture), VK_IMAGE_VIEW_TYPE_2D_ARRAY, SrcTexture->GetSizeX(), SrcTexture->GetSizeY(), SrcTexture->GetSizeZ())
+{
+}
 
 void FVulkanTextureReference::SetReferencedTexture(FRHITexture* InTexture)
 {
@@ -1852,6 +1880,12 @@ FVulkanTextureCube::FVulkanTextureCube(FVulkanDevice& Device, EPixelFormat Forma
 {
 }
 
+FVulkanTextureCube::FVulkanTextureCube(const FVulkanTextureCube* SrcTexture)
+	: FRHITextureCube(SrcTexture->GetSize(), SrcTexture->GetNumMips(), SrcTexture->GetFormat(), SrcTexture->GetFlags(), SrcTexture->GetClearBinding())
+	, FVulkanTextureBase(static_cast<const FVulkanTextureBase*>(SrcTexture), VK_IMAGE_VIEW_TYPE_CUBE, SrcTexture->GetSize(), SrcTexture->GetSize(), 1)
+{
+}
+
 FVulkanTextureCube::~FVulkanTextureCube()
 {
 	if ((GetFlags() & (TexCreate_DepthStencilTargetable | TexCreate_RenderTargetable)) != 0)
diff --git a/Engine/Source/Runtime/VulkanRHI/Public/VulkanDynamicRHI.h b/Engine/Source/Runtime/VulkanRHI/Public/VulkanDynamicRHI.h
--- a/Engine/Source/Runtime/VulkanRHI/Public/VulkanDynamicRHI.h
+++ b/Engine/Source/Runtime/VulkanRHI/Public/VulkanDynamicRHI.h
@@ -137,6 +137,7 @@ public:
 	virtual bool RHIGetRenderQueryResult(FRHIRenderQuery* RenderQuery, uint64& OutResult, bool bWait) final override;
 	virtual FTexture2DRHIRef RHIGetViewportBackBuffer(FRHIViewport* Viewport) final override;
 	virtual void RHIAliasTextureResources(FRHITexture* DestTexture, FRHITexture* SrcTexture) final override;
+	virtual FTextureRHIRef RHICreateAliasedTexture(FRHITexture* SourceTexture) final override;
 	virtual void RHIAdvanceFrameForGetViewportBackBuffer(FRHIViewport* Viewport) final override;
 	virtual void RHIAcquireThreadOwnership() final override;
 	virtual void RHIReleaseThreadOwnership() final override;
diff --git a/Engine/Source/Runtime/Windows/D3D11RHI/Private/D3D11RHIPrivate.h b/Engine/Source/Runtime/Windows/D3D11RHI/Private/D3D11RHIPrivate.h
--- a/Engine/Source/Runtime/Windows/D3D11RHI/Private/D3D11RHIPrivate.h
+++ b/Engine/Source/Runtime/Windows/D3D11RHI/Private/D3D11RHIPrivate.h
@@ -478,6 +478,7 @@ public:
 	virtual bool RHIGetRenderQueryResult(FRHIRenderQuery* RenderQuery, uint64& OutResult, bool bWait) final override;
 	virtual FTexture2DRHIRef RHIGetViewportBackBuffer(FRHIViewport* Viewport) final override;
 	virtual void RHIAliasTextureResources(FRHITexture* DestTexture, FRHITexture* SrcTexture) final override;
+	virtual FTextureRHIRef RHICreateAliasedTexture(FRHITexture* SourceTexture) final override;
 	virtual void RHIAdvanceFrameForGetViewportBackBuffer(FRHIViewport* Viewport) final override;
 	virtual void RHIAcquireThreadOwnership() final override;
 	virtual void RHIReleaseThreadOwnership() final override;
@@ -951,6 +952,9 @@ protected:
 	template<typename BaseResourceType>
 	TD3D11Texture2D<BaseResourceType>* CreateTextureFromResource(bool bTextureArray, bool bCubeTexture, EPixelFormat Format, uint32 TexCreateFlags, const FClearValueBinding& ClearValueBinding, ID3D11Texture2D* TextureResource);
 
+	template<typename BaseResourceType>
+	TD3D11Texture2D<BaseResourceType>* CreateAliasedD3D11Texture2D(TD3D11Texture2D<BaseResourceType>* SourceTexture);
+
 	/** Initializes the constant buffers.  Called once at RHI initialization time. */
 	void InitConstantBuffers();
 
diff --git a/Engine/Source/Runtime/Windows/D3D11RHI/Private/D3D11Texture.cpp b/Engine/Source/Runtime/Windows/D3D11RHI/Private/D3D11Texture.cpp
--- a/Engine/Source/Runtime/Windows/D3D11RHI/Private/D3D11Texture.cpp
+++ b/Engine/Source/Runtime/Windows/D3D11RHI/Private/D3D11Texture.cpp
@@ -2352,6 +2352,93 @@ TD3D11Texture2D<BaseResourceType>* FD3D11DynamicRHI::CreateTextureFromResource(b
 	return Texture2D;
 }
 
+template <typename BaseResourceType>
+TD3D11Texture2D<BaseResourceType>* FD3D11DynamicRHI::CreateAliasedD3D11Texture2D(TD3D11Texture2D<BaseResourceType>* SourceTexture)
+{
+	D3D11_TEXTURE2D_DESC TextureDesc;
+	SourceTexture->GetResource()->GetDesc(&TextureDesc);
+
+	const bool bSRGB = (SourceTexture->Flags & TexCreate_SRGB) != 0;
+
+	const DXGI_FORMAT PlatformResourceFormat = TextureDesc.Format;
+	const DXGI_FORMAT PlatformShaderResourceFormat = FindShaderResourceDXGIFormat(PlatformResourceFormat, bSRGB);
+	const DXGI_FORMAT PlatformRenderTargetFormat = FindShaderResourceDXGIFormat(PlatformResourceFormat, bSRGB);
+
+	// Determine the MSAA settings to use for the texture.
+	D3D11_DSV_DIMENSION DepthStencilViewDimension = D3D11_DSV_DIMENSION_TEXTURE2D;
+	D3D11_RTV_DIMENSION RenderTargetViewDimension = D3D11_RTV_DIMENSION_TEXTURE2D;
+	D3D11_SRV_DIMENSION ShaderResourceViewDimension = D3D11_SRV_DIMENSION_TEXTURE2D;
+
+	if (TextureDesc.SampleDesc.Count > 1)
+	{
+		DepthStencilViewDimension = D3D11_DSV_DIMENSION_TEXTURE2DMS;
+		RenderTargetViewDimension = D3D11_RTV_DIMENSION_TEXTURE2DMS;
+		ShaderResourceViewDimension = D3D11_SRV_DIMENSION_TEXTURE2DMS;
+	}
+
+	TRefCountPtr<ID3D11ShaderResourceView> ShaderResourceView;
+	TArray<TRefCountPtr<ID3D11RenderTargetView> > RenderTargetViews;
+	TRefCountPtr<ID3D11DepthStencilView> DepthStencilViews[FExclusiveDepthStencil::MaxIndex];
+
+	bool bCreatedRTVPerSlice = false;
+	const bool bCubeTexture = SourceTexture->IsCubemap();
+	const bool bTextureArray = !bCubeTexture && TextureDesc.ArraySize > 1;
+
+	if (TextureDesc.BindFlags & D3D11_BIND_RENDER_TARGET)
+	{
+		// Create a render target view for each mip
+		for (uint32 MipIndex = 0; MipIndex < TextureDesc.MipLevels; MipIndex++)
+		{
+			// Just add null RTV entries (we'll be aliasing from source shortly).
+			if ((SourceTexture->Flags & TexCreate_TargetArraySlicesIndependently) && (bTextureArray || bCubeTexture))
+			{
+				bCreatedRTVPerSlice = true;
+
+				for (uint32 SliceIndex = 0; SliceIndex < TextureDesc.ArraySize; SliceIndex++)
+				{
+					RenderTargetViews.Add(nullptr);
+				}
+			}
+			else
+			{
+				RenderTargetViews.Add(nullptr);
+			}
+		}
+	}
+
+	TD3D11Texture2D<BaseResourceType>* Texture2D = new TD3D11Texture2D<BaseResourceType>(
+		this,
+		nullptr,
+		nullptr,
+		bCreatedRTVPerSlice,
+		TextureDesc.ArraySize,
+		RenderTargetViews,
+		nullptr,
+		TextureDesc.Width,
+		TextureDesc.Height,
+		0,
+		TextureDesc.MipLevels,
+		TextureDesc.SampleDesc.Count,
+		static_cast<BaseResourceType*>(SourceTexture)->GetFormat(),
+		bCubeTexture,
+		SourceTexture->Flags,
+		/*bPooledTexture=*/ false,
+		static_cast<BaseResourceType*>(SourceTexture)->GetClearBinding()
+		);
+
+	if (SourceTexture->Flags & TexCreate_RenderTargetable)
+	{
+		Texture2D->SetCurrentGPUAccess(EResourceTransitionAccess::EWritable);
+	}
+
+	// We'll be the same size, since we're the same thing. Avoid the check in D3D11Resources.h (AliasResources).
+	Texture2D->SetMemorySize(SourceTexture->GetMemorySize());
+	RHIAliasTextureResources(Texture2D, SourceTexture);
+
+	return Texture2D;
+}
+
+
 FTexture2DRHIRef FD3D11DynamicRHI::RHICreateTexture2DFromResource(EPixelFormat Format, uint32 TexCreateFlags, const FClearValueBinding& ClearValueBinding, ID3D11Texture2D* TextureResource)
 {
 	return CreateTextureFromResource<FD3D11BaseTexture2D>(false, false, Format, TexCreateFlags, ClearValueBinding, TextureResource);
@@ -2373,6 +2460,25 @@ void FD3D11DynamicRHI::RHIAliasTextureResources(FRHITexture* DestTextureRHI, FRH
 	}
 }
 
+FTextureRHIRef FD3D11DynamicRHI::RHICreateAliasedTexture(FRHITexture* SourceTexture)
+{
+	if (SourceTexture->GetTexture2D() != nullptr)
+	{
+		return CreateAliasedD3D11Texture2D<FD3D11BaseTexture2D>(static_cast<FD3D11Texture2D*>(SourceTexture->GetTexture2D()));
+	}
+	else if (SourceTexture->GetTexture2DArray() != nullptr)
+	{
+		return CreateAliasedD3D11Texture2D<FD3D11BaseTexture2DArray>(static_cast<FD3D11Texture2DArray*>(SourceTexture->GetTexture2DArray()));
+	}
+	else if (SourceTexture->GetTextureCube() != nullptr)
+	{
+		return CreateAliasedD3D11Texture2D<FD3D11BaseTextureCube>(static_cast<FD3D11TextureCube*>(SourceTexture->GetTextureCube()));
+	}
+
+	UE_LOG(LogD3D11RHI, Error, TEXT("Currently FD3D11DynamicRHI::RHICreateAliasedTexture only supports 2D, 2D Array and Cube textures."));
+	return nullptr;
+}
+
 void FD3D11DynamicRHI::RHICopyTexture(FRHITexture* SourceTextureRHI, FRHITexture* DestTextureRHI, const FRHICopyTextureInfo& CopyInfo)
 {
 	if (!SourceTextureRHI || !DestTextureRHI || SourceTextureRHI == DestTextureRHI)
diff --git a/Engine/Source/Runtime/Windows/D3D11RHI/Public/D3D11Resources.h b/Engine/Source/Runtime/Windows/D3D11RHI/Public/D3D11Resources.h
--- a/Engine/Source/Runtime/Windows/D3D11RHI/Public/D3D11Resources.h
+++ b/Engine/Source/Runtime/Windows/D3D11RHI/Public/D3D11Resources.h
@@ -305,6 +305,13 @@ public:
 		check(MemorySize == Texture->MemorySize);
 		check(bCreatedRTVsPerSlice == Texture->bCreatedRTVsPerSlice);
 		check(RTVArraySize == Texture->RTVArraySize);
+
+		// If we're creating an aliased texture, make sure we handle this case correctly.
+		if (Texture->NumDepthStencilViews && !NumDepthStencilViews)
+		{
+			NumDepthStencilViews = Texture->NumDepthStencilViews;
+		}
+
 		check(NumDepthStencilViews == Texture->NumDepthStencilViews);
 
 		// Do not copy the BaseShaderResource from the source texture (this is initialized correctly here, and is used for
