Simplify and unify render target pixel conversions

diff --git a/Engine/Source/Runtime/Apple/MetalRHI/Private/MetalRenderTarget.cpp b/Engine/Source/Runtime/Apple/MetalRHI/Private/MetalRenderTarget.cpp
--- a/Engine/Source/Runtime/Apple/MetalRHI/Private/MetalRenderTarget.cpp
+++ b/Engine/Source/Runtime/Apple/MetalRHI/Private/MetalRenderTarget.cpp
@@ -10,6 +10,7 @@
 #include "MetalCommandBuffer.h"
 #include "ResolveShader.h"
 #include "PipelineStateCache.h"
+#include "Math/PackedVector.h"
 
 static FResolveRect GetDefaultRect(const FResolveRect& Rect, uint32 DefaultWidth, uint32 DefaultHeight)
 {
@@ -217,56 +218,6 @@ void FMetalRHICommandContext::RHICopyToResolveTarget(FTextureRHIParamRef SourceT
 	}
 }
 
-/**
- * Helper for storing IEEE 32 bit float components
- */
-struct FFloatIEEE
-{
-	union
-	{
-		struct
-		{
-			uint32	Mantissa : 23, Exponent : 8, Sign : 1;
-		} Components;
-		
-		float	Float;
-	};
-};
-
-/**
- * Helper for storing DXGI_FORMAT_R11G11B10_FLOAT components
- */
-struct FMetalFloatR11G11B10
-{
-	// http://msdn.microsoft.com/En-US/library/bb173059(v=VS.85).aspx
-	uint32 R_Mantissa : 6;
-	uint32 R_Exponent : 5;
-	uint32 G_Mantissa : 6;
-	uint32 G_Exponent : 5;
-	uint32 B_Mantissa : 5;
-	uint32 B_Exponent : 5;
-	
-	/**
-	 * @return decompress into three 32 bit float
-	 */
-	operator FLinearColor()
-	{
-		FFloatIEEE	Result[3];
-		
-		Result[0].Components.Sign = 0;
-		Result[0].Components.Exponent = R_Exponent - 15 + 127;
-		Result[0].Components.Mantissa = FMath::Min<uint32>(FMath::FloorToInt((float)R_Mantissa / 32.0f * 8388608.0f),(1 << 23) - 1);
-		Result[1].Components.Sign = 0;
-		Result[1].Components.Exponent = G_Exponent - 15 + 127;
-		Result[1].Components.Mantissa = FMath::Min<uint32>(FMath::FloorToInt((float)G_Mantissa / 64.0f * 8388608.0f),(1 << 23) - 1);
-		Result[2].Components.Sign = 0;
-		Result[2].Components.Exponent = B_Exponent - 15 + 127;
-		Result[2].Components.Mantissa = FMath::Min<uint32>(FMath::FloorToInt((float)B_Mantissa / 64.0f * 8388608.0f),(1 << 23) - 1);
-		
-		return FLinearColor(Result[0].Float, Result[1].Float, Result[2].Float);
-	}
-};
-
 /** Helper for accessing R10G10B10A2 colors. */
 struct FMetalR10G10B10A2
 {
@@ -416,16 +367,16 @@ static void ConvertSurfaceDataToFColor(EPixelFormat Format, uint32 Width, uint32
 	}
 	else if (Format == PF_FloatR11G11B10)
 	{
-		check(sizeof(FMetalFloatR11G11B10) == sizeof(uint32));
+		check(sizeof(FFloat3Packed) == sizeof(uint32));
 		
 		for(uint32 Y = 0; Y < Height; Y++)
 		{
-			FMetalFloatR11G11B10* SrcPtr = (FMetalFloatR11G11B10*)(In + Y * SrcPitch);
+			FFloat3Packed* SrcPtr = (FFloat3Packed*)(In + Y * SrcPitch);
 			FColor* DestPtr = Out + Y * Width;
 			
 			for(uint32 X = 0; X < Width; X++)
 			{
-				FLinearColor Value = *SrcPtr;
+				FLinearColor Value = (*SrcPtr).ToLinearColor();
 				
 				FColor NormalizedColor = Value.ToFColor(bLinearToGamma);
 				FMemory::Memcpy(DestPtr++, &NormalizedColor, sizeof(FColor));
diff --git a/Engine/Source/Runtime/D3D12RHI/Private/D3D12RenderTarget.cpp b/Engine/Source/Runtime/D3D12RHI/Private/D3D12RenderTarget.cpp
--- a/Engine/Source/Runtime/D3D12RHI/Private/D3D12RenderTarget.cpp
+++ b/Engine/Source/Runtime/D3D12RHI/Private/D3D12RenderTarget.cpp
@@ -11,6 +11,7 @@
 #include "ResolveShader.h"
 #include "SceneUtils.h"
 #include "PipelineStateCache.h"
+#include "Math/PackedVector.h"
 
 static inline DXGI_FORMAT ConvertTypelessToUnorm(DXGI_FORMAT Format)
 {
@@ -570,87 +571,6 @@ void FD3D12DynamicRHI::RHITransferTexture(FTexture2DRHIParamRef TextureRHI, FInt
 	DEBUG_RHI_EXECUTE_COMMAND_LIST(this);
 }
 
-/**
-* Helper for storing IEEE 32 bit float components
-*/
-struct FFloatIEEE
-{
-	union
-	{
-		struct
-		{
-			uint32	Mantissa : 23, Exponent : 8, Sign : 1;
-		} Components;
-
-		float	Float;
-	};
-};
-
-/**
-* Helper for storing 16 bit float components
-*/
-struct FD3DFloat16
-{
-	union
-	{
-		struct
-		{
-			uint16	Mantissa : 10, Exponent : 5, Sign : 1;
-		} Components;
-
-		uint16	Encoded;
-	};
-
-	/**
-	* @return full 32 bit float from the 16 bit value
-	*/
-	operator float()
-	{
-		FFloatIEEE	Result;
-
-		Result.Components.Sign = Components.Sign;
-		Result.Components.Exponent = Components.Exponent - 15 + 127; // Stored exponents are biased by half their range.
-		Result.Components.Mantissa = FMath::Min<uint32>(FMath::FloorToInt((float)Components.Mantissa / 1024.0f * 8388608.0f), (1 << 23) - 1);
-
-		return Result.Float;
-	}
-};
-
-/**
-* Helper for storing DXGI_FORMAT_R11G11B10_FLOAT components
-*/
-struct FD3DFloatR11G11B10
-{
-	// http://msdn.microsoft.com/En-US/library/bb173059(v=VS.85).aspx
-	uint32 R_Mantissa : 6;
-	uint32 R_Exponent : 5;
-	uint32 G_Mantissa : 6;
-	uint32 G_Exponent : 5;
-	uint32 B_Mantissa : 5;
-	uint32 B_Exponent : 5;
-
-	/**
-	* @return decompress into three 32 bit float
-	*/
-	operator FLinearColor()
-	{
-		FFloatIEEE	Result[3];
-
-		Result[0].Components.Sign = 0;
-		Result[0].Components.Exponent = R_Exponent - 15 + 127;
-		Result[0].Components.Mantissa = FMath::Min<uint32>(FMath::FloorToInt((float)R_Mantissa / 32.0f * 8388608.0f), (1 << 23) - 1);
-		Result[1].Components.Sign = 0;
-		Result[1].Components.Exponent = G_Exponent - 15 + 127;
-		Result[1].Components.Mantissa = FMath::Min<uint32>(FMath::FloorToInt((float)G_Mantissa / 64.0f * 8388608.0f), (1 << 23) - 1);
-		Result[2].Components.Sign = 0;
-		Result[2].Components.Exponent = B_Exponent - 15 + 127;
-		Result[2].Components.Mantissa = FMath::Min<uint32>(FMath::FloorToInt((float)B_Mantissa / 64.0f * 8388608.0f), (1 << 23) - 1);
-
-		return FLinearColor(Result[0].Float, Result[1].Float, Result[2].Float);
-	}
-};
-
-
 static uint32 ComputeBytesPerPixel(DXGI_FORMAT Format)
 {
 	uint32 BytesPerPixel = 0;
@@ -985,11 +905,11 @@ static void ConvertRAWSurfaceDataToFColor(DXGI_FORMAT Format, uint32 Width, uint
 		FPlane	MinValue(0.0f, 0.0f, 0.0f, 0.0f),
 			MaxValue(1.0f, 1.0f, 1.0f, 1.0f);
 
-		check(sizeof(FD3DFloat16) == sizeof(uint16));
+		check(sizeof(FFloat16) == sizeof(uint16));
 
 		for (uint32 Y = 0; Y < Height; Y++)
 		{
-			FD3DFloat16* SrcPtr = (FD3DFloat16*)(In + Y * SrcPitch);
+			FFloat16* SrcPtr = (FFloat16*)(In + Y * SrcPitch);
 
 			for (uint32 X = 0; X < Width; X++)
 			{
@@ -1007,7 +927,7 @@ static void ConvertRAWSurfaceDataToFColor(DXGI_FORMAT Format, uint32 Width, uint
 
 		for (uint32 Y = 0; Y < Height; Y++)
 		{
-			FD3DFloat16* SrcPtr = (FD3DFloat16*)(In + Y * SrcPitch);
+			FFloat16* SrcPtr = (FFloat16*)(In + Y * SrcPitch);
 			FColor* DestPtr = Out + Y * Width;
 
 			for (uint32 X = 0; X < Width; X++)
@@ -1026,16 +946,16 @@ static void ConvertRAWSurfaceDataToFColor(DXGI_FORMAT Format, uint32 Width, uint
 	}
 	else if (Format == DXGI_FORMAT_R11G11B10_FLOAT)
 	{
-		check(sizeof(FD3DFloatR11G11B10) == sizeof(uint32));
+		check(sizeof(FFloat3Packed) == sizeof(uint32));
 
 		for (uint32 Y = 0; Y < Height; Y++)
 		{
-			FD3DFloatR11G11B10* SrcPtr = (FD3DFloatR11G11B10*)(In + Y * SrcPitch);
+			FFloat3Packed* SrcPtr = (FFloat3Packed*)(In + Y * SrcPitch);
 			FColor* DestPtr = Out + Y * Width;
 
 			for (uint32 X = 0; X < Width; X++)
 			{
-				FLinearColor Value = *SrcPtr;
+				FLinearColor Value = (*SrcPtr).ToLinearColor();
 
 				FColor NormalizedColor = Value.ToFColor(bLinearToGamma);
 				FMemory::Memcpy(DestPtr++, &NormalizedColor, sizeof(FColor));
diff --git a/Engine/Source/Runtime/Windows/D3D11RHI/Private/D3D11RenderTarget.cpp b/Engine/Source/Runtime/Windows/D3D11RHI/Private/D3D11RenderTarget.cpp
--- a/Engine/Source/Runtime/Windows/D3D11RHI/Private/D3D11RenderTarget.cpp
+++ b/Engine/Source/Runtime/Windows/D3D11RHI/Private/D3D11RenderTarget.cpp
@@ -10,6 +10,7 @@
 #include "RHIStaticStates.h"
 #include "ResolveShader.h"
 #include "PipelineStateCache.h"
+#include "Math/PackedVector.h"
 
 static inline DXGI_FORMAT ConvertTypelessToUnorm(DXGI_FORMAT Format)
 {
@@ -393,86 +394,6 @@ void FD3D11DynamicRHI::RHICopyToResolveTarget(FTextureRHIParamRef SourceTextureR
 	}
 }
 
-/**
-* Helper for storing IEEE 32 bit float components
-*/
-struct FFloatIEEE
-{
-	union
-	{
-		struct
-		{
-			uint32	Mantissa : 23, Exponent : 8, Sign : 1;
-		} Components;
-
-		float	Float;
-	};
-};
-
-/**
-* Helper for storing 16 bit float components
-*/
-struct FD3DFloat16
-{
-	union
-	{
-		struct
-		{
-			uint16	Mantissa : 10, Exponent : 5, Sign : 1;
-		} Components;
-
-		uint16	Encoded;
-	};
-
-	/**
-	* @return full 32 bit float from the 16 bit value
-	*/
-	operator float()
-	{
-		FFloatIEEE	Result;
-
-		Result.Components.Sign = Components.Sign;
-		Result.Components.Exponent = Components.Exponent - 15 + 127; // Stored exponents are biased by half their range.
-		Result.Components.Mantissa = FMath::Min<uint32>(FMath::FloorToInt((float)Components.Mantissa / 1024.0f * 8388608.0f),(1 << 23) - 1);
-
-		return Result.Float;
-	}
-};
-
-/**
-* Helper for storing DXGI_FORMAT_R11G11B10_FLOAT components
-*/
-struct FD3DFloatR11G11B10
-{
-	// http://msdn.microsoft.com/En-US/library/bb173059(v=VS.85).aspx
-	uint32 R_Mantissa : 6;
-	uint32 R_Exponent : 5;
-	uint32 G_Mantissa : 6;
-	uint32 G_Exponent : 5;
-	uint32 B_Mantissa : 5;
-	uint32 B_Exponent : 5;
-
-	/**
-	* @return decompress into three 32 bit float
-	*/
-	operator FLinearColor()
-	{
-		FFloatIEEE	Result[3];
-
-		Result[0].Components.Sign = 0;
-		Result[0].Components.Exponent = R_Exponent - 15 + 127;
-		Result[0].Components.Mantissa = FMath::Min<uint32>(FMath::FloorToInt((float)R_Mantissa / 64.0f * 8388608.0f),(1 << 23) - 1);
-		Result[1].Components.Sign = 0;
-		Result[1].Components.Exponent = G_Exponent - 15 + 127;
-		Result[1].Components.Mantissa = FMath::Min<uint32>(FMath::FloorToInt((float)G_Mantissa / 64.0f * 8388608.0f),(1 << 23) - 1);
-		Result[2].Components.Sign = 0;
-		Result[2].Components.Exponent = B_Exponent - 15 + 127;
-		Result[2].Components.Mantissa = FMath::Min<uint32>(FMath::FloorToInt((float)B_Mantissa / 32.0f * 8388608.0f),(1 << 23) - 1);
-
-		return FLinearColor(Result[0].Float, Result[1].Float, Result[2].Float);
-	}
-};
-
 // Only supports the formats that are supported by ConvertRAWSurfaceDataToFColor()
 static uint32 ComputeBytesPerPixel(DXGI_FORMAT Format)
 {
@@ -790,11 +711,11 @@ static void ConvertRAWSurfaceDataToFColor(DXGI_FORMAT Format, uint32 Width, uint
 		FPlane	MinValue(0.0f,0.0f,0.0f,0.0f),
 			MaxValue(1.0f,1.0f,1.0f,1.0f);
 
-		check(sizeof(FD3DFloat16)==sizeof(uint16));
+		check(sizeof(FFloat16)==sizeof(uint16));
 
 		for(uint32 Y = 0; Y < Height; Y++)
 		{
-			FD3DFloat16* SrcPtr = (FD3DFloat16*)(In + Y * SrcPitch);
+			FFloat16* SrcPtr = (FFloat16*)(In + Y * SrcPitch);
 
 			for(uint32 X = 0; X < Width; X++)
 			{
@@ -812,7 +733,7 @@ static void ConvertRAWSurfaceDataToFColor(DXGI_FORMAT Format, uint32 Width, uint
 
 		for(uint32 Y = 0; Y < Height; Y++)
 		{
-			FD3DFloat16* SrcPtr = (FD3DFloat16*)(In + Y * SrcPitch);
+			FFloat16* SrcPtr = (FFloat16*)(In + Y * SrcPitch);
 			FColor* DestPtr = Out + Y * Width;
 
 			for(uint32 X = 0; X < Width; X++)
@@ -831,16 +752,16 @@ static void ConvertRAWSurfaceDataToFColor(DXGI_FORMAT Format, uint32 Width, uint
 	}
 	else if (Format == DXGI_FORMAT_R11G11B10_FLOAT)
 	{
-		check(sizeof(FD3DFloatR11G11B10) == sizeof(uint32));
+		check(sizeof(FFloat3Packed) == sizeof(uint32));
 
 		for(uint32 Y = 0; Y < Height; Y++)
 		{
-			FD3DFloatR11G11B10* SrcPtr = (FD3DFloatR11G11B10*)(In + Y * SrcPitch);
+			FFloat3Packed* SrcPtr = (FFloat3Packed*)(In + Y * SrcPitch);
 			FColor* DestPtr = Out + Y * Width;
 
 			for(uint32 X = 0; X < Width; X++)
 			{
-				FLinearColor Value = *SrcPtr;
+				FLinearColor Value = (*SrcPtr).ToLinearColor();
 
 				FColor NormalizedColor = Value.ToFColor(bLinearToGamma);
 				FMemory::Memcpy(DestPtr++, &NormalizedColor, sizeof(FColor));
@@ -1347,7 +1268,7 @@ static void ConvertRAWSurfaceDataToFLinearColor(EPixelFormat Format, uint32 Widt
 		{
 			for (uint32 Y = 0; Y < Height; Y++)
 			{
-				FD3DFloat16* SrcPtr = (FD3DFloat16*)(In + Y * SrcPitch);
+				FFloat16* SrcPtr = (FFloat16*)(In + Y * SrcPitch);
 				FLinearColor* DestPtr = Out + Y * Width;
 
 				for (uint32 X = 0; X < Width; X++)
@@ -1363,11 +1284,11 @@ static void ConvertRAWSurfaceDataToFLinearColor(EPixelFormat Format, uint32 Widt
 			FPlane	MinValue(0.0f, 0.0f, 0.0f, 0.0f);
 			FPlane	MaxValue(1.0f, 1.0f, 1.0f, 1.0f);
 
-			check(sizeof(FD3DFloat16) == sizeof(uint16));
+			check(sizeof(FFloat16) == sizeof(uint16));
 
 			for (uint32 Y = 0; Y < Height; Y++)
 			{
-				FD3DFloat16* SrcPtr = (FD3DFloat16*)(In + Y * SrcPitch);
+				FFloat16* SrcPtr = (FFloat16*)(In + Y * SrcPitch);
 
 				for (uint32 X = 0; X < Width; X++)
 				{
@@ -1385,7 +1306,7 @@ static void ConvertRAWSurfaceDataToFLinearColor(EPixelFormat Format, uint32 Widt
 
 			for (uint32 Y = 0; Y < Height; Y++)
 			{
-				FD3DFloat16* SrcPtr = (FD3DFloat16*)(In + Y * SrcPitch);
+				FFloat16* SrcPtr = (FFloat16*)(In + Y * SrcPitch);
 				FLinearColor* DestPtr = Out + Y * Width;
 
 				for (uint32 X = 0; X < Width; X++)
@@ -1404,16 +1325,16 @@ static void ConvertRAWSurfaceDataToFLinearColor(EPixelFormat Format, uint32 Widt
 	}
 	else if (Format == PF_FloatRGB || Format == PF_FloatR11G11B10)
 	{
-		check(sizeof(FD3DFloatR11G11B10) == sizeof(uint32));
+		check(sizeof(FFloat3Packed) == sizeof(uint32));
 
 		for (uint32 Y = 0; Y < Height; Y++)
 		{
-			FD3DFloatR11G11B10* SrcPtr = (FD3DFloatR11G11B10*)(In + Y * SrcPitch);
+			FFloat3Packed* SrcPtr = (FFloat3Packed*)(In + Y * SrcPitch);
 			FLinearColor* DestPtr = Out + Y * Width;
 
 			for (uint32 X = 0; X < Width; X++)
 			{
-				*DestPtr = *SrcPtr;
+				*DestPtr = (*SrcPtr).ToLinearColor();
 				++DestPtr;
 				++SrcPtr;
 			}
