integrate change from Jeff.Rous\Intel: Enable capture mode automatically when Intel GPA is attached.

diff --git a/Engine/Source/Runtime/D3D12RHI/Private/D3D12Device.cpp b/Engine/Source/Runtime/D3D12RHI/Private/D3D12Device.cpp
--- a/Engine/Source/Runtime/D3D12RHI/Private/D3D12Device.cpp
+++ b/Engine/Source/Runtime/D3D12RHI/Private/D3D12Device.cpp
@@ -182,6 +182,18 @@ void FD3D12Device::SetupAfterDeviceCreation()
 		}
 	}
 
+	// Intel GPA
+	{
+		TRefCountPtr<IUnknown> IntelGPA;
+		static const IID IntelGPAID = { 0xCCFFEF16, 0x7B69, 0x468F, {0xBC, 0xE3, 0xCD, 0x95, 0x33, 0x69, 0xA3, 0x9A} };
+
+		if (SUCCEEDED(Direct3DDevice->QueryInterface(IntelGPAID, (void**)(IntelGPA.GetInitReference()))))
+		{
+			// Running under Intel GPA, so enable capturing mode
+			bUnderGPUCapture = true;
+		}
+	}
+
 	// AMD RGP profiler
 	if (GEmitRgpFrameMarkers && GetOwningRHI()->GetAmdAgsContext())
 	{
diff --git a/Engine/Source/Runtime/Windows/D3D11RHI/Private/D3D11Device.cpp b/Engine/Source/Runtime/Windows/D3D11RHI/Private/D3D11Device.cpp
--- a/Engine/Source/Runtime/Windows/D3D11RHI/Private/D3D11Device.cpp
+++ b/Engine/Source/Runtime/Windows/D3D11RHI/Private/D3D11Device.cpp
@@ -428,6 +428,15 @@ void FD3D11DynamicRHI::SetupAfterDeviceCreation()
 			GDynamicRHI->EnableIdealGPUCaptureOptions(true);
 		}
 	}
+
+	IUnknown* IntelGPA;
+	static const IID IntelGPAID = { 0xCCFFEF16, 0x7B69, 0x468F, {0xBC, 0xE3, 0xCD, 0x95, 0x33, 0x69, 0xA3, 0x9A} };
+
+	if (SUCCEEDED(Direct3DDevice->QueryInterface(IntelGPAID, (void**)(&IntelGPA))))
+	{
+		// Running under Intel GPA, so enable capturing mode
+		GDynamicRHI->EnableIdealGPUCaptureOptions(true);
+	}
 #endif
 
 	// Check for typed UAV load support
