DR - Fix a case where D3d11 renderqueries could be forgotten, which eventually disabled stat GPU This was potentially the SlateUI timestamp query crossing some frame boundary.

diff --git a/Engine/Source/Runtime/Windows/D3D11RHI/Private/D3D11Query.cpp b/Engine/Source/Runtime/Windows/D3D11RHI/Private/D3D11Query.cpp
--- a/Engine/Source/Runtime/Windows/D3D11RHI/Private/D3D11Query.cpp
+++ b/Engine/Source/Runtime/Windows/D3D11RHI/Private/D3D11Query.cpp
@@ -30,6 +30,22 @@ public:
 			// This is correct even if CurrentId overflowed
 			if (CurrentId - Batch.Id >= NumBufferedFrames)
 			{
+				TArray<FRenderQueryRHIRef>& Queries = ActiveBatches[BatchIdx].Queries;
+				for (int32 QueryIdx = 0; QueryIdx < Queries.Num(); ++QueryIdx)
+				{
+					FD3D11RenderQuery* Query = FD3D11DynamicRHI::ResourceCast(Queries[QueryIdx].GetReference());
+
+					if (Query->bResultIsCached || Query->GetRefCount() == 1)
+					{
+						continue;
+					}
+
+					if (Query->QueryType == ERenderQueryType::RQT_AbsoluteTime)
+					{
+						check(GD3D11RHI->GetQueryData(Query->Resource, &Query->Result, sizeof(Query->Result), Query->QueryType, true, false));
+						Query->bResultIsCached = true;
+					}
+				}
 				ActiveBatches.RemoveAtSwap(BatchIdx--);
 			}
 		}
