DR - Remove HDR from D3D11 & 12 temporarily to fix AMD crashes
#rnx

diff --git a/Engine/Source/Runtime/D3D12RHI/Private/D3D12Commands.cpp b/Engine/Source/Runtime/D3D12RHI/Private/D3D12Commands.cpp
--- a/Engine/Source/Runtime/D3D12RHI/Private/D3D12Commands.cpp
+++ b/Engine/Source/Runtime/D3D12RHI/Private/D3D12Commands.cpp
@@ -1811,11 +1811,6 @@ uint32 FD3D12DynamicRHI::RHIGetGPUFrameCycles()
 	return GGPUFrameTime;
 }
 
-EColorSpace FD3D12DynamicRHI::RHIGetColorSpace(FViewportRHIParamRef Viewport)
-{
-	return EColorSpace::ERec709sRGB;
-}
-
 void FD3D12DynamicRHI::RHIExecuteCommandList(FRHICommandList* CmdList)
 {
 	check(0); // this path has gone stale and needs updated methods, starting at ERCT_SetScissorRect
diff --git a/Engine/Source/Runtime/D3D12RHI/Private/D3D12RHIPrivate.h b/Engine/Source/Runtime/D3D12RHI/Private/D3D12RHIPrivate.h
--- a/Engine/Source/Runtime/D3D12RHI/Private/D3D12RHIPrivate.h
+++ b/Engine/Source/Runtime/D3D12RHI/Private/D3D12RHIPrivate.h
@@ -352,7 +352,6 @@ public:
 	virtual FViewportRHIRef RHICreateViewport(void* WindowHandle, uint32 SizeX, uint32 SizeY, bool bIsFullscreen, EPixelFormat PreferredPixelFormat) final override;
 	virtual void RHIResizeViewport(FViewportRHIParamRef Viewport, uint32 SizeX, uint32 SizeY, bool bIsFullscreen) final override;
 	virtual void RHIResizeViewport(FViewportRHIParamRef ViewportRHI, uint32 SizeX, uint32 SizeY, bool bIsFullscreen, EPixelFormat PreferredPixelFormat) final override;
-	virtual EColorSpace RHIGetColorSpace(FViewportRHIParamRef Viewport) final override;
 	virtual void RHITick(float DeltaTime) final override;
 	virtual void RHISetStreamOutTargets(uint32 NumTargets, const FVertexBufferRHIParamRef* VertexBuffers, const uint32* Offsets) final override;
 	virtual void RHIBlockUntilGPUIdle() final override;
diff --git a/Engine/Source/Runtime/D3D12RHI/Private/D3D12Texture.cpp b/Engine/Source/Runtime/D3D12RHI/Private/D3D12Texture.cpp
--- a/Engine/Source/Runtime/D3D12RHI/Private/D3D12Texture.cpp
+++ b/Engine/Source/Runtime/D3D12RHI/Private/D3D12Texture.cpp
@@ -329,13 +329,7 @@ void FD3D12TextureStats::D3D12TextureAllocated(TD3D12Texture2D<BaseResourceType>
 		}
 		else
 		{
-#if PLATFORM_XBOXONE
-			D3D12_RESOURCE_DESC Desc = D3D12Texture2D->GetDesc();
-			Desc.Layout = D3D12XBOX_BANK_ROTATED_TILE_MODE(Desc.Layout, 0);	// Remove bank rotation as this can cause the GetResourceAllocationInfo() to return 0
-#else
 			const D3D12_RESOURCE_DESC& Desc = D3D12Texture2D->GetDesc();
-#endif
-
 			const D3D12_RESOURCE_ALLOCATION_INFO AllocationInfo = Texture.GetParentDevice()->GetDevice()->GetResourceAllocationInfo(0, 1, &Desc);
 			const int64 TextureSize = AllocationInfo.SizeInBytes;
 
diff --git a/Engine/Source/Runtime/D3D12RHI/Public/D3D12Util.h b/Engine/Source/Runtime/D3D12RHI/Public/D3D12Util.h
--- a/Engine/Source/Runtime/D3D12RHI/Public/D3D12Util.h
+++ b/Engine/Source/Runtime/D3D12RHI/Public/D3D12Util.h
@@ -1143,38 +1143,4 @@ private:
 	TArray<Type*> AllObjects;
 	FCriticalSection CS;
 	static __declspec(thread) void* ThisThreadObject;
-};
-
-#if PLATFORM_XBOXONE
-#include "XboxOneAllowPlatformTypes.h"
-#include <xg.h>
-#include "XboxOneHidePlatformTypes.h"
-
-FORCEINLINE XG_TILE_MODE GetXgOptimalTileMode(const D3D12_RESOURCE_DESC& TextureDesc, XG_BIND_FLAG BindFlags, XG_RESOURCE_MISC_FLAG MiscFlags)
-{
-	XG_RESOURCE_DIMENSION xgDimension = (XG_RESOURCE_DIMENSION)TextureDesc.Dimension;
-	check(xgDimension == XG_RESOURCE_DIMENSION_TEXTURE2D || xgDimension == XG_RESOURCE_DIMENSION_TEXTURE3D);  
-	XG_TILE_MODE iTileMode = XG_TILE_MODE_INVALID;
-	if ((BindFlags & XG_BIND_DEPTH_STENCIL) == XG_BIND_DEPTH_STENCIL)
-	{
-		XG_TILE_MODE iTileModeStencil = XG_TILE_MODE_INVALID;
-		const bool bCompressed = (MiscFlags & XG_RESOURCE_MISC_NO_DEPTH_COMPRESSION) != 0;
-
-#if _XDK_EDITION >= 170600
-		const bool TileResource = false;
-		const bool TextureCompatible = false;
-		XGComputeOptimalDepthStencilTileModes((XG_FORMAT)TextureDesc.Format, TextureDesc.Width, TextureDesc.Height, TextureDesc.DepthOrArraySize, TextureDesc.SampleDesc.Count, bCompressed, TileResource, TextureCompatible, &iTileMode, &iTileModeStencil);
-#else
-		XGComputeOptimalDepthStencilTileModes((XG_FORMAT)TextureDesc.Format, TextureDesc.Width, TextureDesc.Height, TextureDesc.DepthOrArraySize, TextureDesc.SampleDesc.Count, bCompressed, &iTileMode, &iTileModeStencil);
-#endif
-	}
-	else
-	{
-		iTileMode = XGComputeOptimalTileMode(xgDimension, (XG_FORMAT)TextureDesc.Format, TextureDesc.Width, TextureDesc.Height, TextureDesc.Dimension, TextureDesc.SampleDesc.Count, BindFlags, MiscFlags);
-
-	}
-	check(iTileMode != XG_TILE_MODE_INVALID);
-	return iTileMode;
-}
-
-#endif
\ No newline at end of file
+};
\ No newline at end of file
diff --git a/Engine/Source/Runtime/Windows/D3D11RHI/Private/D3D11Device.cpp b/Engine/Source/Runtime/Windows/D3D11RHI/Private/D3D11Device.cpp
--- a/Engine/Source/Runtime/Windows/D3D11RHI/Private/D3D11Device.cpp
+++ b/Engine/Source/Runtime/Windows/D3D11RHI/Private/D3D11Device.cpp
@@ -94,33 +94,6 @@ FD3D11DynamicRHI::FD3D11DynamicRHI(IDXGIFactory1* InDXGIFactory1,D3D_FEATURE_LEV
 		GSupportsDepthFetchDuringDepthTest = false;
 	}
 
-   	
-	bAllowHDR		= false;
-	bAllowFlip		= true;
-	bAllowTearing	= true;
-
-	TRefCountPtr<IDXGIFactory5> Factory5;
-	HRESULT HResult = DXGIFactory1->QueryInterface(IID_PPV_ARGS(Factory5.GetInitReference()));
-	if (SUCCEEDED(HResult))
-	{
-		HResult = Factory5->CheckFeatureSupport(DXGI_FEATURE_PRESENT_ALLOW_TEARING, &bAllowTearing, sizeof(bAllowTearing));
-		bAllowHDR = true;
-	}
-
-	if (FAILED(HResult))
-	{
-		bAllowTearing = false;
-	}
-
-	// Disable FLIP if not on a supporting OS
-	TRefCountPtr<IDXGIFactory4> Factory4;
-	HResult = DXGIFactory1->QueryInterface(IID_PPV_ARGS(Factory4.GetInitReference()));
-	if (FAILED(HResult))
-	{
-		bAllowFlip =false;
-	}
-
-
 	ERHIFeatureLevel::Type PreviewFeatureLevel;
 	if (RHIGetPreviewFeatureLevel(PreviewFeatureLevel))
 	{
diff --git a/Engine/Source/Runtime/Windows/D3D11RHI/Private/D3D11RHIPrivate.h b/Engine/Source/Runtime/Windows/D3D11RHI/Private/D3D11RHIPrivate.h
--- a/Engine/Source/Runtime/Windows/D3D11RHI/Private/D3D11RHIPrivate.h
+++ b/Engine/Source/Runtime/Windows/D3D11RHI/Private/D3D11RHIPrivate.h
@@ -482,7 +482,6 @@ public:
 	virtual FViewportRHIRef RHICreateViewport(void* WindowHandle, uint32 SizeX, uint32 SizeY, bool bIsFullscreen, EPixelFormat PreferredPixelFormat) final override;
 	virtual void RHIResizeViewport(FViewportRHIParamRef Viewport, uint32 SizeX, uint32 SizeY, bool bIsFullscreen) final override;
 	virtual void RHIResizeViewport(FViewportRHIParamRef Viewport, uint32 SizeX, uint32 SizeY, bool bIsFullscreen, EPixelFormat PreferredPixelFormat) final override;
-	virtual EColorSpace RHIGetColorSpace(FViewportRHIParamRef Viewport) final override;
 	virtual void RHITick(float DeltaTime) final override;
 	virtual void RHISetStreamOutTargets(uint32 NumTargets,const FVertexBufferRHIParamRef* VertexBuffers,const uint32* Offsets) final override;
 	virtual void RHIBlockUntilGPUIdle() final override;
@@ -779,11 +778,6 @@ protected:
 	/** The global D3D interface. */
 	TRefCountPtr<IDXGIFactory1> DXGIFactory1;
 
-	// Flags determined from the particualt factories available
-	bool bAllowTearing;
-	bool bAllowHDR;
-	bool bAllowFlip;
-
 	/** The global D3D device's immediate context */
 	TRefCountPtr<FD3D11DeviceContext> Direct3DDeviceIMContext;
 
diff --git a/Engine/Source/Runtime/Windows/D3D11RHI/Private/D3D11Viewport.cpp b/Engine/Source/Runtime/Windows/D3D11RHI/Private/D3D11Viewport.cpp
--- a/Engine/Source/Runtime/Windows/D3D11RHI/Private/D3D11Viewport.cpp
+++ b/Engine/Source/Runtime/Windows/D3D11RHI/Private/D3D11Viewport.cpp
@@ -254,7 +254,7 @@ void FD3D11Viewport::Resize(uint32 InSizeX, uint32 InSizeY, bool bInIsFullscreen
 		{
 			// Resize the swap chain.
 			DXGI_FORMAT RenderTargetFormat = GetRenderTargetFormat(PixelFormat);
-			VERIFYD3D11RESIZEVIEWPORTRESULT(SwapChain->ResizeBuffers(BackBufferCount, SizeX, SizeY, RenderTargetFormat, DXGI_SWAP_CHAIN_FLAG_ALLOW_MODE_SWITCH), SizeX, SizeY, RenderTargetFormat, D3DRHI->GetDevice());
+			VERIFYD3D11RESIZEVIEWPORTRESULT(SwapChain->ResizeBuffers(1, SizeX, SizeY, RenderTargetFormat, DXGI_SWAP_CHAIN_FLAG_ALLOW_MODE_SWITCH), SizeX, SizeY, RenderTargetFormat, D3DRHI->GetDevice());
 
 			if (bInIsFullscreen)
 			{
@@ -538,12 +538,6 @@ bool FD3D11Viewport::Present(bool bLockToVsync)
 	return bNativelyPresented;
 }
 
-EColorSpace FD3D11DynamicRHI::RHIGetColorSpace(FViewportRHIParamRef ViewportRHI)
-{
-	FD3D11Viewport* Viewport = ResourceCast(ViewportRHI);
-	return Viewport->GetPixelColorSpace();
-}
-
 /*=============================================================================
  *	The following RHI functions must be called from the main thread.
  *=============================================================================*/
diff --git a/Engine/Source/Runtime/Windows/D3D11RHI/Private/Windows/WindowsD3D11Viewport.cpp b/Engine/Source/Runtime/Windows/D3D11RHI/Private/Windows/WindowsD3D11Viewport.cpp
--- a/Engine/Source/Runtime/Windows/D3D11RHI/Private/Windows/WindowsD3D11Viewport.cpp
+++ b/Engine/Source/Runtime/Windows/D3D11RHI/Private/Windows/WindowsD3D11Viewport.cpp
@@ -11,12 +11,7 @@
 #include "Windows/AllowWindowsPlatformTypes.h"
 #include <dwmapi.h>
 
-#if defined(NTDDI_WIN10_RS2)
-#include <dxgi1_6.h>
-#else
-#include <dxgi1_5.h>
-#endif
-
+#include "dxgi1_2.h"
 
 extern FD3D11Texture2D* GetSwapChainSurface(FD3D11DynamicRHI* D3DRHI, EPixelFormat PixelFormat, uint32 SizeX, uint32 SizeY, IDXGISwapChain* SwapChain);
 
@@ -33,7 +28,6 @@ FD3D11Viewport::FD3D11Viewport(FD3D11DynamicRHI* InD3DRHI,HWND InWindowHandle,ui
 	SizeY(InSizeY),
 	bIsFullscreen(bInIsFullscreen),
 	PixelFormat(InPreferredPixelFormat),
-	PixelColorSpace(EColorSpace::ERec709sRGB),
 	bIsValid(true),
 	FrameSyncEvent(InD3DRHI)
 {
@@ -118,91 +112,24 @@ FD3D11Viewport::FD3D11Viewport(FD3D11DynamicRHI* InD3DRHI,HWND InWindowHandle,ui
 		if (SwapChain == nullptr)
 		{
 			// Create the swapchain.
-			DXGI_SWAP_CHAIN_DESC1 SwapChainDesc;
-			FMemory::Memzero(&SwapChainDesc, sizeof(DXGI_SWAP_CHAIN_DESC1));
-			SwapChainDesc.Width = SizeX;
-			SwapChainDesc.Height = SizeY;
+			DXGI_SWAP_CHAIN_DESC SwapChainDesc;
+			FMemory::Memzero(&SwapChainDesc, sizeof(DXGI_SWAP_CHAIN_DESC));
+
+			SwapChainDesc.BufferDesc = SetupDXGI_MODE_DESC();
+			// MSAA Sample count
 			SwapChainDesc.SampleDesc.Count = 1;
 			SwapChainDesc.SampleDesc.Quality = 0;
-			SwapChainDesc.Format = DXGI_FORMAT_R10G10B10A2_UNORM;
 			SwapChainDesc.BufferUsage = DXGI_USAGE_RENDER_TARGET_OUTPUT | DXGI_USAGE_SHADER_INPUT;
-
-			DXGI_SWAP_CHAIN_FULLSCREEN_DESC fsSwapChainDesc = {};
-			fsSwapChainDesc.Windowed = !bIsFullscreen;
-
-			
-			if (InD3DRHI->bAllowTearing || InD3DRHI->bAllowHDR || InD3DRHI->bAllowFlip)
-			{
-				// Needed for HDR
-				BackBufferCount = 2;
-				SwapChainDesc.SwapEffect = DXGI_SWAP_EFFECT_FLIP_DISCARD;
-			}
-			else
-			{
-				// Optional for LDR
-				BackBufferCount = 1;
-				SwapChainDesc.SwapEffect = DXGI_SWAP_EFFECT_DISCARD;
-			}
-
-			SwapChainDesc.BufferCount = BackBufferCount;
+			// 1:single buffering, 2:double buffering, 3:triple buffering
+			SwapChainDesc.BufferCount = 1;
+			SwapChainDesc.OutputWindow = WindowHandle;
+			SwapChainDesc.Windowed = !bIsFullscreen;
+			// DXGI_SWAP_EFFECT_DISCARD / DXGI_SWAP_EFFECT_SEQUENTIAL
+			SwapChainDesc.SwapEffect = DXGI_SWAP_EFFECT_DISCARD;
 			SwapChainDesc.Flags = DXGI_SWAP_CHAIN_FLAG_ALLOW_MODE_SWITCH;
-
-			IDXGISwapChain1* SwapChain1 = nullptr;
-			IDXGIFactory2* Factory2 = (IDXGIFactory2*)D3DRHI->GetFactory();
-
-			VERIFYD3D11RESULT_EX((Factory2->CreateSwapChainForHwnd(D3DRHI->GetDevice(), WindowHandle, &SwapChainDesc, &fsSwapChainDesc, nullptr, &SwapChain1)), D3DRHI->GetDevice());
-			SwapChain1->QueryInterface(__uuidof(IDXGISwapChain1), (void**)SwapChain.GetInitReference());
-
-			// See if we are running on a HDR monitor
-			
-			EColorSpace ColorSpace = EColorSpace::ERec709sRGB;
-#if WITH_EDITOR
-
-			// 
-
-			static auto CVarHDREnable = IConsoleManager::Get().FindConsoleVariable(TEXT("Editor.HDRSupport"));
-
-			if (CVarHDREnable->GetInt() !=0 )
-			{
-				TRefCountPtr<IDXGIOutput> Output;
-				if (SUCCEEDED(SwapChain->GetContainingOutput(Output.GetInitReference())))
-				{
-					TRefCountPtr<IDXGIOutput6> Output6;
-					if (SUCCEEDED(Output->QueryInterface(IID_PPV_ARGS(Output6.GetInitReference()))))
-					{
-						DXGI_OUTPUT_DESC1 desc;
-						Output6->GetDesc1(&desc);
-
-						if (desc.ColorSpace == DXGI_COLOR_SPACE_RGB_FULL_G2084_NONE_P2020)
-						{
-							// Display output is HDR10.
-							ColorSpace = EColorSpace::ERec2020PQ;
-						}
-					}
-				}
-
-				if (ColorSpace == EColorSpace::ERec2020PQ)
-				{
-					TRefCountPtr<IDXGISwapChain3> swapChain3;
-					if (SUCCEEDED(SwapChain->QueryInterface(IID_PPV_ARGS(swapChain3.GetInitReference()))))
-					{
-						UINT colorSpaceSupport = 0;
-						DXGI_COLOR_SPACE_TYPE colorSpace = DXGI_COLOR_SPACE_RGB_FULL_G22_NONE_P709;
-
-						swapChain3->CheckColorSpaceSupport(colorSpace, &colorSpaceSupport);
-
-						if (colorSpaceSupport & DXGI_SWAP_CHAIN_COLOR_SPACE_SUPPORT_FLAG_PRESENT)
-						{
-							swapChain3->SetColorSpace1(colorSpace);
-						}
-					}
-				}
-			}
-#endif // WITH_EDITOR
-			PixelColorSpace = ColorSpace;
+			VERIFYD3D11RESULT_EX(D3DRHI->GetFactory()->CreateSwapChain(DXGIDevice, &SwapChainDesc, SwapChain.GetInitReference()), D3DRHI->GetDevice());
 		}
 
-
 		// Set the DXGI message hook to not change the window behind our back.
 		D3DRHI->GetFactory()->MakeWindowAssociation(WindowHandle,DXGI_MWA_NO_WINDOW_CHANGES);
 	}
@@ -242,68 +169,6 @@ void FD3D11Viewport::ConditionalResetSwapChain(bool bIgnoreFocus)
 			}
 		}
 	}
-
-	DXGI_COLOR_SPACE_TYPE colorSpace = DXGI_COLOR_SPACE_RGB_FULL_G22_NONE_P709;
-
-#if WITH_EDITOR
-#if defined(NTDDI_WIN10_RS2)
-
-	static auto CVarHDREnable = IConsoleManager::Get().FindConsoleVariable(TEXT("Editor.HDRSupport"));
-	if (CVarHDREnable->GetInt() != 0)
-	{
-		if (SwapChain)
-		{
-			TRefCountPtr<IDXGIOutput> Output;
-			if (SUCCEEDED(SwapChain->GetContainingOutput(Output.GetInitReference())))
-			{
-				TRefCountPtr<IDXGIOutput6> Output6;
-				if (SUCCEEDED(Output->QueryInterface(IID_PPV_ARGS(Output6.GetInitReference()))))
-				{
-					DXGI_OUTPUT_DESC1 desc;
-					Output6->GetDesc1(&desc);
-
-					if (desc.ColorSpace == DXGI_COLOR_SPACE_RGB_FULL_G2084_NONE_P2020)
-					{
-						// Display output is HDR10.
-						colorSpace = DXGI_COLOR_SPACE_RGB_FULL_G2084_NONE_P2020;
-					}
-				}
-			}
-		}
-	}
-
-	TRefCountPtr<IDXGISwapChain3> swapChain3;
-
-	if (SUCCEEDED(SwapChain->QueryInterface(IID_PPV_ARGS(swapChain3.GetInitReference()))))
-	{
-		UINT colorSpaceSupport = 0;
-		swapChain3->CheckColorSpaceSupport(colorSpace, &colorSpaceSupport);
-		DXGI_SWAP_CHAIN_DESC desc;
-		DXGI_SWAP_CHAIN_DESC1 desc1;
-
-		swapChain3->GetDesc(&desc);
-		swapChain3->GetDesc1(&desc1);
-
-
-		if (colorSpaceSupport & DXGI_SWAP_CHAIN_COLOR_SPACE_SUPPORT_FLAG_PRESENT)
-		{
-			swapChain3->SetColorSpace1(colorSpace);
-		}
-	}
-#endif
-#endif
-
-
-	if (colorSpace == DXGI_COLOR_SPACE_RGB_FULL_G2084_NONE_P2020)
-	{
-		PixelColorSpace = EColorSpace::ERec2020PQ;
-	}
-	else
-	{
-		PixelColorSpace = EColorSpace::ERec709sRGB;
-	}
-
-
 }
 
 #include "Windows/HideWindowsPlatformTypes.h"
diff --git a/Engine/Source/Runtime/Windows/D3D11RHI/Public/D3D11Viewport.h b/Engine/Source/Runtime/Windows/D3D11RHI/Public/D3D11Viewport.h
--- a/Engine/Source/Runtime/Windows/D3D11RHI/Public/D3D11Viewport.h
+++ b/Engine/Source/Runtime/Windows/D3D11RHI/Public/D3D11Viewport.h
@@ -9,7 +9,6 @@
 #include "CoreMinimal.h"
 #include "RenderResource.h"
 #include "RenderUtils.h"
-#include "DXGI1_5.h"
 
 /** A D3D event query resource. */
 class FD3D11EventQuery : public FRenderResource
@@ -76,7 +75,6 @@ public:
 	// Accessors.
 	FIntPoint GetSizeXY() const { return FIntPoint(SizeX, SizeY); }
 	FD3D11Texture2D* GetBackBuffer() const { return BackBuffer; }
-	EColorSpace GetPixelColorSpace() const { return PixelColorSpace; }
 
 	void WaitForFrameEventCompletion()
 	{
@@ -123,12 +121,10 @@ private:
 	uint32 MaximumFrameLatency;
 	uint32 SizeX;
 	uint32 SizeY;
-	uint32 BackBufferCount;
 	bool bIsFullscreen;
 	EPixelFormat PixelFormat;
-	EColorSpace PixelColorSpace;
 	bool bIsValid;
-	TRefCountPtr<IDXGISwapChain1> SwapChain;
+	TRefCountPtr<IDXGISwapChain> SwapChain;
 	TRefCountPtr<FD3D11Texture2D> BackBuffer;
 
 	// Support for selecting non-default output for display in fullscreen exclusive
