DR - Deprecate FShaderResourceViewRHIParamRef and FUnorderedAccessViewRHIParamRef

diff --git a/Engine/Source/Runtime/Renderer/Private/DistanceFieldLightingShared.h b/Engine/Source/Runtime/Renderer/Private/DistanceFieldLightingShared.h
--- a/Engine/Source/Runtime/Renderer/Private/DistanceFieldLightingShared.h
+++ b/Engine/Source/Runtime/Renderer/Private/DistanceFieldLightingShared.h
@@ -160,7 +160,7 @@ public:
 	{
 		if (bBarrier)
 		{
-			FUnorderedAccessViewRHIParamRef OutUAVs[2];
+			FRHIUnorderedAccessView* OutUAVs[2];
 			OutUAVs[0] = ObjectBuffers.Bounds.UAV;
 			OutUAVs[1] = ObjectBuffers.Data.UAV;
 			RHICmdList.TransitionResources(EResourceTransitionAccess::ERWBarrier, EResourceTransitionPipeline::EComputeToCompute, OutUAVs, ARRAY_COUNT(OutUAVs));
@@ -194,7 +194,7 @@ public:
 
 		if (bBarrier)
 		{
-			FUnorderedAccessViewRHIParamRef OutUAVs[2];
+			FRHIUnorderedAccessView* OutUAVs[2];
 			OutUAVs[0] = ObjectBuffers.Bounds.UAV;
 			OutUAVs[1] = ObjectBuffers.Data.UAV;
 			RHICmdList.TransitionResources(EResourceTransitionAccess::EReadable, EResourceTransitionPipeline::EComputeToCompute, OutUAVs, ARRAY_COUNT(OutUAVs));
@@ -402,7 +402,7 @@ public:
 		CulledObjectBoxBounds.UnsetUAV(RHICmdList, ShaderRHI);
 	}
 
-	void GetUAVs(const FDistanceFieldCulledObjectBuffers& ObjectBuffers, TArray<FUnorderedAccessViewRHIParamRef>& UAVs)
+	void GetUAVs(const FDistanceFieldCulledObjectBuffers& ObjectBuffers, TArray<FRHIUnorderedAccessView*>& UAVs)
 	{
 		uint32 MaxIndex = 0;
 		MaxIndex = FMath::Max(MaxIndex, ObjectIndirectArguments.GetUAVIndex());
@@ -577,7 +577,7 @@ public:
 		SetShaderValue(RHICmdList, ShaderRHI, ShadowAverageObjectsPerTile, GAverageObjectsPerShadowCullTile);
 	}
 
-	void GetUAVs(FLightTileIntersectionResources& TileIntersectionResources, TArray<FUnorderedAccessViewRHIParamRef>& UAVs)
+	void GetUAVs(FLightTileIntersectionResources& TileIntersectionResources, TArray<FRHIUnorderedAccessView*>& UAVs)
 	{
 		int32 MaxIndex = FMath::Max(
 			FMath::Max(ShadowTileNumCulledObjects.GetUAVIndex(), ShadowTileStartOffsets.GetUAVIndex()), 
diff --git a/Engine/Source/Runtime/Renderer/Private/DistanceFieldObjectCulling.cpp b/Engine/Source/Runtime/Renderer/Private/DistanceFieldObjectCulling.cpp
--- a/Engine/Source/Runtime/Renderer/Private/DistanceFieldObjectCulling.cpp
+++ b/Engine/Source/Runtime/Renderer/Private/DistanceFieldObjectCulling.cpp
@@ -136,7 +136,7 @@ public:
 
 	void SetParameters(FRHICommandList& RHICmdList, const FScene* Scene, const FSceneView& View, const FDistanceFieldAOParameters& Parameters)
 	{
-		FUnorderedAccessViewRHIParamRef OutUAVs[6];
+		FRHIUnorderedAccessView* OutUAVs[6];
 		OutUAVs[0] = GAOCulledObjectBuffers.Buffers.ObjectIndirectArguments.UAV;
 		OutUAVs[1] = GAOCulledObjectBuffers.Buffers.Bounds.UAV;
 		OutUAVs[2] = GAOCulledObjectBuffers.Buffers.Data.UAV;
@@ -163,7 +163,7 @@ public:
 		ObjectBufferParameters.UnsetParameters(RHICmdList, GetComputeShader(), *(Scene->DistanceFieldSceneData.ObjectBuffers));
 		CulledObjectParameters.UnsetParameters(RHICmdList, GetComputeShader());
 
-		FUnorderedAccessViewRHIParamRef OutUAVs[6];
+		FRHIUnorderedAccessView* OutUAVs[6];
 		OutUAVs[0] = GAOCulledObjectBuffers.Buffers.ObjectIndirectArguments.UAV;
 		OutUAVs[1] = GAOCulledObjectBuffers.Buffers.Bounds.UAV;
 		OutUAVs[2] = GAOCulledObjectBuffers.Buffers.Data.UAV;
@@ -271,7 +271,7 @@ public:
 
 		FTileIntersectionResources* TileIntersectionResources = ((FSceneViewState*)View.State)->AOTileIntersectionResources;
 
-		FUnorderedAccessViewRHIParamRef OutUAVs[2];
+		FRHIUnorderedAccessView* OutUAVs[2];
 		OutUAVs[0] = TileIntersectionResources->TileConeAxisAndCos.UAV;
 		OutUAVs[1] = TileIntersectionResources->TileConeDepthRanges.UAV;
 		RHICmdList.TransitionResources(EResourceTransitionAccess::ERWBarrier, EResourceTransitionPipeline::EComputeToCompute, OutUAVs, ARRAY_COUNT(OutUAVs));
@@ -300,7 +300,7 @@ public:
 
 		FTileIntersectionResources* TileIntersectionResources = ((FSceneViewState*)View.State)->AOTileIntersectionResources;
 
-		FUnorderedAccessViewRHIParamRef OutUAVs[2];
+		FRHIUnorderedAccessView* OutUAVs[2];
 		OutUAVs[0] = TileIntersectionResources->TileConeAxisAndCos.UAV;
 		OutUAVs[1] = TileIntersectionResources->TileConeDepthRanges.UAV;
 		RHICmdList.TransitionResources(EResourceTransitionAccess::EReadable, EResourceTransitionPipeline::EComputeToCompute, OutUAVs, ARRAY_COUNT(OutUAVs));
@@ -440,7 +440,7 @@ public:
 		SetShaderValue(RHICmdList, ShaderRHI, NumGroups, NumGroupsValue);
 	}
 
-	void GetUAVs(const FSceneView& View, TArray<FUnorderedAccessViewRHIParamRef>& UAVs)
+	void GetUAVs(const FSceneView& View, TArray<FRHIUnorderedAccessView*>& UAVs)
 	{
 		FTileIntersectionResources* TileIntersectionResources = ((FSceneViewState*)View.State)->AOTileIntersectionResources;
 		TileIntersectionParameters.GetUAVs(*TileIntersectionResources, UAVs);
@@ -513,7 +513,7 @@ public:
 
 		FTileIntersectionResources* TileIntersectionResources = ((FSceneViewState*)View.State)->AOTileIntersectionResources;
 		
-		TArray<FUnorderedAccessViewRHIParamRef> UAVs;
+		TArray<FRHIUnorderedAccessView*> UAVs;
 		TileIntersectionParameters.GetUAVs(*TileIntersectionResources, UAVs);
 
 		RHICmdList.TransitionResources(EResourceTransitionAccess::EWritable, EResourceTransitionPipeline::EComputeToCompute, UAVs.GetData(), UAVs.Num());
@@ -527,7 +527,7 @@ public:
 
 		TileIntersectionParameters.UnsetParameters(RHICmdList, GetComputeShader());
 
-		TArray<FUnorderedAccessViewRHIParamRef> UAVs;
+		TArray<FRHIUnorderedAccessView*> UAVs;
 		TileIntersectionParameters.GetUAVs(*TileIntersectionResources, UAVs);
 
 		RHICmdList.TransitionResources(EResourceTransitionAccess::EReadable, EResourceTransitionPipeline::EComputeToCompute, UAVs.GetData(), UAVs.Num());
@@ -555,7 +555,7 @@ void ScatterTilesToObjects(FRHICommandListImmediate& RHICmdList, const FViewInfo
 	TShaderMapRef<FObjectCullVS> VertexShader(View.ShaderMap);
 	TShaderMapRef<TObjectCullPS<bCountingPass>> PixelShader(View.ShaderMap);
 
-	TArray<FUnorderedAccessViewRHIParamRef> UAVs;
+	TArray<FRHIUnorderedAccessView*> UAVs;
 	PixelShader->GetUAVs(View, UAVs);
 	RHICmdList.TransitionResources(EResourceTransitionAccess::ERWBarrier, EResourceTransitionPipeline::EComputeToGfx, UAVs.GetData(), UAVs.Num());
 
diff --git a/Engine/Source/Runtime/Renderer/Private/DistanceFieldObjectManagement.cpp b/Engine/Source/Runtime/Renderer/Private/DistanceFieldObjectManagement.cpp
--- a/Engine/Source/Runtime/Renderer/Private/DistanceFieldObjectManagement.cpp
+++ b/Engine/Source/Runtime/Renderer/Private/DistanceFieldObjectManagement.cpp
@@ -176,7 +176,7 @@ public:
 	{
 	}
 
-	void SetParameters(FRHICommandList& RHICmdList, const FScene* Scene, uint32 NumUploadOperationsValue, FShaderResourceViewRHIParamRef InUploadOperationIndices, FShaderResourceViewRHIParamRef InUploadOperationData)
+	void SetParameters(FRHICommandList& RHICmdList, const FScene* Scene, uint32 NumUploadOperationsValue, FRHIShaderResourceView* InUploadOperationIndices, FRHIShaderResourceView* InUploadOperationData)
 	{
 		FComputeShaderRHIParamRef ShaderRHI = GetComputeShader();
 
@@ -194,7 +194,7 @@ public:
 		ObjectBufferParameters.UnsetParameters(RHICmdList, GetComputeShader(), *(Scene->DistanceFieldSceneData.ObjectBuffers), true);
 
 		const FDistanceFieldObjectBuffers& ObjectBuffers = *(Scene->DistanceFieldSceneData.ObjectBuffers);
-		FUnorderedAccessViewRHIParamRef OutUAVs[2];
+		FRHIUnorderedAccessView* OutUAVs[2];
 		OutUAVs[0] = ObjectBuffers.Bounds.UAV;
 		OutUAVs[1] = ObjectBuffers.Data.UAV;
 		RHICmdList.TransitionResources(EResourceTransitionAccess::EReadable, EResourceTransitionPipeline::EComputeToCompute, OutUAVs, ARRAY_COUNT(OutUAVs));
@@ -252,7 +252,7 @@ public:
 	{
 		FComputeShaderRHIParamRef ShaderRHI = GetComputeShader();
 
-		FUnorderedAccessViewRHIParamRef OutUAVs[2];
+		FRHIUnorderedAccessView* OutUAVs[2];
 		OutUAVs[0] = ObjectBuffersDest.Bounds.UAV;
 		OutUAVs[1] = ObjectBuffersDest.Data.UAV;
 		RHICmdList.TransitionResources(EResourceTransitionAccess::ERWBarrier, EResourceTransitionPipeline::EComputeToCompute, OutUAVs, ARRAY_COUNT(OutUAVs));
@@ -269,7 +269,7 @@ public:
 		CopyObjectBounds.UnsetUAV(RHICmdList, GetComputeShader());
 		CopyObjectData.UnsetUAV(RHICmdList, GetComputeShader());
 
-		FUnorderedAccessViewRHIParamRef OutUAVs[2];
+		FRHIUnorderedAccessView* OutUAVs[2];
 		OutUAVs[0] = ObjectBuffersDest.Bounds.UAV;
 		OutUAVs[1] = ObjectBuffersDest.Data.UAV;
 		RHICmdList.TransitionResources(EResourceTransitionAccess::EReadable, EResourceTransitionPipeline::EComputeToCompute, OutUAVs, ARRAY_COUNT(OutUAVs));
@@ -326,7 +326,7 @@ public:
 	{
 		FComputeShaderRHIParamRef ShaderRHI = GetComputeShader();
 
-		FUnorderedAccessViewRHIParamRef OutUAVs[2];
+		FRHIUnorderedAccessView* OutUAVs[2];
 		OutUAVs[0] = SurfelBuffersDest.InterpolatedVertexData.UAV;
 		OutUAVs[1] = SurfelBuffersDest.Surfels.UAV;		
 		RHICmdList.TransitionResources(EResourceTransitionAccess::ERWBarrier, EResourceTransitionPipeline::EComputeToCompute, OutUAVs, ARRAY_COUNT(OutUAVs));
@@ -343,7 +343,7 @@ public:
 		CopyInterpolatedVertexData.UnsetUAV(RHICmdList, GetComputeShader());
 		CopySurfelData.UnsetUAV(RHICmdList, GetComputeShader());
 
-		FUnorderedAccessViewRHIParamRef OutUAVs[2];
+		FRHIUnorderedAccessView* OutUAVs[2];
 		OutUAVs[0] = SurfelBuffersDest.InterpolatedVertexData.UAV;
 		OutUAVs[1] = SurfelBuffersDest.Surfels.UAV;
 		RHICmdList.TransitionResources(EResourceTransitionAccess::EReadable, EResourceTransitionPipeline::EComputeToCompute, OutUAVs, ARRAY_COUNT(OutUAVs));
@@ -469,9 +469,9 @@ public:
 		FRHICommandList& RHICmdList, 
 		const FScene* Scene, 
 		uint32 NumRemoveOperationsValue, 
-		FShaderResourceViewRHIParamRef InRemoveOperationIndices, 
-		FShaderResourceViewRHIParamRef InObjectBounds2, 
-		FShaderResourceViewRHIParamRef InObjectData2)
+		FRHIShaderResourceView* InRemoveOperationIndices,
+		FRHIShaderResourceView* InObjectBounds2,
+		FRHIShaderResourceView* InObjectData2)
 	{
 		FComputeShaderRHIParamRef ShaderRHI = GetComputeShader();
 
diff --git a/Engine/Source/Runtime/Renderer/Private/DistanceFieldScreenGridLighting.cpp b/Engine/Source/Runtime/Renderer/Private/DistanceFieldScreenGridLighting.cpp
--- a/Engine/Source/Runtime/Renderer/Private/DistanceFieldScreenGridLighting.cpp
+++ b/Engine/Source/Runtime/Renderer/Private/DistanceFieldScreenGridLighting.cpp
@@ -323,7 +323,7 @@ public:
 		FAOScreenGridResources* ScreenGridResources = View.ViewState->AOScreenGridResources;
 
 		int32 NumOutUAVs = 0;
-		FUnorderedAccessViewRHIParamRef OutUAVs[1];
+		FRHIUnorderedAccessView* OutUAVs[1];
 		OutUAVs[NumOutUAVs++] = ScreenGridResources->ScreenGridConeVisibility.UAV;
 
 		// Note: no transition, want to overlap object cone tracing and global DF cone tracing since both shaders use atomics to ScreenGridConeVisibility
diff --git a/Engine/Source/Runtime/Renderer/Private/DistanceFieldShadowing.cpp b/Engine/Source/Runtime/Renderer/Private/DistanceFieldShadowing.cpp
--- a/Engine/Source/Runtime/Renderer/Private/DistanceFieldShadowing.cpp
+++ b/Engine/Source/Runtime/Renderer/Private/DistanceFieldShadowing.cpp
@@ -138,7 +138,7 @@ public:
 		FGlobalShader::SetParameters<FViewUniformShaderParameters>(RHICmdList, ShaderRHI, View.ViewUniformBuffer);
 		ObjectBufferParameters.Set(RHICmdList, ShaderRHI, *(Scene->DistanceFieldSceneData.ObjectBuffers), Scene->DistanceFieldSceneData.NumObjectsInBuffer);
 
-		FUnorderedAccessViewRHIParamRef OutUAVs[4];
+		FRHIUnorderedAccessView* OutUAVs[4];
 		OutUAVs[0] = GShadowCulledObjectBuffers.Buffers.ObjectIndirectArguments.UAV;
 		OutUAVs[1] = GShadowCulledObjectBuffers.Buffers.Bounds.UAV;
 		OutUAVs[2] = GShadowCulledObjectBuffers.Buffers.Data.UAV;
@@ -167,7 +167,7 @@ public:
 		ObjectBufferParameters.UnsetParameters(RHICmdList, GetComputeShader(), *(Scene->DistanceFieldSceneData.ObjectBuffers));
 		CulledObjectParameters.UnsetParameters(RHICmdList, GetComputeShader());
 
-		FUnorderedAccessViewRHIParamRef OutUAVs[4];
+		FRHIUnorderedAccessView* OutUAVs[4];
 		OutUAVs[0] = GShadowCulledObjectBuffers.Buffers.ObjectIndirectArguments.UAV;
 		OutUAVs[1] = GShadowCulledObjectBuffers.Buffers.Bounds.UAV;
 		OutUAVs[2] = GShadowCulledObjectBuffers.Buffers.Data.UAV;
@@ -293,7 +293,7 @@ public:
 		LightTileIntersectionParameters.Set(RHICmdList, ShaderRHI, *TileIntersectionResources);
 	}
 
-	void GetUAVs(const FSceneView& View, FLightTileIntersectionResources* TileIntersectionResources, TArray<FUnorderedAccessViewRHIParamRef>& UAVs)
+	void GetUAVs(const FSceneView& View, FLightTileIntersectionResources* TileIntersectionResources, TArray<FRHIUnorderedAccessView*>& UAVs)
 	{
 		LightTileIntersectionParameters.GetUAVs(*TileIntersectionResources, UAVs);
 	}
@@ -641,7 +641,7 @@ public:
 
 		FGlobalShader::SetParameters<FViewUniformShaderParameters>(RHICmdList, ShaderRHI, View.ViewUniformBuffer);
 
-		TArray<FUnorderedAccessViewRHIParamRef> UAVs;
+		TArray<FRHIUnorderedAccessView*> UAVs;
 		TileIntersectionParameters.GetUAVs(*TileIntersectionResources, UAVs);
 
 		RHICmdList.TransitionResources(EResourceTransitionAccess::EWritable, EResourceTransitionPipeline::EComputeToCompute, UAVs.GetData(), UAVs.Num());
@@ -653,7 +653,7 @@ public:
 	{
 		TileIntersectionParameters.UnsetParameters(RHICmdList, GetComputeShader());
 
-		TArray<FUnorderedAccessViewRHIParamRef> UAVs;
+		TArray<FRHIUnorderedAccessView*> UAVs;
 		TileIntersectionParameters.GetUAVs(*TileIntersectionResources, UAVs);
 
 		RHICmdList.TransitionResources(EResourceTransitionAccess::EReadable, EResourceTransitionPipeline::EComputeToCompute, UAVs.GetData(), UAVs.Num());
@@ -685,7 +685,7 @@ void ScatterObjectsToShadowTiles(
 	TShaderMapRef<FShadowObjectCullVS> VertexShader(View.ShaderMap);
 	TShaderMapRef<FShadowObjectCullPS<bCountingPass>> PixelShader(View.ShaderMap);
 
-	TArray<FUnorderedAccessViewRHIParamRef> UAVs;
+	TArray<FRHIUnorderedAccessView*> UAVs;
 	PixelShader->GetUAVs(View, TileIntersectionResources, UAVs);
 	RHICmdList.TransitionResources(EResourceTransitionAccess::ERWBarrier, EResourceTransitionPipeline::EComputeToGfx, UAVs.GetData(), UAVs.Num());
 
diff --git a/Engine/Source/Runtime/Renderer/Private/GlobalDistanceField.cpp b/Engine/Source/Runtime/Renderer/Private/GlobalDistanceField.cpp
--- a/Engine/Source/Runtime/Renderer/Private/GlobalDistanceField.cpp
+++ b/Engine/Source/Runtime/Renderer/Private/GlobalDistanceField.cpp
@@ -186,7 +186,7 @@ public:
 		FGlobalShader::SetParameters<FViewUniformShaderParameters>(RHICmdList, ShaderRHI, View.ViewUniformBuffer);
 		ObjectBufferParameters.Set(RHICmdList, ShaderRHI, *(Scene->DistanceFieldSceneData.ObjectBuffers), Scene->DistanceFieldSceneData.NumObjectsInBuffer);
 
-		FUnorderedAccessViewRHIParamRef OutUAVs[4];
+		FRHIUnorderedAccessView* OutUAVs[4];
 		OutUAVs[0] = GGlobalDistanceFieldCulledObjectBuffers.Buffers.ObjectIndirectArguments.UAV;
 		OutUAVs[1] = GGlobalDistanceFieldCulledObjectBuffers.Buffers.Bounds.UAV;
 		OutUAVs[2] = GGlobalDistanceFieldCulledObjectBuffers.Buffers.Data.UAV;
@@ -221,7 +221,7 @@ public:
 		ObjectBufferParameters.UnsetParameters(RHICmdList, GetComputeShader(), *(Scene->DistanceFieldSceneData.ObjectBuffers));
 		CulledObjectParameters.UnsetParameters(RHICmdList, GetComputeShader());
 
-		TArray<FUnorderedAccessViewRHIParamRef> UAVs;
+		TArray<FRHIUnorderedAccessView*> UAVs;
 		CulledObjectParameters.GetUAVs(GGlobalDistanceFieldCulledObjectBuffers.Buffers, UAVs);
 		RHICmdList.TransitionResources(EResourceTransitionAccess::EReadable, EResourceTransitionPipeline::EComputeToCompute, UAVs.GetData(), UAVs.Num());
 	}
diff --git a/Engine/Source/Runtime/Renderer/Private/HeightfieldLighting.cpp b/Engine/Source/Runtime/Renderer/Private/HeightfieldLighting.cpp
--- a/Engine/Source/Runtime/Renderer/Private/HeightfieldLighting.cpp
+++ b/Engine/Source/Runtime/Renderer/Private/HeightfieldLighting.cpp
@@ -590,7 +590,7 @@ public:
 
 TGlobalResource<FHeightfieldDescriptionsResource> GHeightfieldDescriptions;
 
-FShaderResourceViewRHIParamRef GetHeightfieldDescriptionsSRV()
+FRHIShaderResourceView* GetHeightfieldDescriptionsSRV()
 {
 	return GHeightfieldDescriptions.Data.BufferSRV;
 }
diff --git a/Engine/Source/Runtime/Renderer/Private/HeightfieldLighting.h b/Engine/Source/Runtime/Renderer/Private/HeightfieldLighting.h
--- a/Engine/Source/Runtime/Renderer/Private/HeightfieldLighting.h
+++ b/Engine/Source/Runtime/Renderer/Private/HeightfieldLighting.h
@@ -144,7 +144,7 @@ private:
 	FHeightfieldDescription Heightfield;
 };
 
-extern FShaderResourceViewRHIParamRef GetHeightfieldDescriptionsSRV();
+extern FRHIShaderResourceView* GetHeightfieldDescriptionsSRV();
 
 class FHeightfieldDescriptionParameters
 {
@@ -163,7 +163,7 @@ public:
 	}
 
 	template<typename ShaderRHIParamRef>
-	void Set(FRHICommandList& RHICmdList, const ShaderRHIParamRef ShaderRHI, FShaderResourceViewRHIParamRef HeightfieldDescriptionsValue, int32 NumHeightfieldsValue)
+	void Set(FRHICommandList& RHICmdList, const ShaderRHIParamRef ShaderRHI, FRHIShaderResourceView* HeightfieldDescriptionsValue, int32 NumHeightfieldsValue)
 	{
 		SetSRVParameter(RHICmdList, ShaderRHI, HeightfieldDescriptions, HeightfieldDescriptionsValue);
 		SetShaderValue(RHICmdList, ShaderRHI, NumHeightfields, NumHeightfieldsValue);
diff --git a/Engine/Source/Runtime/Renderer/Private/LightGridInjection.cpp b/Engine/Source/Runtime/Renderer/Private/LightGridInjection.cpp
--- a/Engine/Source/Runtime/Renderer/Private/LightGridInjection.cpp
+++ b/Engine/Source/Runtime/Renderer/Private/LightGridInjection.cpp
@@ -688,7 +688,7 @@ void FDeferredShadingSceneRenderer::ComputeLightGrid(FRHICommandListImmediate& R
 			const FIntVector NumGroups = FIntVector::DivideAndRoundUp(FIntVector(LightGridSizeXY.X, LightGridSizeXY.Y, GLightGridSizeZ), LightGridInjectionGroupSize);
 
 			// TODO: When external buffers and extract can be set up maybe use that to get these out instead.
-			TArray<FUnorderedAccessViewRHIParamRef, TInlineAllocator<2>> OutUAVs({
+			TArray<FRHIUnorderedAccessView*, TInlineAllocator<2>> OutUAVs({
 				View.ForwardLightingResources->NumCulledLightsGrid.UAV,
 				View.ForwardLightingResources->CulledLightDataGrid.UAV });
 
diff --git a/Engine/Source/Runtime/Renderer/Private/LightPropagationVolume.cpp b/Engine/Source/Runtime/Renderer/Private/LightPropagationVolume.cpp
--- a/Engine/Source/Runtime/Renderer/Private/LightPropagationVolume.cpp
+++ b/Engine/Source/Runtime/Renderer/Private/LightPropagationVolume.cpp
@@ -205,7 +205,7 @@ public:
 		SetUniformBufferParameter(RHICmdList, ShaderRHI, GetUniformBufferParameter<FLpvWriteUniformBufferParameters>(), Params.UniformBuffer );
 
 		TArray<int32> ResourceIndices;
-		TArray<FUnorderedAccessViewRHIParamRef> UAVs;
+		TArray<FRHIUnorderedAccessView*> UAVs;
 
 		for(int i  =0; i < 7; i++)
 		{
@@ -295,7 +295,7 @@ public:
 	void UnbindBuffers(FRHICommandList& RHICmdList, const FLpvBaseWriteShaderParams& Params)
 	{
 		TArray<int32> ResourceIndices;
-		TArray<FUnorderedAccessViewRHIParamRef> UAVs;
+		TArray<FRHIUnorderedAccessView*> UAVs;
 
 		FComputeShaderRHIParamRef ShaderRHI = GetComputeShader();
 		for ( int i = 0; i < 7; i++ )
@@ -312,7 +312,7 @@ public:
 	    }
 		if ( VplListHeadBufferSRV.IsBound() )
 		{
-			RHICmdList.SetShaderResourceViewParameter( ShaderRHI, VplListHeadBufferSRV.GetBaseIndex(), FShaderResourceViewRHIParamRef() );
+			RHICmdList.SetShaderResourceViewParameter( ShaderRHI, VplListHeadBufferSRV.GetBaseIndex(), nullptr );
 		}
 		if ( VplListHeadBufferUAV.IsBound() )
 		{
@@ -321,7 +321,7 @@ public:
 		}
 		if ( VplListBufferSRV.IsBound() )
 		{
-			RHICmdList.SetShaderResourceViewParameter( ShaderRHI, VplListBufferSRV.GetBaseIndex(), FShaderResourceViewRHIParamRef() );
+			RHICmdList.SetShaderResourceViewParameter( ShaderRHI, VplListBufferSRV.GetBaseIndex(), nullptr );
 		}
 		if ( VplListBufferUAV.IsBound() )
 		{
@@ -348,7 +348,7 @@ public:
 		}
 		if ( AOVolumeTextureSRV.IsBound() )
 		{
-			RHICmdList.SetShaderResourceViewParameter( ShaderRHI, AOVolumeTextureSRV.GetBaseIndex(), FShaderResourceViewRHIParamRef() );
+			RHICmdList.SetShaderResourceViewParameter( ShaderRHI, AOVolumeTextureSRV.GetBaseIndex(), nullptr );
 		}
 		if(GvListBufferUAV.IsBound())
 		{
@@ -362,16 +362,16 @@ public:
 		}
 		if ( GvListBufferSRV.IsBound() )
 		{
-			RHICmdList.SetShaderResourceViewParameter( ShaderRHI, GvListBufferSRV.GetBaseIndex(), FShaderResourceViewRHIParamRef() );
+			RHICmdList.SetShaderResourceViewParameter( ShaderRHI, GvListBufferSRV.GetBaseIndex(), nullptr );
 		}
 		if ( GvListHeadBufferSRV.IsBound() )
 		{
-			RHICmdList.SetShaderResourceViewParameter( ShaderRHI, GvListHeadBufferSRV.GetBaseIndex(), FShaderResourceViewRHIParamRef() );
+			RHICmdList.SetShaderResourceViewParameter( ShaderRHI, GvListHeadBufferSRV.GetBaseIndex(), nullptr );
 		}
 
 		check(ResourceIndices.Num() == UAVs.Num());
 		RHICmdList.TransitionResources(EResourceTransitionAccess::EReadable, EResourceTransitionPipeline::EComputeToGfx, UAVs.GetData(), UAVs.Num());
-		FUnorderedAccessViewRHIParamRef NullUAV = nullptr;
+		FRHIUnorderedAccessView* NullUAV = nullptr;
 		for (int32 i = 0; i < ResourceIndices.Num(); ++i)
 		{
 			RHICmdList.SetUAVParameter(ShaderRHI, ResourceIndices[i], NullUAV);
@@ -1055,7 +1055,7 @@ void FLightPropagationVolume::Clear(FRHICommandListImmediate& RHICmdList, FViewI
 
 	RHICmdList.SetUAVParameter( FComputeShaderRHIRef(), 7, mVplListBuffer->UAV, 0 );
 	RHICmdList.SetUAVParameter( FComputeShaderRHIRef(), 7, GvListBuffer->UAV, 0 );
-	RHICmdList.SetUAVParameter( FComputeShaderRHIRef(), 7, FUnorderedAccessViewRHIParamRef(), 0 );
+	RHICmdList.SetUAVParameter( FComputeShaderRHIRef(), 7, nullptr, 0 );
 }
 
 /**
diff --git a/Engine/Source/Runtime/Renderer/Private/LightPropagationVolume.h b/Engine/Source/Runtime/Renderer/Private/LightPropagationVolume.h
--- a/Engine/Source/Runtime/Renderer/Private/LightPropagationVolume.h
+++ b/Engine/Source/Runtime/Renderer/Private/LightPropagationVolume.h
@@ -128,22 +128,22 @@ struct FLpvBaseWriteShaderParams
 {
 	FLpvWriteUniformBufferRef		UniformBuffer;
 	FTextureRHIParamRef				LpvBufferSRVs[7];
-	FUnorderedAccessViewRHIParamRef LpvBufferUAVs[7];
+	FRHIUnorderedAccessView* LpvBufferUAVs[7];
 
-	FShaderResourceViewRHIParamRef	VplListHeadBufferSRV;
-	FUnorderedAccessViewRHIParamRef VplListHeadBufferUAV;
-	FShaderResourceViewRHIParamRef	VplListBufferSRV;
-	FUnorderedAccessViewRHIParamRef VplListBufferUAV;
+	FRHIShaderResourceView*	VplListHeadBufferSRV;
+	FRHIUnorderedAccessView* VplListHeadBufferUAV;
+	FRHIShaderResourceView*	VplListBufferSRV;
+	FRHIUnorderedAccessView* VplListBufferUAV;
 
 	FTextureRHIParamRef				GvBufferSRVs[3];
-	FUnorderedAccessViewRHIParamRef GvBufferUAVs[3];
+	FRHIUnorderedAccessView* GvBufferUAVs[3];
 
-	FShaderResourceViewRHIParamRef	GvListHeadBufferSRV;
-	FUnorderedAccessViewRHIParamRef GvListHeadBufferUAV;
-	FShaderResourceViewRHIParamRef	GvListBufferSRV;
-	FUnorderedAccessViewRHIParamRef GvListBufferUAV;
+	FRHIShaderResourceView*	GvListHeadBufferSRV;
+	FRHIUnorderedAccessView* GvListHeadBufferUAV;
+	FRHIShaderResourceView*	GvListBufferSRV;
+	FRHIUnorderedAccessView* GvListBufferUAV;
 
-	FUnorderedAccessViewRHIParamRef AOVolumeTextureUAV;
+	FRHIUnorderedAccessView* AOVolumeTextureUAV;
 	FTextureRHIParamRef				AOVolumeTextureSRV;
 };
 
@@ -190,11 +190,11 @@ public:
 
 	FTextureRHIParamRef GetLpvBufferSrv( int i )						{ return LpvVolumeTextures[ 1-mWriteBufferIndex ][i]->GetRenderTargetItem().ShaderResourceTexture; }
 
-	FUnorderedAccessViewRHIParamRef GetVplListBufferUav()				{ return mVplListBuffer->UAV; }
-	FUnorderedAccessViewRHIParamRef GetVplListHeadBufferUav()			{ return mVplListHeadBuffer->UAV; }
+	FRHIUnorderedAccessView* GetVplListBufferUav()				{ return mVplListBuffer->UAV; }
+	FRHIUnorderedAccessView* GetVplListHeadBufferUav()			{ return mVplListHeadBuffer->UAV; }
 
-	FUnorderedAccessViewRHIParamRef GetGvListBufferUav()				{ return GvListBuffer->UAV; }
-	FUnorderedAccessViewRHIParamRef GetGvListHeadBufferUav()			{ return GvListHeadBuffer->UAV; }
+	FRHIUnorderedAccessView* GetGvListBufferUav()				{ return GvListBuffer->UAV; }
+	FRHIUnorderedAccessView* GetGvListHeadBufferUav()			{ return GvListHeadBuffer->UAV; }
 
 	bool IsEnabled() const												{ return bEnabled;	}
 	bool IsDirectionalOcclusionEnabled() const							{ return bDirectionalOcclusionEnabled; }
diff --git a/Engine/Source/Runtime/Renderer/Private/MeshPassProcessor.cpp b/Engine/Source/Runtime/Renderer/Private/MeshPassProcessor.cpp
--- a/Engine/Source/Runtime/Renderer/Private/MeshPassProcessor.cpp
+++ b/Engine/Source/Runtime/Renderer/Private/MeshPassProcessor.cpp
@@ -43,7 +43,7 @@ class FShaderBindingState
 {
 public:
 	int32 MaxSRVUsed = -1;
-	FShaderResourceViewRHIParamRef SRVs[MAX_SRVs_PER_SHADER_STAGE] = {};
+	FRHIShaderResourceView* SRVs[MAX_SRVs_PER_SHADER_STAGE] = {};
 	int32 MaxUniformBufferUsed = -1;
 	FUniformBufferRHIParamRef UniformBuffers[MAX_UNIFORM_BUFFERS_PER_SHADER_STAGE] = {};
 	int32 MaxTextureUsed = -1;
@@ -152,7 +152,7 @@ void FMeshDrawShaderBindings::SetShaderBindings(
 
 		if (SRVType[TypeByteIndex] & (1 << TypeBitIndex))
 		{
-			FShaderResourceViewRHIParamRef SRV = (FShaderResourceViewRHIParamRef)SRVBindings[SRVIndex];
+			FRHIShaderResourceView* SRV = (FRHIShaderResourceView*)SRVBindings[SRVIndex];
 
 			if (SRV != ShaderBindingState.SRVs[Parameter.BaseIndex])
 			{
@@ -237,7 +237,7 @@ void FMeshDrawShaderBindings::SetShaderBindings(
 
 		if (SRVType[TypeByteIndex] & (1 << TypeBitIndex))
 		{
-			FShaderResourceViewRHIParamRef SRV = (FShaderResourceViewRHIParamRef)SRVBindings[SRVIndex];
+			FRHIShaderResourceView* SRV = (FRHIShaderResourceView*)SRVBindings[SRVIndex];
 			RHICmdList.SetShaderResourceViewParameter(Shader, Parameter.BaseIndex, SRV);
 		}
 		else
@@ -652,7 +652,7 @@ void FMeshDrawShaderBindings::Finalize(const FMeshProcessorShaders* ShadersForDe
 
 			if (SRVType[TypeByteIndex] & (1 << TypeBitIndex))
 			{
-				FShaderResourceViewRHIParamRef SRV = (FShaderResourceViewRHIParamRef)SRVBindings[SRVIndex];
+				FRHIShaderResourceView* SRV = (FRHIShaderResourceView*)SRVBindings[SRVIndex];
 
 				ensureMsgf(SRV, TEXT("Shader %s with vertex factory %s never set SRV at BaseIndex %u.  This can cause GPU hangs, depending on how the shader uses it."), 
 					Shader->GetType()->GetName(), 
diff --git a/Engine/Source/Runtime/Renderer/Private/PathTracing.cpp b/Engine/Source/Runtime/Renderer/Private/PathTracing.cpp
--- a/Engine/Source/Runtime/Renderer/Private/PathTracing.cpp
+++ b/Engine/Source/Runtime/Renderer/Private/PathTracing.cpp
@@ -122,10 +122,10 @@ public:
 		FIntVector VarianceDimensions,
 		const FRWBuffer& VarianceMipTree,
 		// Output
-		FUnorderedAccessViewRHIParamRef RadianceUAV,
-		FUnorderedAccessViewRHIParamRef SampleCountUAV,
-		FUnorderedAccessViewRHIParamRef PixelPositionUAV,
-		FUnorderedAccessViewRHIParamRef RayCountPerPixelUAV)
+		FRHIUnorderedAccessView* RadianceUAV,
+		FRHIUnorderedAccessView* SampleCountUAV,
+		FRHIUnorderedAccessView* PixelPositionUAV,
+		FRHIUnorderedAccessView* RayCountPerPixelUAV)
 	{
 
 		GlobalResources.Set(TLASParameter, RayTracingScene.RayTracingSceneRHI->GetShaderResourceView());
diff --git a/Engine/Source/Runtime/Renderer/Private/PathTracing/PathCompactionCompute.cpp b/Engine/Source/Runtime/Renderer/Private/PathTracing/PathCompactionCompute.cpp
--- a/Engine/Source/Runtime/Renderer/Private/PathTracing/PathCompactionCompute.cpp
+++ b/Engine/Source/Runtime/Renderer/Private/PathTracing/PathCompactionCompute.cpp
@@ -58,11 +58,11 @@ public:
 		FTextureRHIParamRef RadianceTexture,
 		FTextureRHIParamRef SampleCountTexture,
 		FTextureRHIParamRef PixelPositionTexture,
-		FUnorderedAccessViewRHIParamRef RadianceSortedRedUAV,
-		FUnorderedAccessViewRHIParamRef RadianceSortedGreenUAV,
-		FUnorderedAccessViewRHIParamRef RadianceSortedBlueUAV,
-		FUnorderedAccessViewRHIParamRef RadianceSortedAlphaUAV,
-		FUnorderedAccessViewRHIParamRef SampleCountSortedUAV)
+		FRHIUnorderedAccessView* RadianceSortedRedUAV,
+		FRHIUnorderedAccessView* RadianceSortedGreenUAV,
+		FRHIUnorderedAccessView* RadianceSortedBlueUAV,
+		FRHIUnorderedAccessView* RadianceSortedAlphaUAV,
+		FRHIUnorderedAccessView* SampleCountSortedUAV)
 	{
 		FComputeShaderRHIParamRef ShaderRHI = GetComputeShader();
 		FGlobalShader::SetParameters<FViewUniformShaderParameters>(RHICmdList, ShaderRHI, View.ViewUniformBuffer);
@@ -84,11 +84,11 @@ public:
 		FRHICommandList& RHICmdList,
 		EResourceTransitionAccess TransitionAccess,
 		EResourceTransitionPipeline TransitionPipeline,
-		FUnorderedAccessViewRHIParamRef RadianceSortedRedUAV,
-		FUnorderedAccessViewRHIParamRef RadianceSortedGreenUAV,
-		FUnorderedAccessViewRHIParamRef RadianceSortedBlueUAV,
-		FUnorderedAccessViewRHIParamRef RadianceSortedAlphaUAV,
-		FUnorderedAccessViewRHIParamRef SampleCountSortedUAV,
+		FRHIUnorderedAccessView* RadianceSortedRedUAV,
+		FRHIUnorderedAccessView* RadianceSortedGreenUAV,
+		FRHIUnorderedAccessView* RadianceSortedBlueUAV,
+		FRHIUnorderedAccessView* RadianceSortedAlphaUAV,
+		FRHIUnorderedAccessView* SampleCountSortedUAV,
 		FRHIComputeFence* Fence)
 	{
 		FComputeShaderRHIParamRef ShaderRHI = GetComputeShader();
@@ -98,7 +98,7 @@ public:
 		SetUAVParameter(RHICmdList, ShaderRHI, RadianceSortedBlueUAVParameter, FUnorderedAccessViewRHIRef());
 		SetUAVParameter(RHICmdList, ShaderRHI, RadianceSortedAlphaUAVParameter, FUnorderedAccessViewRHIRef());
 		SetUAVParameter(RHICmdList, ShaderRHI, SampleCountSortedUAVParameter, FUnorderedAccessViewRHIRef());
-		FUnorderedAccessViewRHIParamRef UAVs[] = {
+		FRHIUnorderedAccessView* UAVs[] = {
 			RadianceSortedRedUAV,
 			RadianceSortedGreenUAV,
 			RadianceSortedBlueUAV,
@@ -146,11 +146,11 @@ void FDeferredShadingSceneRenderer::ComputePathCompaction(
 	FTextureRHIParamRef RadianceTexture,
 	FTextureRHIParamRef SampleCountTexture,
 	FTextureRHIParamRef PixelPositionTexture,
-	FUnorderedAccessViewRHIParamRef RadianceSortedRedUAV,
-	FUnorderedAccessViewRHIParamRef RadianceSortedGreenUAV,
-	FUnorderedAccessViewRHIParamRef RadianceSortedBlueUAV,
-	FUnorderedAccessViewRHIParamRef RadianceSortedAlphaUAV,
-	FUnorderedAccessViewRHIParamRef SampleCountSortedUAV)
+	FRHIUnorderedAccessView* RadianceSortedRedUAV,
+	FRHIUnorderedAccessView* RadianceSortedGreenUAV,
+	FRHIUnorderedAccessView* RadianceSortedBlueUAV,
+	FRHIUnorderedAccessView* RadianceSortedAlphaUAV,
+	FRHIUnorderedAccessView* SampleCountSortedUAV)
 {
 	const auto ShaderMap = GetGlobalShaderMap(FeatureLevel);
 	TShaderMapRef<FPathCompactionCS> PathCompactionComputeShader(ShaderMap);
diff --git a/Engine/Source/Runtime/Renderer/Private/PathTracing/PathTracingTestbed.cpp b/Engine/Source/Runtime/Renderer/Private/PathTracing/PathTracingTestbed.cpp
--- a/Engine/Source/Runtime/Renderer/Private/PathTracing/PathTracingTestbed.cpp
+++ b/Engine/Source/Runtime/Renderer/Private/PathTracing/PathTracingTestbed.cpp
@@ -51,8 +51,8 @@ public:
 		FRHICommandList& RHICmdList,
 		uint32 SamplesCount, 
 		uint32 BrdfType,
-		FUnorderedAccessViewRHIParamRef ResultsBuffer,
-		FUnorderedAccessViewRHIParamRef FloatResultsBuffer)
+		FRHIUnorderedAccessView* ResultsBuffer,
+		FRHIUnorderedAccessView* FloatResultsBuffer)
 	{
 		FComputeShaderRHIParamRef ShaderRHI = GetComputeShader();
 		SetShaderValue(RHICmdList, ShaderRHI, SamplesCountParameter, SamplesCount);
@@ -137,7 +137,7 @@ public:
 		FVector Wo,
 		uint32 NumThetaSteps,
 		uint32 NumPhiSteps,
-		FUnorderedAccessViewRHIParamRef PdfsResultsBuffer)
+		FRHIUnorderedAccessView* PdfsResultsBuffer)
 	{
 		FComputeShaderRHIParamRef ShaderRHI = GetComputeShader();
 		SetShaderValue(RHICmdList, ShaderRHI, BrdfTypeParameter, BrdfType);
@@ -225,7 +225,7 @@ public:
 		FVector Wo,
 		uint32 NumThetaSteps,
 		uint32 NumPhiSteps,
-		FUnorderedAccessViewRHIParamRef WisBinDistributionResultsBuffer)
+		FRHIUnorderedAccessView* WisBinDistributionResultsBuffer)
 	{
 		FComputeShaderRHIParamRef ShaderRHI = GetComputeShader();
 		SetShaderValue(RHICmdList, ShaderRHI, BrdfTypeParameter, BrdfType);
@@ -314,7 +314,7 @@ public:
 		FVector Wo,
 		uint32 NumThetaSteps,
 		uint32 NumPhiSteps,
-		FUnorderedAccessViewRHIParamRef PdfsResultsBuffer)
+		FRHIUnorderedAccessView* PdfsResultsBuffer)
 	{
 		FComputeShaderRHIParamRef ShaderRHI = GetComputeShader();
 		SetShaderValue(RHICmdList, ShaderRHI, BrdfTypeParameter, BrdfType);
diff --git a/Engine/Source/Runtime/Renderer/Private/PathTracing/RayCounterCompute.cpp b/Engine/Source/Runtime/Renderer/Private/PathTracing/RayCounterCompute.cpp
--- a/Engine/Source/Runtime/Renderer/Private/PathTracing/RayCounterCompute.cpp
+++ b/Engine/Source/Runtime/Renderer/Private/PathTracing/RayCounterCompute.cpp
@@ -51,7 +51,7 @@ public:
 		FRHICommandList& RHICmdList,
 		FTextureRHIParamRef RayCountPerPixelBuffer,
 		const FIntPoint& ViewSize,
-		FUnorderedAccessViewRHIParamRef TotalRayCountBuffer)
+		FRHIUnorderedAccessView* TotalRayCountBuffer)
 	{
 		FComputeShaderRHIParamRef ShaderRHI = GetComputeShader();
 
diff --git a/Engine/Source/Runtime/Renderer/Private/PathTracing/SkyLightCdfBuild.cpp b/Engine/Source/Runtime/Renderer/Private/PathTracing/SkyLightCdfBuild.cpp
--- a/Engine/Source/Runtime/Renderer/Private/PathTracing/SkyLightCdfBuild.cpp
+++ b/Engine/Source/Runtime/Renderer/Private/PathTracing/SkyLightCdfBuild.cpp
@@ -63,7 +63,7 @@ public:
 		bool IsBound = RowCdfParameter.IsUAVBound();
 		RowCdfParameter.SetBuffer(RHICmdList, ShaderRHI, RowCdf);
 
-		FUnorderedAccessViewRHIParamRef UAVs[] = {
+		FRHIUnorderedAccessView* UAVs[] = {
 			RowCdf.UAV
 		};
 		RHICmdList.TransitionResources(EResourceTransitionAccess::EWritable, EResourceTransitionPipeline::EGfxToCompute, UAVs, 1);
@@ -75,7 +75,7 @@ public:
 		FComputeShaderRHIParamRef ShaderRHI = GetComputeShader();
 		RowCdfParameter.UnsetUAV(RHICmdList, ShaderRHI);
 
-		FUnorderedAccessViewRHIParamRef UAVs[] = {
+		FRHIUnorderedAccessView* UAVs[] = {
 			Buffer.UAV
 		};
 		RHICmdList.TransitionResources(TransitionAccess, TransitionPipeline, UAVs, 1, Fence);
@@ -150,7 +150,7 @@ public:
 		SetShaderValue(RHICmdList, ShaderRHI, LevelParameter, Level);
 		ColumnCdfParameter.SetBuffer(RHICmdList, ShaderRHI, ColumnCdf);
 
-		FUnorderedAccessViewRHIParamRef UAVs[] = {
+		FRHIUnorderedAccessView* UAVs[] = {
 			ColumnCdf.UAV
 		};
 		RHICmdList.TransitionResources(EResourceTransitionAccess::EWritable, EResourceTransitionPipeline::EGfxToCompute, UAVs, 1);
@@ -162,7 +162,7 @@ public:
 		FComputeShaderRHIParamRef ShaderRHI = GetComputeShader();
 		ColumnCdfParameter.UnsetUAV(RHICmdList, ShaderRHI);
 
-		FUnorderedAccessViewRHIParamRef UAVs[] = {
+		FRHIUnorderedAccessView* UAVs[] = {
 			ColumnCdf.UAV
 		};
 		RHICmdList.TransitionResources(TransitionAccess, TransitionPipeline, UAVs, 1, Fence);
@@ -229,7 +229,7 @@ public:
 		ColumnCdfParameter.SetBuffer(RHICmdList, ShaderRHI, ColumnCdf);
 		CubeFaceCdfParameter.SetBuffer(RHICmdList, ShaderRHI, CubeFaceCdf);
 
-		FUnorderedAccessViewRHIParamRef UAVs[] = {
+		FRHIUnorderedAccessView* UAVs[] = {
 			CubeFaceCdf.UAV
 		};
 		RHICmdList.TransitionResources(EResourceTransitionAccess::EWritable, EResourceTransitionPipeline::EGfxToCompute, UAVs, 1);
@@ -241,7 +241,7 @@ public:
 		FComputeShaderRHIParamRef ShaderRHI = GetComputeShader();
 		CubeFaceCdfParameter.UnsetUAV(RHICmdList, ShaderRHI);
 
-		FUnorderedAccessViewRHIParamRef UAVs[] = {
+		FRHIUnorderedAccessView* UAVs[] = {
 			CubeFaceCdf.UAV
 		};
 		RHICmdList.TransitionResources(TransitionAccess, TransitionPipeline, UAVs, 1, Fence);
diff --git a/Engine/Source/Runtime/Renderer/Private/PostProcess/PostProcessBloomSetup.cpp b/Engine/Source/Runtime/Renderer/Private/PostProcess/PostProcessBloomSetup.cpp
--- a/Engine/Source/Runtime/Renderer/Private/PostProcess/PostProcessBloomSetup.cpp
+++ b/Engine/Source/Runtime/Renderer/Private/PostProcess/PostProcessBloomSetup.cpp
@@ -178,7 +178,7 @@ public:
 	}
 
 	template <typename TRHICmdList>
-	void SetParameters(TRHICmdList& RHICmdList, const FRenderingCompositePassContext& Context, const FIntPoint& DestSize, FUnorderedAccessViewRHIParamRef DestUAV, FTextureRHIParamRef EyeAdaptationTex)
+	void SetParameters(TRHICmdList& RHICmdList, const FRenderingCompositePassContext& Context, const FIntPoint& DestSize, FRHIUnorderedAccessView* DestUAV, FTextureRHIParamRef EyeAdaptationTex)
 	{
 		const FComputeShaderRHIParamRef ShaderRHI = GetComputeShader();
 		const FPostProcessSettings& Settings = Context.View.FinalPostProcessSettings;
@@ -329,7 +329,7 @@ void FRCPassPostProcessBloomSetup::Process(FRenderingCompositePassContext& Conte
 }
 
 template <typename TRHICmdList>
-void FRCPassPostProcessBloomSetup::DispatchCS(TRHICmdList& RHICmdList, FRenderingCompositePassContext& Context, const FIntRect& DestRect, FUnorderedAccessViewRHIParamRef DestUAV, FTextureRHIParamRef EyeAdaptationTex)
+void FRCPassPostProcessBloomSetup::DispatchCS(TRHICmdList& RHICmdList, FRenderingCompositePassContext& Context, const FIntRect& DestRect, FRHIUnorderedAccessView* DestUAV, FTextureRHIParamRef EyeAdaptationTex)
 {
 	auto ShaderMap = Context.GetShaderMap();
 	TShaderMapRef<FPostProcessBloomSetupCS> ComputeShader(ShaderMap);
diff --git a/Engine/Source/Runtime/Renderer/Private/PostProcess/PostProcessBloomSetup.h b/Engine/Source/Runtime/Renderer/Private/PostProcess/PostProcessBloomSetup.h
--- a/Engine/Source/Runtime/Renderer/Private/PostProcess/PostProcessBloomSetup.h
+++ b/Engine/Source/Runtime/Renderer/Private/PostProcess/PostProcessBloomSetup.h
@@ -31,7 +31,7 @@ public:
 
 private:
 	template <typename TRHICmdList>
-	void DispatchCS(TRHICmdList& RHICmdList, FRenderingCompositePassContext& Context, const FIntRect& DestRect, FUnorderedAccessViewRHIParamRef DestUAV, FTextureRHIParamRef EyeAdaptationTex);
+	void DispatchCS(TRHICmdList& RHICmdList, FRenderingCompositePassContext& Context, const FIntRect& DestRect, FRHIUnorderedAccessView* DestUAV, FTextureRHIParamRef EyeAdaptationTex);
 
 	FComputeFenceRHIRef AsyncEndFence;
 };
diff --git a/Engine/Source/Runtime/Renderer/Private/PostProcess/PostProcessCombineLUTs.cpp b/Engine/Source/Runtime/Renderer/Private/PostProcess/PostProcessCombineLUTs.cpp
--- a/Engine/Source/Runtime/Renderer/Private/PostProcess/PostProcessCombineLUTs.cpp
+++ b/Engine/Source/Runtime/Renderer/Private/PostProcess/PostProcessCombineLUTs.cpp
@@ -646,7 +646,7 @@ public:
 	}
 
 	template <typename TRHICmdList>
-	void SetParameters(TRHICmdList& RHICmdList, const FRenderingCompositePassContext& Context, const FIntPoint& DestSize, FUnorderedAccessViewRHIParamRef DestUAV, FTexture** Textures, float* Weights)
+	void SetParameters(TRHICmdList& RHICmdList, const FRenderingCompositePassContext& Context, const FIntPoint& DestSize, FRHIUnorderedAccessView* DestUAV, FTexture** Textures, float* Weights)
 	{
 		const FComputeShaderRHIParamRef ShaderRHI = GetComputeShader();
 
@@ -971,7 +971,7 @@ void FRCPassPostProcessCombineLUTs::Process(FRenderingCompositePassContext& Cont
 }
 
 template <typename TRHICmdList>
-void FRCPassPostProcessCombineLUTs::DispatchCS(TRHICmdList& RHICmdList, FRenderingCompositePassContext& Context, const FIntRect& DestRect, FUnorderedAccessViewRHIParamRef DestUAV, int32 BlendCount, FTexture* Textures[], float Weights[])
+void FRCPassPostProcessCombineLUTs::DispatchCS(TRHICmdList& RHICmdList, FRenderingCompositePassContext& Context, const FIntRect& DestRect, FRHIUnorderedAccessView* DestUAV, int32 BlendCount, FTexture* Textures[], float Weights[])
 {
 	auto ShaderMap = Context.GetShaderMap();
 
diff --git a/Engine/Source/Runtime/Renderer/Private/PostProcess/PostProcessCombineLUTs.h b/Engine/Source/Runtime/Renderer/Private/PostProcess/PostProcessCombineLUTs.h
--- a/Engine/Source/Runtime/Renderer/Private/PostProcess/PostProcessCombineLUTs.h
+++ b/Engine/Source/Runtime/Renderer/Private/PostProcess/PostProcessCombineLUTs.h
@@ -43,7 +43,7 @@ public:
 
 private:
 	template <typename TRHICmdList>
-	void DispatchCS(TRHICmdList& RHICmdList, FRenderingCompositePassContext& Context, const FIntRect& DestRect, FUnorderedAccessViewRHIParamRef DestUAV, int32 BlendCount, FTexture* Textures[], float Weights[]);
+	void DispatchCS(TRHICmdList& RHICmdList, FRenderingCompositePassContext& Context, const FIntRect& DestRect, FRHIUnorderedAccessView* DestUAV, int32 BlendCount, FTexture* Textures[], float Weights[]);
 
 	FComputeFenceRHIRef AsyncEndFence;
 
diff --git a/Engine/Source/Runtime/Renderer/Private/PostProcess/PostProcessDownsample.cpp b/Engine/Source/Runtime/Renderer/Private/PostProcess/PostProcessDownsample.cpp
--- a/Engine/Source/Runtime/Renderer/Private/PostProcess/PostProcessDownsample.cpp
+++ b/Engine/Source/Runtime/Renderer/Private/PostProcess/PostProcessDownsample.cpp
@@ -207,7 +207,7 @@ public:
 	}
 
 	template <typename TRHICmdList>
-	void SetParameters(TRHICmdList& RHICmdList, const FRenderingCompositePassContext& Context, const FIntPoint& SrcSize, FUnorderedAccessViewRHIParamRef DestUAV)
+	void SetParameters(TRHICmdList& RHICmdList, const FRenderingCompositePassContext& Context, const FIntPoint& SrcSize, FRHIUnorderedAccessView* DestUAV)
 	{
 		const FComputeShaderRHIParamRef ShaderRHI = GetComputeShader();
 		FGlobalShader::SetParameters<FViewUniformShaderParameters>(RHICmdList, ShaderRHI, Context.View.ViewUniformBuffer);
@@ -434,7 +434,7 @@ void FRCPassPostProcessDownsample::Process(FRenderingCompositePassContext& Conte
 }
 
 template <uint32 Method, typename TRHICmdList>
-void FRCPassPostProcessDownsample::DispatchCS(TRHICmdList& RHICmdList, FRenderingCompositePassContext& Context, const FIntPoint& SrcSize, const FIntRect& DestRect, FUnorderedAccessViewRHIParamRef DestUAV)
+void FRCPassPostProcessDownsample::DispatchCS(TRHICmdList& RHICmdList, FRenderingCompositePassContext& Context, const FIntPoint& SrcSize, const FIntRect& DestRect, FRHIUnorderedAccessView* DestUAV)
 {
 	auto ShaderMap = Context.GetShaderMap();
 	TShaderMapRef<FPostProcessDownsampleCS<Method>> ComputeShader(ShaderMap);
diff --git a/Engine/Source/Runtime/Renderer/Private/PostProcess/PostProcessDownsample.h b/Engine/Source/Runtime/Renderer/Private/PostProcess/PostProcessDownsample.h
--- a/Engine/Source/Runtime/Renderer/Private/PostProcess/PostProcessDownsample.h
+++ b/Engine/Source/Runtime/Renderer/Private/PostProcess/PostProcessDownsample.h
@@ -37,7 +37,7 @@ private:
 	void SetShader(const FRenderingCompositePassContext& Context, const FPooledRenderTargetDesc* InputDesc, const FIntPoint& SrcSize, const FIntRect& SrcRect);
 
 	template <uint32 Method, typename TRHICmdList>
-	void DispatchCS(TRHICmdList& RHICmdList, FRenderingCompositePassContext& Context, const FIntPoint& SrcSize, const FIntRect& DestRect, FUnorderedAccessViewRHIParamRef DestUAV);
+	void DispatchCS(TRHICmdList& RHICmdList, FRenderingCompositePassContext& Context, const FIntPoint& SrcSize, const FIntRect& DestRect, FRHIUnorderedAccessView* DestUAV);
 
 	FComputeFenceRHIRef AsyncEndFence;
 
diff --git a/Engine/Source/Runtime/Renderer/Private/PostProcess/PostProcessEyeAdaptation.cpp b/Engine/Source/Runtime/Renderer/Private/PostProcess/PostProcessEyeAdaptation.cpp
--- a/Engine/Source/Runtime/Renderer/Private/PostProcess/PostProcessEyeAdaptation.cpp
+++ b/Engine/Source/Runtime/Renderer/Private/PostProcess/PostProcessEyeAdaptation.cpp
@@ -271,7 +271,7 @@ public:
 	}
 
 	template <typename TRHICmdList>
-	void SetParameters(TRHICmdList& RHICmdList, const FRenderingCompositePassContext& Context, FUnorderedAccessViewRHIParamRef DestUAV, IPooledRenderTarget* LastEyeAdaptation)
+	void SetParameters(TRHICmdList& RHICmdList, const FRenderingCompositePassContext& Context, FRHIUnorderedAccessView* DestUAV, IPooledRenderTarget* LastEyeAdaptation)
 	{
 		const FComputeShaderRHIParamRef ShaderRHI = GetComputeShader();
 
@@ -416,7 +416,7 @@ void FRCPassPostProcessEyeAdaptation::Process(FRenderingCompositePassContext& Co
 }
 
 template <typename TRHICmdList>
-void FRCPassPostProcessEyeAdaptation::DispatchCS(TRHICmdList& RHICmdList, FRenderingCompositePassContext& Context, FUnorderedAccessViewRHIParamRef DestUAV, IPooledRenderTarget* LastEyeAdaptation)
+void FRCPassPostProcessEyeAdaptation::DispatchCS(TRHICmdList& RHICmdList, FRenderingCompositePassContext& Context, FRHIUnorderedAccessView* DestUAV, IPooledRenderTarget* LastEyeAdaptation)
 {
 	auto ShaderMap = Context.GetShaderMap();
 	TShaderMapRef<FPostProcessEyeAdaptationCS> ComputeShader(ShaderMap);
@@ -807,7 +807,7 @@ public:
 		OutComputeTex.UnsetUAV(Context.RHICmdList, ShaderRHI);
 	}
 
-	void SetParameters(const FRenderingCompositePassContext& Context, const FIntRect& SrcRect, IPooledRenderTarget* LastEyeAdaptation, FUnorderedAccessViewRHIParamRef DestUAV)
+	void SetParameters(const FRenderingCompositePassContext& Context, const FIntRect& SrcRect, IPooledRenderTarget* LastEyeAdaptation, FRHIUnorderedAccessView* DestUAV)
 	{
 		const FComputeShaderRHIParamRef ShaderRHI = GetComputeShader();
 
diff --git a/Engine/Source/Runtime/Renderer/Private/PostProcess/PostProcessEyeAdaptation.h b/Engine/Source/Runtime/Renderer/Private/PostProcess/PostProcessEyeAdaptation.h
--- a/Engine/Source/Runtime/Renderer/Private/PostProcess/PostProcessEyeAdaptation.h
+++ b/Engine/Source/Runtime/Renderer/Private/PostProcess/PostProcessEyeAdaptation.h
@@ -63,7 +63,7 @@ public:
 
 private:
 	template <typename TRHICmdList>
-	void DispatchCS(TRHICmdList& RHICmdList, FRenderingCompositePassContext& Context, FUnorderedAccessViewRHIParamRef DestUAV, IPooledRenderTarget* LastEyeAdaptation);
+	void DispatchCS(TRHICmdList& RHICmdList, FRenderingCompositePassContext& Context, FRHIUnorderedAccessView* DestUAV, IPooledRenderTarget* LastEyeAdaptation);
 
 	FComputeFenceRHIRef AsyncEndFence;
 };
diff --git a/Engine/Source/Runtime/Renderer/Private/PostProcess/PostProcessTonemap.cpp b/Engine/Source/Runtime/Renderer/Private/PostProcess/PostProcessTonemap.cpp
--- a/Engine/Source/Runtime/Renderer/Private/PostProcess/PostProcessTonemap.cpp
+++ b/Engine/Source/Runtime/Renderer/Private/PostProcess/PostProcessTonemap.cpp
@@ -830,7 +830,7 @@ public:
 	}
 
 	template <typename TRHICmdList>
-	void SetParameters(TRHICmdList& RHICmdList, const FRenderingCompositePassContext& Context, const FIntPoint& DestSize, FUnorderedAccessViewRHIParamRef DestUAV, FTextureRHIParamRef EyeAdaptationTex)
+	void SetParameters(TRHICmdList& RHICmdList, const FRenderingCompositePassContext& Context, const FIntPoint& DestSize, FRHIUnorderedAccessView* DestUAV, FTextureRHIParamRef EyeAdaptationTex)
 	{
 		const FComputeShaderRHIParamRef ShaderRHI = GetComputeShader();
 		const FPostProcessSettings& Settings = Context.View.FinalPostProcessSettings;
@@ -934,7 +934,7 @@ static inline void ShaderTransitionResources(const FRenderingCompositePassContex
 } // PostProcessTonemapUtil
 
 template <typename TRHICmdList>
-inline void DispatchComputeShader(TRHICmdList& RHICmdList, FRenderingCompositePassContext& Context, const FIntRect& DestRect, FUnorderedAccessViewRHIParamRef DestUAV, const TonemapperPermutation::FDesktopDomain& DesktopPermutationVector, FTextureRHIParamRef EyeAdaptationTex, bool bDoEyeAdaptation)
+inline void DispatchComputeShader(TRHICmdList& RHICmdList, FRenderingCompositePassContext& Context, const FIntRect& DestRect, FRHIUnorderedAccessView* DestUAV, const TonemapperPermutation::FDesktopDomain& DesktopPermutationVector, FTextureRHIParamRef EyeAdaptationTex, bool bDoEyeAdaptation)
 {
 	FPostProcessTonemapCS::FPermutationDomain PermutationVector;
 	PermutationVector.Set<TonemapperPermutation::FDesktopDomain>(DesktopPermutationVector);
diff --git a/Engine/Source/Runtime/Renderer/Private/PostProcess/PostProcessWeightedSampleSum.cpp b/Engine/Source/Runtime/Renderer/Private/PostProcess/PostProcessWeightedSampleSum.cpp
--- a/Engine/Source/Runtime/Renderer/Private/PostProcess/PostProcessWeightedSampleSum.cpp
+++ b/Engine/Source/Runtime/Renderer/Private/PostProcess/PostProcessWeightedSampleSum.cpp
@@ -399,7 +399,7 @@ public:
 	}
 
 	template <typename TRHICmdList>
-	void SetParameters(TRHICmdList& RHICmdList, const FRenderingCompositePassContext& Context, const FIntPoint& DestSize, FUnorderedAccessViewRHIParamRef DestUAV,
+	void SetParameters(TRHICmdList& RHICmdList, const FRenderingCompositePassContext& Context, const FIntPoint& DestSize, FRHIUnorderedAccessView* DestUAV,
 		FTextureRHIParamRef FilterTextureRHI, FTextureRHIParamRef AdditiveTextureRHI, FLinearColor* SampleWeightValues, FVector2D* SampleOffsetValues, uint32 NumSamples)
 	{
 		const FComputeShaderRHIParamRef ShaderRHI = GetComputeShader();
@@ -1082,7 +1082,7 @@ void FRCPassPostProcessWeightedSampleSum::Process(FRenderingCompositePassContext
 }
 
 template <uint32 CompileTimeNumSamples, uint32 CombineMethodInt, typename TRHICmdList>
-void DispatchCSTemplate(TRHICmdList& RHICmdList, FRenderingCompositePassContext& Context, const FIntRect& DestRect, FUnorderedAccessViewRHIParamRef DestUAV,
+void DispatchCSTemplate(TRHICmdList& RHICmdList, FRenderingCompositePassContext& Context, const FIntRect& DestRect, FRHIUnorderedAccessView* DestUAV,
 	FTextureRHIParamRef FilterTextureRHI, FTextureRHIParamRef AdditiveTextureRHI, FLinearColor* SampleWeightValues, FVector2D* SampleOffsetValues, uint32 NumSamples)
 {
 	auto ShaderMap = Context.GetShaderMap();
@@ -1102,7 +1102,7 @@ void DispatchCSTemplate(TRHICmdList& RHICmdList, FRenderingCompositePassContext&
 }
 
 template <typename TRHICmdList>
-void FRCPassPostProcessWeightedSampleSum::DispatchCS(TRHICmdList& RHICmdList, FRenderingCompositePassContext& Context, const FIntRect& DestRect, FUnorderedAccessViewRHIParamRef DestUAV,
+void FRCPassPostProcessWeightedSampleSum::DispatchCS(TRHICmdList& RHICmdList, FRenderingCompositePassContext& Context, const FIntRect& DestRect, FRHIUnorderedAccessView* DestUAV,
 	FTextureRHIParamRef FilterTextureRHI, FTextureRHIParamRef AdditiveTextureRHI, uint32 CombineMethodInt, FLinearColor* SampleWeightValues, FVector2D* SampleOffsetValues, uint32 NumSamples)
 {
 	// Helper macros
diff --git a/Engine/Source/Runtime/Renderer/Private/PostProcess/PostProcessWeightedSampleSum.h b/Engine/Source/Runtime/Renderer/Private/PostProcess/PostProcessWeightedSampleSum.h
--- a/Engine/Source/Runtime/Renderer/Private/PostProcess/PostProcessWeightedSampleSum.h
+++ b/Engine/Source/Runtime/Renderer/Private/PostProcess/PostProcessWeightedSampleSum.h
@@ -74,7 +74,7 @@ private:
 	bool DoFastBlur() const;
 
 	template <typename TRHICmdList>
-	void DispatchCS(TRHICmdList& RHICmdList, FRenderingCompositePassContext& Context, const FIntRect& DestRect, FUnorderedAccessViewRHIParamRef DestUAV,
+	void DispatchCS(TRHICmdList& RHICmdList, FRenderingCompositePassContext& Context, const FIntRect& DestRect, FRHIUnorderedAccessView* DestUAV,
 		FTextureRHIParamRef FilterTextureRHI, FTextureRHIParamRef AdditiveTextureRHI, uint32 CombineMethodInt, FLinearColor* SampleWeightValues, FVector2D* SampleOffsetValues, uint32 NumSamples);
 
 	FComputeFenceRHIRef AsyncEndFence;
diff --git a/Engine/Source/Runtime/Renderer/Private/PostProcess/SceneRenderTargets.cpp b/Engine/Source/Runtime/Renderer/Private/PostProcess/SceneRenderTargets.cpp
--- a/Engine/Source/Runtime/Renderer/Private/PostProcess/SceneRenderTargets.cpp
+++ b/Engine/Source/Runtime/Renderer/Private/PostProcess/SceneRenderTargets.cpp
@@ -3055,7 +3055,7 @@ void SetupSceneTextureUniformParameters(
 		const bool bSetupCustomDepth = (SetupMode & ESceneTextureSetupMode::CustomDepth) != ESceneTextureSetupMode::None;
 
 		FTextureRHIParamRef CustomDepth = DepthDefault;
-		FShaderResourceViewRHIParamRef CustomStencilSRV = GSystemTextures.WhiteDummySRV;
+		FRHIShaderResourceView* CustomStencilSRV = GSystemTextures.WhiteDummySRV;
 
 		// if there is no custom depth it's better to have the far distance there
 		IPooledRenderTarget* CustomDepthTarget = SceneContext.bCustomDepthIsValid ? SceneContext.CustomDepth.GetReference() : 0;
diff --git a/Engine/Source/Runtime/Renderer/Private/RayTracing/RayTracingRectLight.cpp b/Engine/Source/Runtime/Renderer/Private/RayTracing/RayTracingRectLight.cpp
--- a/Engine/Source/Runtime/Renderer/Private/RayTracing/RayTracingRectLight.cpp
+++ b/Engine/Source/Runtime/Renderer/Private/RayTracing/RayTracingRectLight.cpp
@@ -262,8 +262,8 @@ public:
 		FUniformBufferRHIParamRef ViewUniformBuffer,
 		FUniformBufferRHIParamRef SceneTexturesUniformBuffer,
 		FUniformBufferRHIParamRef RectLightUniformBuffer,
-		FUnorderedAccessViewRHIParamRef LuminanceUAV,
-		FUnorderedAccessViewRHIParamRef RayDistanceUAV,
+		FRHIUnorderedAccessView* LuminanceUAV,
+		FRHIUnorderedAccessView* RayDistanceUAV,
 		uint32 Width, uint32 Height
 	)
 	{
diff --git a/Engine/Source/Runtime/Renderer/Private/ShadowDepthRendering.cpp b/Engine/Source/Runtime/Renderer/Private/ShadowDepthRendering.cpp
--- a/Engine/Source/Runtime/Renderer/Private/ShadowDepthRendering.cpp
+++ b/Engine/Source/Runtime/Renderer/Private/ShadowDepthRendering.cpp
@@ -1666,7 +1666,7 @@ void FSceneRenderer::RenderShadowDepthMaps(FRHICommandListImmediate& RHICmdList)
 				RenderTargets[1] = ColorTarget1.TargetableTexture;
 
 				// Hook up the geometry volume UAVs
-				FUnorderedAccessViewRHIParamRef Uavs[4];
+				FRHIUnorderedAccessView* Uavs[4];
 				Uavs[0] = LightPropagationVolume->GetGvListBufferUav();
 				Uavs[1] = LightPropagationVolume->GetGvListHeadBufferUav();
 				Uavs[2] = LightPropagationVolume->GetVplListBufferUav();
@@ -1714,7 +1714,7 @@ void FSceneRenderer::RenderShadowDepthMaps(FRHICommandListImmediate& RHICmdList)
 				RHICmdList.CopyToResolveTarget(ColorTarget0.TargetableTexture, ColorTarget0.ShaderResourceTexture, FResolveParams());
 				RHICmdList.CopyToResolveTarget(ColorTarget1.TargetableTexture, ColorTarget1.ShaderResourceTexture, FResolveParams());
 
-				FUnorderedAccessViewRHIParamRef UavsToReadable[2];
+				FRHIUnorderedAccessView* UavsToReadable[2];
 				UavsToReadable[0] = LightPropagationVolume->GetGvListBufferUav();
 				UavsToReadable[1] = LightPropagationVolume->GetGvListHeadBufferUav();
 				RHICmdList.TransitionResources(EResourceTransitionAccess::EReadable, EResourceTransitionPipeline::EGfxToGfx, UavsToReadable, ARRAY_COUNT(UavsToReadable));
diff --git a/Engine/Source/Runtime/Renderer/Private/SurfelTree.cpp b/Engine/Source/Runtime/Renderer/Private/SurfelTree.cpp
--- a/Engine/Source/Runtime/Renderer/Private/SurfelTree.cpp
+++ b/Engine/Source/Runtime/Renderer/Private/SurfelTree.cpp
@@ -225,7 +225,7 @@ public:
 
 		const FScene* Scene = (const FScene*)View.Family->Scene;
 
-		FUnorderedAccessViewRHIParamRef UniformMeshUAVs[1];
+		FRHIUnorderedAccessView* UniformMeshUAVs[1];
 		UniformMeshUAVs[0] = Scene->DistanceFieldSceneData.SurfelBuffers->InterpolatedVertexData.UAV;
 		RHICmdList.TransitionResources(EResourceTransitionAccess::ERWBarrier, EResourceTransitionPipeline::EComputeToCompute, UniformMeshUAVs, ARRAY_COUNT(UniformMeshUAVs));
 
@@ -248,7 +248,7 @@ public:
 		// RHISetStreamOutTargets doesn't unbind existing uses like render targets do 
 		SetSRVParameter(RHICmdList, ShaderRHI, TriangleVertexData, NULL);
 
-		FUnorderedAccessViewRHIParamRef UniformMeshUAVs[1];
+		FRHIUnorderedAccessView* UniformMeshUAVs[1];
 		UniformMeshUAVs[0] = Scene->DistanceFieldSceneData.SurfelBuffers->InterpolatedVertexData.UAV;
 		RHICmdList.TransitionResources(EResourceTransitionAccess::EReadable, EResourceTransitionPipeline::EComputeToCompute, UniformMeshUAVs, ARRAY_COUNT(UniformMeshUAVs));
 	}
@@ -302,7 +302,7 @@ void GenerateSurfelRepresentation(FRHICommandListImmediate& RHICmdList, FSceneRe
 
 		if (NumUniformTriangles > 0 && MaterialRenderProxy && PrimitiveUniformBuffer)
 		{
-			FUnorderedAccessViewRHIParamRef UniformMeshUAVs[2];
+			FRHIUnorderedAccessView* UniformMeshUAVs[2];
 			UniformMeshUAVs[0] = UniformMeshBuffers->TriangleAreas.UAV;
 			UniformMeshUAVs[1] = UniformMeshBuffers->TriangleCDFs.UAV;
 			RHICmdList.TransitionResources(EResourceTransitionAccess::ERWBarrier, EResourceTransitionPipeline::EComputeToCompute, UniformMeshUAVs, ARRAY_COUNT(UniformMeshUAVs));
diff --git a/Engine/Source/Runtime/Renderer/Private/TiledDeferredLightRendering.cpp b/Engine/Source/Runtime/Renderer/Private/TiledDeferredLightRendering.cpp
--- a/Engine/Source/Runtime/Renderer/Private/TiledDeferredLightRendering.cpp
+++ b/Engine/Source/Runtime/Renderer/Private/TiledDeferredLightRendering.cpp
@@ -139,7 +139,7 @@ public:
 		SceneTextureParameters.Set(RHICmdList, ShaderRHI, View.FeatureLevel, ESceneTextureSetupMode::All);
 		SetTextureParameter(RHICmdList, ShaderRHI, InTexture, InTextureValue.GetRenderTargetItem().ShaderResourceTexture);
 
-		FUnorderedAccessViewRHIParamRef OutUAV = OutTextureValue.GetRenderTargetItem().UAV;
+		FRHIUnorderedAccessView* OutUAV = OutTextureValue.GetRenderTargetItem().UAV;
 		RHICmdList.TransitionResources(EResourceTransitionAccess::ERWBarrier, EResourceTransitionPipeline::EGfxToCompute, &OutUAV, 1);
 		OutTexture.SetTexture(RHICmdList, ShaderRHI, 0, OutUAV);
 
@@ -269,7 +269,7 @@ public:
 	{
 		OutTexture.UnsetUAV(RHICmdList, GetComputeShader());
 
-		FUnorderedAccessViewRHIParamRef OutUAV = OutTextureValue.GetRenderTargetItem().UAV;
+		FRHIUnorderedAccessView* OutUAV = OutTextureValue.GetRenderTargetItem().UAV;
 		RHICmdList.TransitionResources(EResourceTransitionAccess::EReadable, EResourceTransitionPipeline::EComputeToGfx, &OutUAV, 1);
 	}
 
diff --git a/Engine/Source/Runtime/Renderer/Private/TranslucentLighting.cpp b/Engine/Source/Runtime/Renderer/Private/TranslucentLighting.cpp
--- a/Engine/Source/Runtime/Renderer/Private/TranslucentLighting.cpp
+++ b/Engine/Source/Runtime/Renderer/Private/TranslucentLighting.cpp
@@ -957,7 +957,7 @@ public:
 
 	void SetParameters(
 		FRHIAsyncComputeCommandListImmediate& RHICmdList,
-		FUnorderedAccessViewRHIParamRef* VolumeUAVs,
+		FRHIUnorderedAccessView** VolumeUAVs,
 		int32 NumUAVs
 	)
 	{
@@ -1005,7 +1005,7 @@ void FDeferredShadingSceneRenderer::ClearTranslucentVolumeLightingAsyncCompute(F
 
 	for(int i = 0; i < Views.Num(); ++i)
 	{
-		FUnorderedAccessViewRHIParamRef VolumeUAVs[4] = {
+		FRHIUnorderedAccessView* VolumeUAVs[4] = {
 			SceneContext.TranslucencyLightingVolumeAmbient[(i * NumTranslucentVolumeRenderTargetSets)]->GetRenderTargetItem().UAV,
 			SceneContext.TranslucencyLightingVolumeDirectional[(i * NumTranslucentVolumeRenderTargetSets)]->GetRenderTargetItem().UAV,
 			SceneContext.TranslucencyLightingVolumeAmbient[(i * NumTranslucentVolumeRenderTargetSets) + 1]->GetRenderTargetItem().UAV,
diff --git a/Engine/Source/Runtime/Renderer/Public/MeshDrawShaderBindings.h b/Engine/Source/Runtime/Renderer/Public/MeshDrawShaderBindings.h
--- a/Engine/Source/Runtime/Renderer/Public/MeshDrawShaderBindings.h
+++ b/Engine/Source/Runtime/Renderer/Public/MeshDrawShaderBindings.h
@@ -45,7 +45,7 @@ public:
 			+ ParameterMapInfo.TextureSamplers.Num() 
 			+ ParameterMapInfo.SRVs.Num());
 
-		// Allocate a bit for each SRV tracking whether it is a FTextureRHIParamRef or FShaderResourceViewRHIParamRef
+		// Allocate a bit for each SRV tracking whether it is a FTextureRHIParamRef or FRHIShaderResourceView*
 		DataSize += FMath::DivideAndRoundUp(ParameterMapInfo.SRVs.Num(), 8);
 
 		for (int32 LooseBufferIndex = 0; LooseBufferIndex < ParameterMapInfo.LooseParameterBuffers.Num(); LooseBufferIndex++)
@@ -79,7 +79,7 @@ protected:
 	inline uint32 GetSRVTypeOffset() const
 	{
 		return GetSRVOffset() 
-			+ ParameterMapInfo.SRVs.Num() * sizeof(FShaderResourceViewRHIParamRef);
+			+ ParameterMapInfo.SRVs.Num() * sizeof(FRHIShaderResourceView*);
 	}
 
 	inline uint32 GetLooseDataOffset() const
@@ -138,7 +138,7 @@ public:
 		}
 	}
 
-	void Add(FShaderResourceParameter Parameter, FShaderResourceViewRHIParamRef Value)
+	void Add(FShaderResourceParameter Parameter, FRHIShaderResourceView* Value)
 	{
 		checkfSlow(Parameter.IsInitialized(), TEXT("Parameter was not serialized"));
 
@@ -303,7 +303,7 @@ private:
 		checkfSlow(FoundIndex >= 0, TEXT("Attempted to set a texture sampler at BaseIndex %u which was never in the shader's parameter map."), BaseIndex);
 	}
 
-	inline void WriteBindingSRV(FShaderResourceViewRHIParamRef Value, uint32 BaseIndex)
+	inline void WriteBindingSRV(FRHIShaderResourceView* Value, uint32 BaseIndex)
 	{
 		int32 FoundIndex = -1;
 
diff --git a/Engine/Source/Runtime/UtilityShaders/Private/ClearQuad.cpp b/Engine/Source/Runtime/UtilityShaders/Private/ClearQuad.cpp
--- a/Engine/Source/Runtime/UtilityShaders/Private/ClearQuad.cpp
+++ b/Engine/Source/Runtime/UtilityShaders/Private/ClearQuad.cpp
@@ -125,7 +125,7 @@ static void ClearQuadSetup( FRHICommandList& RHICmdList, bool bClearColor, int32
 	PixelShader->SetColors(RHICmdList, ClearColorArray, NumClearColors);
 }
 
-static void ClearUAVShader(FRHICommandList& RHICmdList, FUnorderedAccessViewRHIParamRef UnorderedAccessViewRHI, uint32 SizeInBytes, uint32 ClearValue)
+static void ClearUAVShader(FRHICommandList& RHICmdList, FRHIUnorderedAccessView* UnorderedAccessViewRHI, uint32 SizeInBytes, uint32 ClearValue)
 {
 	UE_CLOG((SizeInBytes & 0x3) != 0, LogClearQuad, Warning,
 		TEXT("Buffer size is not a multiple of DWORDs. Up to 3 bytes after buffer end will also be cleared"));
@@ -239,7 +239,7 @@ inline void ClearUAV_T(FRHICommandList& RHICmdList, FRHITexture* Texture, FRHIUn
 	}
 }
 
-void ClearTexture2DUAV(FRHICommandList& RHICmdList, FUnorderedAccessViewRHIParamRef UAV, int32 Width, int32 Height, const FLinearColor& ClearColor)
+void ClearTexture2DUAV(FRHICommandList& RHICmdList, FRHIUnorderedAccessView* UAV, int32 Width, int32 Height, const FLinearColor& ClearColor)
 {
 	TShaderMapRef< FClearTexture2DReplacementCS<float> > ComputeShader(GetGlobalShaderMap(GMaxRHIFeatureLevel));
 	FComputeShaderRHIParamRef ShaderRHI = ComputeShader->GetComputeShader();
diff --git a/Engine/Source/Runtime/UtilityShaders/Private/ClearReplacementShaders.cpp b/Engine/Source/Runtime/UtilityShaders/Private/ClearReplacementShaders.cpp
--- a/Engine/Source/Runtime/UtilityShaders/Private/ClearReplacementShaders.cpp
+++ b/Engine/Source/Runtime/UtilityShaders/Private/ClearReplacementShaders.cpp
@@ -6,7 +6,7 @@
 #include "RendererInterface.h"
 
 template< typename T >
-void FClearTexture2DReplacementCS<T>::SetParameters( FRHICommandList& RHICmdList, FUnorderedAccessViewRHIParamRef TextureRW, const T(&Values)[4] )
+void FClearTexture2DReplacementCS<T>::SetParameters( FRHICommandList& RHICmdList, FRHIUnorderedAccessView* TextureRW, const T(&Values)[4] )
 {
 	FComputeShaderRHIParamRef ComputeShaderRHI = GetComputeShader();
 	SetShaderValue(RHICmdList, ComputeShaderRHI, ClearColor, Values);
@@ -16,7 +16,7 @@ void FClearTexture2DReplacementCS<T>::SetParameters( FRHICommandList& RHICmdList
 }
 
 template< typename T >
-void FClearTexture2DReplacementCS<T>::FinalizeParameters(FRHICommandList& RHICmdList, FUnorderedAccessViewRHIParamRef TextureRW)
+void FClearTexture2DReplacementCS<T>::FinalizeParameters(FRHICommandList& RHICmdList, FRHIUnorderedAccessView* TextureRW)
 {
 	RHICmdList.TransitionResource(EResourceTransitionAccess::ERWBarrier, EResourceTransitionPipeline::EComputeToCompute, TextureRW);
 	FComputeShaderRHIParamRef ComputeShaderRHI = GetComputeShader();
@@ -24,7 +24,7 @@ void FClearTexture2DReplacementCS<T>::FinalizeParameters(FRHICommandList& RHICmd
 }
 
 template< typename T >
-void FClearTexture2DArrayReplacementCS<T>::SetParameters(FRHICommandList& RHICmdList, FUnorderedAccessViewRHIParamRef TextureRW, const T(&Values)[4])
+void FClearTexture2DArrayReplacementCS<T>::SetParameters(FRHICommandList& RHICmdList, FRHIUnorderedAccessView* TextureRW, const T(&Values)[4])
 {
 	FComputeShaderRHIParamRef ComputeShaderRHI = GetComputeShader();
 	SetShaderValue(RHICmdList, ComputeShaderRHI, ClearColor, Values);
@@ -34,7 +34,7 @@ void FClearTexture2DArrayReplacementCS<T>::SetParameters(FRHICommandList& RHICmd
 }
 
 template< typename T >
-void FClearTexture2DArrayReplacementCS<T>::FinalizeParameters(FRHICommandList& RHICmdList, FUnorderedAccessViewRHIParamRef TextureRW)
+void FClearTexture2DArrayReplacementCS<T>::FinalizeParameters(FRHICommandList& RHICmdList, FRHIUnorderedAccessView* TextureRW)
 {
 	RHICmdList.TransitionResource(EResourceTransitionAccess::ERWBarrier, EResourceTransitionPipeline::EComputeToCompute, TextureRW);
 	FComputeShaderRHIParamRef ComputeShaderRHI = GetComputeShader();
@@ -42,7 +42,7 @@ void FClearTexture2DArrayReplacementCS<T>::FinalizeParameters(FRHICommandList& R
 }
 
 template< typename T >
-void FClearVolumeReplacementCS<T>::SetParameters(FRHICommandList& RHICmdList, FUnorderedAccessViewRHIParamRef TextureRW, const T(&Values)[4])
+void FClearVolumeReplacementCS<T>::SetParameters(FRHICommandList& RHICmdList, FRHIUnorderedAccessView* TextureRW, const T(&Values)[4])
 {
 	FComputeShaderRHIParamRef ComputeShaderRHI = GetComputeShader();
 	SetShaderValue(RHICmdList, ComputeShaderRHI, ClearColor, Values);
@@ -52,14 +52,14 @@ void FClearVolumeReplacementCS<T>::SetParameters(FRHICommandList& RHICmdList, FU
 }
 
 template< typename T >
-void FClearVolumeReplacementCS<T>::FinalizeParameters(FRHICommandList& RHICmdList, FUnorderedAccessViewRHIParamRef TextureRW)
+void FClearVolumeReplacementCS<T>::FinalizeParameters(FRHICommandList& RHICmdList, FRHIUnorderedAccessView* TextureRW)
 {
 	RHICmdList.TransitionResource(EResourceTransitionAccess::ERWBarrier, EResourceTransitionPipeline::EComputeToCompute, TextureRW);
 	FComputeShaderRHIParamRef ComputeShaderRHI = GetComputeShader();
 	SetUAVParameter(RHICmdList, ComputeShaderRHI, ClearVolumeRW, FUnorderedAccessViewRHIRef());
 }
 
-void FClearTexture2DReplacementScissorCS::SetParameters(FRHICommandList& RHICmdList, FUnorderedAccessViewRHIParamRef TextureRW, FLinearColor InClearColor, const FVector4& InTargetBounds)
+void FClearTexture2DReplacementScissorCS::SetParameters(FRHICommandList& RHICmdList, FRHIUnorderedAccessView* TextureRW, FLinearColor InClearColor, const FVector4& InTargetBounds)
 {
 	FComputeShaderRHIParamRef ComputeShaderRHI = GetComputeShader();
 	SetShaderValue(RHICmdList, ComputeShaderRHI, ClearColor, InClearColor);
@@ -68,14 +68,14 @@ void FClearTexture2DReplacementScissorCS::SetParameters(FRHICommandList& RHICmdL
 	RHICmdList.SetUAVParameter(ComputeShaderRHI, ClearTextureRW.GetBaseIndex(), TextureRW);
 }
 
-void FClearTexture2DReplacementScissorCS::FinalizeParameters(FRHICommandList& RHICmdList, FUnorderedAccessViewRHIParamRef TextureRW)
+void FClearTexture2DReplacementScissorCS::FinalizeParameters(FRHICommandList& RHICmdList, FRHIUnorderedAccessView* TextureRW)
 {
 	RHICmdList.TransitionResource(EResourceTransitionAccess::ERWBarrier, EResourceTransitionPipeline::EComputeToCompute, TextureRW);
 	FComputeShaderRHIParamRef ComputeShaderRHI = GetComputeShader();
 	SetUAVParameter(RHICmdList, ComputeShaderRHI, ClearTextureRW, FUnorderedAccessViewRHIRef());
 }
 
-void FClearBufferReplacementCS::SetParameters(FRHICommandList& RHICmdList, FUnorderedAccessViewRHIParamRef BufferRW, uint32 NumDWordsToClear,uint32 ClearValue)
+void FClearBufferReplacementCS::SetParameters(FRHICommandList& RHICmdList, FRHIUnorderedAccessView* BufferRW, uint32 NumDWordsToClear,uint32 ClearValue)
 {
 	FComputeShaderRHIParamRef ComputeShaderRHI = GetComputeShader();
 	SetShaderValue(RHICmdList, ComputeShaderRHI, ClearBufferCSParams, FUintVector4(ClearValue, NumDWordsToClear, 0, 0));
@@ -83,7 +83,7 @@ void FClearBufferReplacementCS::SetParameters(FRHICommandList& RHICmdList, FUnor
 	RHICmdList.SetUAVParameter(ComputeShaderRHI, ClearBufferRW.GetBaseIndex(), BufferRW);
 }
 
-void FClearBufferReplacementCS::FinalizeParameters(FRHICommandList& RHICmdList, FUnorderedAccessViewRHIParamRef BufferRW)
+void FClearBufferReplacementCS::FinalizeParameters(FRHICommandList& RHICmdList, FRHIUnorderedAccessView* BufferRW)
 {
 	RHICmdList.TransitionResource(EResourceTransitionAccess::ERWBarrier, EResourceTransitionPipeline::EComputeToCompute, BufferRW);
 	FComputeShaderRHIParamRef ComputeShaderRHI = GetComputeShader();
diff --git a/Engine/Source/Runtime/UtilityShaders/Private/RGBAToYUV420Shader.cpp b/Engine/Source/Runtime/UtilityShaders/Private/RGBAToYUV420Shader.cpp
--- a/Engine/Source/Runtime/UtilityShaders/Private/RGBAToYUV420Shader.cpp
+++ b/Engine/Source/Runtime/UtilityShaders/Private/RGBAToYUV420Shader.cpp
@@ -11,7 +11,7 @@
 IMPLEMENT_SHADER_TYPE(, FRGBAToYUV420CS, TEXT("/Engine/Private/RGBAToYUV420.usf"), TEXT("RGBAToYUV420Main"), SF_Compute);
 
 
-void FRGBAToYUV420CS::SetParameters(FRHICommandList& RHICmdList, TRefCountPtr<FRHITexture2D> InSrcTexture, FUnorderedAccessViewRHIParamRef InOutUAV, float InTargetHeight, float InScaleFactorX, float InScaleFactorY, float InTextureYOffset)
+void FRGBAToYUV420CS::SetParameters(FRHICommandList& RHICmdList, TRefCountPtr<FRHITexture2D> InSrcTexture, FRHIUnorderedAccessView* InOutUAV, float InTargetHeight, float InScaleFactorX, float InScaleFactorY, float InTextureYOffset)
 {
 	FComputeShaderRHIParamRef ComputeShaderRHI = GetComputeShader();
 	SetShaderValue(RHICmdList, ComputeShaderRHI, TargetHeight, InTargetHeight);
@@ -28,7 +28,7 @@ void FRGBAToYUV420CS::SetParameters(FRHICommandList& RHICmdList, TRefCountPtr<FR
 void FRGBAToYUV420CS::UnbindBuffers(FRHICommandList& RHICmdList)
 {
 	FComputeShaderRHIParamRef ComputeShaderRHI = GetComputeShader();
-	RHICmdList.SetUAVParameter(ComputeShaderRHI, OutTextureRW.GetBaseIndex(), FUnorderedAccessViewRHIParamRef());
+	RHICmdList.SetUAVParameter(ComputeShaderRHI, OutTextureRW.GetBaseIndex(), nullptr);
 }
 
 #endif //defined(HAS_MORPHEUS) && HAS_MORPHEUS
diff --git a/Engine/Source/Runtime/UtilityShaders/Public/ClearReplacementShaders.h b/Engine/Source/Runtime/UtilityShaders/Public/ClearReplacementShaders.h
--- a/Engine/Source/Runtime/UtilityShaders/Public/ClearReplacementShaders.h
+++ b/Engine/Source/Runtime/UtilityShaders/Public/ClearReplacementShaders.h
@@ -86,8 +86,8 @@ public:
 		return bShaderHasOutdatedParameters;
 	}
 	
-	UTILITYSHADERS_API void SetParameters(FRHICommandList& RHICmdList, FUnorderedAccessViewRHIParamRef TextureRW, const T(&Values)[4]);
-	UTILITYSHADERS_API void FinalizeParameters(FRHICommandList& RHICmdList, FUnorderedAccessViewRHIParamRef TextureRW);
+	UTILITYSHADERS_API void SetParameters(FRHICommandList& RHICmdList, FRHIUnorderedAccessView* TextureRW, const T(&Values)[4]);
+	UTILITYSHADERS_API void FinalizeParameters(FRHICommandList& RHICmdList, FRHIUnorderedAccessView* TextureRW);
 	
 	static bool ShouldCompilePermutation(const FGlobalShaderPermutationParameters& Parameters)
 	{
@@ -136,8 +136,8 @@ public:
 		return bShaderHasOutdatedParameters;
 	}
 
-	UTILITYSHADERS_API void SetParameters(FRHICommandList& RHICmdList, FUnorderedAccessViewRHIParamRef TextureRW, const T(&Values)[4]);
-	UTILITYSHADERS_API void FinalizeParameters(FRHICommandList& RHICmdList, FUnorderedAccessViewRHIParamRef TextureRW);
+	UTILITYSHADERS_API void SetParameters(FRHICommandList& RHICmdList, FRHIUnorderedAccessView* TextureRW, const T(&Values)[4]);
+	UTILITYSHADERS_API void FinalizeParameters(FRHICommandList& RHICmdList, FRHIUnorderedAccessView* TextureRW);
 
 	static bool ShouldCompilePermutation(const FGlobalShaderPermutationParameters& Parameters)
 	{
@@ -186,8 +186,8 @@ public:
 		return bShaderHasOutdatedParameters;
 	}
 
-	UTILITYSHADERS_API void SetParameters(FRHICommandList& RHICmdList, FUnorderedAccessViewRHIParamRef TextureRW, const T(&Values)[4]);
-	UTILITYSHADERS_API void FinalizeParameters(FRHICommandList& RHICmdList, FUnorderedAccessViewRHIParamRef TextureRW);
+	UTILITYSHADERS_API void SetParameters(FRHICommandList& RHICmdList, FRHIUnorderedAccessView* TextureRW, const T(&Values)[4]);
+	UTILITYSHADERS_API void FinalizeParameters(FRHICommandList& RHICmdList, FRHIUnorderedAccessView* TextureRW);
 
 	static bool ShouldCompilePermutation(const FGlobalShaderPermutationParameters& Parameters)
 	{
@@ -236,8 +236,8 @@ public:
 		return bShaderHasOutdatedParameters;
 	}
 
-	UTILITYSHADERS_API void SetParameters(FRHICommandList& RHICmdList, FUnorderedAccessViewRHIParamRef TextureRW, FLinearColor ClearColor, const FVector4& InTargetBounds);
-	UTILITYSHADERS_API void FinalizeParameters(FRHICommandList& RHICmdList, FUnorderedAccessViewRHIParamRef TextureRW);
+	UTILITYSHADERS_API void SetParameters(FRHICommandList& RHICmdList, FRHIUnorderedAccessView* TextureRW, FLinearColor ClearColor, const FVector4& InTargetBounds);
+	UTILITYSHADERS_API void FinalizeParameters(FRHICommandList& RHICmdList, FRHIUnorderedAccessView* TextureRW);
 
 	static bool ShouldCompilePermutation(const FGlobalShaderPermutationParameters& Parameters)
 	{
@@ -277,8 +277,8 @@ public:
 		ClearBufferRW.Bind(Initializer.ParameterMap, TEXT("ClearBufferRW"), SPF_Mandatory);
 	}
 	
-	UTILITYSHADERS_API void SetParameters(FRHICommandList& RHICmdList, FUnorderedAccessViewRHIParamRef BufferRW, uint32 NumDWordsToClear, uint32 ClearValue);
-	UTILITYSHADERS_API void FinalizeParameters(FRHICommandList& RHICmdList, FUnorderedAccessViewRHIParamRef BufferRW);
+	UTILITYSHADERS_API void SetParameters(FRHICommandList& RHICmdList, FRHIUnorderedAccessView* BufferRW, uint32 NumDWordsToClear, uint32 ClearValue);
+	UTILITYSHADERS_API void FinalizeParameters(FRHICommandList& RHICmdList, FRHIUnorderedAccessView* BufferRW);
 	
 	// FShader interface.
 	virtual bool Serialize(FArchive& Ar) override
diff --git a/Engine/Source/Runtime/UtilityShaders/Public/RGBAToYUV420Shader.h b/Engine/Source/Runtime/UtilityShaders/Public/RGBAToYUV420Shader.h
--- a/Engine/Source/Runtime/UtilityShaders/Public/RGBAToYUV420Shader.h
+++ b/Engine/Source/Runtime/UtilityShaders/Public/RGBAToYUV420Shader.h
@@ -58,7 +58,7 @@ public:
 		return bShaderHasOutdatedParameters;
 	}
 
-	void SetParameters(FRHICommandList& RHICmdList, TRefCountPtr<FRHITexture2D> SrcTex, FUnorderedAccessViewRHIParamRef OutUAV, float TargetHeight, float ScaleFactorX, float ScaleFactorY, float TextureYOffset);
+	void SetParameters(FRHICommandList& RHICmdList, TRefCountPtr<FRHITexture2D> SrcTex, FRHIUnorderedAccessView* OutUAV, float TargetHeight, float ScaleFactorX, float ScaleFactorY, float TextureYOffset);
 	void UnbindBuffers(FRHICommandList& RHICmdList);
 
 protected:
diff --git a/Engine/Source/Runtime/VulkanRHI/Private/VulkanCommands.cpp b/Engine/Source/Runtime/VulkanRHI/Private/VulkanCommands.cpp
--- a/Engine/Source/Runtime/VulkanRHI/Private/VulkanCommands.cpp
+++ b/Engine/Source/Runtime/VulkanRHI/Private/VulkanCommands.cpp
@@ -156,7 +156,7 @@ void FVulkanCommandListContext::RHIDispatchIndirectComputeShader(FVertexBufferRH
 	VulkanRHI::/*Debug*/HeavyWeightBarrier(CmdBuffer/*, 2*/);
 }
 
-void FVulkanCommandListContext::RHISetUAVParameter(FComputeShaderRHIParamRef ComputeShaderRHI, uint32 UAVIndex, FUnorderedAccessViewRHIParamRef UAVRHI)
+void FVulkanCommandListContext::RHISetUAVParameter(FComputeShaderRHIParamRef ComputeShaderRHI, uint32 UAVIndex, FRHIUnorderedAccessView* UAVRHI)
 {
 	check(PendingComputeState->GetCurrentShader() == ResourceCast(ComputeShaderRHI));
 
@@ -168,7 +168,7 @@ void FVulkanCommandListContext::RHISetUAVParameter(FComputeShaderRHIParamRef Com
 	}
 }
 
-void FVulkanCommandListContext::RHISetUAVParameter(FComputeShaderRHIParamRef ComputeShaderRHI,uint32 UAVIndex,FUnorderedAccessViewRHIParamRef UAVRHI, uint32 InitialCount)
+void FVulkanCommandListContext::RHISetUAVParameter(FComputeShaderRHIParamRef ComputeShaderRHI,uint32 UAVIndex, FRHIUnorderedAccessView* UAVRHI, uint32 InitialCount)
 {
 	check(PendingComputeState->GetCurrentShader() == ResourceCast(ComputeShaderRHI));
 
@@ -245,14 +245,14 @@ void FVulkanCommandListContext::RHISetShaderTexture(FComputeShaderRHIParamRef Co
 	NewTextureRHI->SetLastRenderTime((float)FPlatformTime::Seconds());
 }
 
-void FVulkanCommandListContext::RHISetShaderResourceViewParameter(FVertexShaderRHIParamRef VertexShaderRHI, uint32 TextureIndex, FShaderResourceViewRHIParamRef SRVRHI)
+void FVulkanCommandListContext::RHISetShaderResourceViewParameter(FVertexShaderRHIParamRef VertexShaderRHI, uint32 TextureIndex, FRHIShaderResourceView* SRVRHI)
 {
 	check(PendingGfxState->GetCurrentShaderKey(ShaderStage::Vertex) == GetShaderKey(VertexShaderRHI));
 	FVulkanShaderResourceView* SRV = ResourceCast(SRVRHI);
 	PendingGfxState->SetSRVForStage(ShaderStage::Vertex, TextureIndex, SRV);
 }
 
-void FVulkanCommandListContext::RHISetShaderResourceViewParameter(FHullShaderRHIParamRef HullShaderRHI,uint32 TextureIndex,FShaderResourceViewRHIParamRef SRVRHI)
+void FVulkanCommandListContext::RHISetShaderResourceViewParameter(FHullShaderRHIParamRef HullShaderRHI,uint32 TextureIndex, FRHIShaderResourceView* SRVRHI)
 {
 #if PLATFORM_SUPPORTS_TESSELLATION_SHADERS
 	check(PendingGfxState->GetCurrentShaderKey(ShaderStage::Hull) == GetShaderKey(HullShaderRHI));
@@ -263,7 +263,7 @@ void FVulkanCommandListContext::RHISetShaderResourceViewParameter(FHullShaderRHI
 #endif
 }
 
-void FVulkanCommandListContext::RHISetShaderResourceViewParameter(FDomainShaderRHIParamRef DomainShaderRHI,uint32 TextureIndex,FShaderResourceViewRHIParamRef SRVRHI)
+void FVulkanCommandListContext::RHISetShaderResourceViewParameter(FDomainShaderRHIParamRef DomainShaderRHI,uint32 TextureIndex, FRHIShaderResourceView* SRVRHI)
 {
 #if PLATFORM_SUPPORTS_TESSELLATION_SHADERS
 	check(PendingGfxState->GetCurrentShaderKey(ShaderStage::Domain) == GetShaderKey(DomainShaderRHI));
@@ -274,7 +274,7 @@ void FVulkanCommandListContext::RHISetShaderResourceViewParameter(FDomainShaderR
 #endif
 }
 
-void FVulkanCommandListContext::RHISetShaderResourceViewParameter(FGeometryShaderRHIParamRef GeometryShaderRHI,uint32 TextureIndex,FShaderResourceViewRHIParamRef SRVRHI)
+void FVulkanCommandListContext::RHISetShaderResourceViewParameter(FGeometryShaderRHIParamRef GeometryShaderRHI,uint32 TextureIndex, FRHIShaderResourceView* SRVRHI)
 {
 #if VULKAN_SUPPORTS_GEOMETRY_SHADERS
 	check(PendingGfxState->GetCurrentShaderKey(ShaderStage::Geometry) == GetShaderKey(GeometryShaderRHI));
@@ -285,14 +285,14 @@ void FVulkanCommandListContext::RHISetShaderResourceViewParameter(FGeometryShade
 #endif
 }
 
-void FVulkanCommandListContext::RHISetShaderResourceViewParameter(FPixelShaderRHIParamRef PixelShaderRHI,uint32 TextureIndex,FShaderResourceViewRHIParamRef SRVRHI)
+void FVulkanCommandListContext::RHISetShaderResourceViewParameter(FPixelShaderRHIParamRef PixelShaderRHI,uint32 TextureIndex, FRHIShaderResourceView* SRVRHI)
 {
 	check(PendingGfxState->GetCurrentShaderKey(ShaderStage::Pixel) == GetShaderKey(PixelShaderRHI));
 	FVulkanShaderResourceView* SRV = ResourceCast(SRVRHI);
 	PendingGfxState->SetSRVForStage(ShaderStage::Pixel, TextureIndex, SRV);
 }
 
-void FVulkanCommandListContext::RHISetShaderResourceViewParameter(FComputeShaderRHIParamRef ComputeShaderRHI,uint32 TextureIndex, FShaderResourceViewRHIParamRef SRVRHI)
+void FVulkanCommandListContext::RHISetShaderResourceViewParameter(FComputeShaderRHIParamRef ComputeShaderRHI,uint32 TextureIndex, FRHIShaderResourceView* SRVRHI)
 {
 	check(PendingComputeState->GetCurrentShader() == ResourceCast(ComputeShaderRHI));
 
diff --git a/Engine/Source/Runtime/VulkanRHI/Private/VulkanContext.h b/Engine/Source/Runtime/VulkanRHI/Private/VulkanContext.h
--- a/Engine/Source/Runtime/VulkanRHI/Private/VulkanContext.h
+++ b/Engine/Source/Runtime/VulkanRHI/Private/VulkanContext.h
@@ -185,14 +185,14 @@ public:
 	virtual void RHISetShaderSampler(FDomainShaderRHIParamRef DomainShader, uint32 SamplerIndex, FRHISamplerState* NewState) final override;
 	virtual void RHISetShaderSampler(FHullShaderRHIParamRef HullShader, uint32 SamplerIndex, FRHISamplerState* NewState) final override;
 	virtual void RHISetShaderSampler(FPixelShaderRHIParamRef PixelShader, uint32 SamplerIndex, FRHISamplerState* NewState) final override;
-	virtual void RHISetUAVParameter(FComputeShaderRHIParamRef ComputeShader, uint32 UAVIndex, FUnorderedAccessViewRHIParamRef UAV) final override;
-	virtual void RHISetUAVParameter(FComputeShaderRHIParamRef ComputeShader, uint32 UAVIndex, FUnorderedAccessViewRHIParamRef UAV, uint32 InitialCount) final override;
-	virtual void RHISetShaderResourceViewParameter(FPixelShaderRHIParamRef PixelShader, uint32 SamplerIndex, FShaderResourceViewRHIParamRef SRV) final override;
-	virtual void RHISetShaderResourceViewParameter(FVertexShaderRHIParamRef VertexShader, uint32 SamplerIndex, FShaderResourceViewRHIParamRef SRV) final override;
-	virtual void RHISetShaderResourceViewParameter(FComputeShaderRHIParamRef ComputeShader, uint32 SamplerIndex, FShaderResourceViewRHIParamRef SRV) final override;
-	virtual void RHISetShaderResourceViewParameter(FHullShaderRHIParamRef HullShader, uint32 SamplerIndex, FShaderResourceViewRHIParamRef SRV) final override;
-	virtual void RHISetShaderResourceViewParameter(FDomainShaderRHIParamRef DomainShader, uint32 SamplerIndex, FShaderResourceViewRHIParamRef SRV) final override;
-	virtual void RHISetShaderResourceViewParameter(FGeometryShaderRHIParamRef GeometryShader, uint32 SamplerIndex, FShaderResourceViewRHIParamRef SRV) final override;
+	virtual void RHISetUAVParameter(FComputeShaderRHIParamRef ComputeShader, uint32 UAVIndex, FRHIUnorderedAccessView* UAV) final override;
+	virtual void RHISetUAVParameter(FComputeShaderRHIParamRef ComputeShader, uint32 UAVIndex, FRHIUnorderedAccessView* UAV, uint32 InitialCount) final override;
+	virtual void RHISetShaderResourceViewParameter(FPixelShaderRHIParamRef PixelShader, uint32 SamplerIndex, FRHIShaderResourceView* SRV) final override;
+	virtual void RHISetShaderResourceViewParameter(FVertexShaderRHIParamRef VertexShader, uint32 SamplerIndex, FRHIShaderResourceView* SRV) final override;
+	virtual void RHISetShaderResourceViewParameter(FComputeShaderRHIParamRef ComputeShader, uint32 SamplerIndex, FRHIShaderResourceView* SRV) final override;
+	virtual void RHISetShaderResourceViewParameter(FHullShaderRHIParamRef HullShader, uint32 SamplerIndex, FRHIShaderResourceView* SRV) final override;
+	virtual void RHISetShaderResourceViewParameter(FDomainShaderRHIParamRef DomainShader, uint32 SamplerIndex, FRHIShaderResourceView* SRV) final override;
+	virtual void RHISetShaderResourceViewParameter(FGeometryShaderRHIParamRef GeometryShader, uint32 SamplerIndex, FRHIShaderResourceView* SRV) final override;
 	virtual void RHISetShaderUniformBuffer(FVertexShaderRHIParamRef VertexShader, uint32 BufferIndex, FUniformBufferRHIParamRef Buffer) final override;
 	virtual void RHISetShaderUniformBuffer(FHullShaderRHIParamRef HullShader, uint32 BufferIndex, FUniformBufferRHIParamRef Buffer) final override;
 	virtual void RHISetShaderUniformBuffer(FDomainShaderRHIParamRef DomainShader, uint32 BufferIndex, FUniformBufferRHIParamRef Buffer) final override;
@@ -206,7 +206,7 @@ public:
 	virtual void RHISetShaderParameter(FPixelShaderRHIParamRef PixelShader, uint32 BufferIndex, uint32 BaseIndex, uint32 NumBytes, const void* NewValue) final override;
 	virtual void RHISetShaderParameter(FComputeShaderRHIParamRef ComputeShader, uint32 BufferIndex, uint32 BaseIndex, uint32 NumBytes, const void* NewValue) final override;
 	virtual void RHISetStencilRef(uint32 StencilRef) final override;
-	virtual void RHISetRenderTargets(uint32 NumSimultaneousRenderTargets, const FRHIRenderTargetView* NewRenderTargets, const FRHIDepthRenderTargetView* NewDepthStencilTarget, uint32 NumUAVs, const FUnorderedAccessViewRHIParamRef* UAVs) final override;
+	virtual void RHISetRenderTargets(uint32 NumSimultaneousRenderTargets, const FRHIRenderTargetView* NewRenderTargets, const FRHIDepthRenderTargetView* NewDepthStencilTarget, uint32 NumUAVs, FRHIUnorderedAccessView* const* UAVs) final override;
 	virtual void RHISetRenderTargetsAndClear(const FRHISetRenderTargetsInfo& RenderTargetsInfo) final override;
 	virtual void RHIDrawPrimitive(uint32 BaseVertexIndex, uint32 NumPrimitives, uint32 NumInstances) final override;
 	virtual void RHIDrawPrimitiveIndirect(FVertexBufferRHIParamRef ArgumentBuffer, uint32 ArgumentOffset) final override;
@@ -226,11 +226,11 @@ public:
 
 	virtual void RHIFlushComputeShaderCache() final override;
 	virtual void RHISetMultipleViewports(uint32 Count, const FViewportBounds* Data) final override;
-	virtual void RHIClearTinyUAV(FUnorderedAccessViewRHIParamRef UnorderedAccessViewRHI, const uint32* Values) final override;
+	virtual void RHIClearTinyUAV(FRHIUnorderedAccessView* UnorderedAccessViewRHI, const uint32* Values) final override;
 	virtual void RHICopyToResolveTarget(FTextureRHIParamRef SourceTexture, FTextureRHIParamRef DestTexture, const FResolveParams& ResolveParams) final override;
 	virtual void RHICopyTexture(FTextureRHIParamRef SourceTexture, FTextureRHIParamRef DestTexture, const FRHICopyTextureInfo& CopyInfo) final override;
 	virtual void RHITransitionResources(EResourceTransitionAccess TransitionType, FTextureRHIParamRef* InRenderTargets, int32 NumTextures) final override;
-	virtual void RHITransitionResources(EResourceTransitionAccess TransitionType, EResourceTransitionPipeline TransitionPipeline, FUnorderedAccessViewRHIParamRef* InUAVs, int32 NumUAVs, FRHIComputeFence* WriteComputeFence) final override;
+	virtual void RHITransitionResources(EResourceTransitionAccess TransitionType, EResourceTransitionPipeline TransitionPipeline, FRHIUnorderedAccessView** InUAVs, int32 NumUAVs, FRHIComputeFence* WriteComputeFence) final override;
 	virtual void RHICopyToStagingBuffer(FVertexBufferRHIParamRef SourceBuffer, FStagingBufferRHIParamRef DestinationStagingBuffer, uint32 Offset, uint32 NumBytes) final override;
 	virtual void RHIWriteGPUFence(FGPUFenceRHIParamRef Fence) final override;
 
diff --git a/Engine/Source/Runtime/VulkanRHI/Private/VulkanRenderTarget.cpp b/Engine/Source/Runtime/VulkanRHI/Private/VulkanRenderTarget.cpp
--- a/Engine/Source/Runtime/VulkanRHI/Private/VulkanRenderTarget.cpp
+++ b/Engine/Source/Runtime/VulkanRHI/Private/VulkanRenderTarget.cpp
@@ -470,7 +470,7 @@ void FTransitionAndLayoutManager::TransitionResource(FVulkanCmdBuffer* CmdBuffer
 	}
 }
 
-void FVulkanCommandListContext::RHISetRenderTargets(uint32 NumSimultaneousRenderTargets, const FRHIRenderTargetView* NewRenderTargets, const FRHIDepthRenderTargetView* NewDepthStencilTarget, uint32 NumUAVs, const FUnorderedAccessViewRHIParamRef* UAVs)
+void FVulkanCommandListContext::RHISetRenderTargets(uint32 NumSimultaneousRenderTargets, const FRHIRenderTargetView* NewRenderTargets, const FRHIDepthRenderTargetView* NewDepthStencilTarget, uint32 NumUAVs, FRHIUnorderedAccessView* const* UAVs)
 {
 	FRHIDepthRenderTargetView DepthView;
 	if (NewDepthStencilTarget)
@@ -1204,7 +1204,7 @@ void FVulkanDynamicRHI::RHIRead3DSurfaceFloatData(FTextureRHIParamRef TextureRHI
 	Device->GetImmediateContext().GetCommandBufferManager()->PrepareForNewActiveCommandBuffer();
 }
 
-void FVulkanCommandListContext::RHITransitionResources(EResourceTransitionAccess TransitionType, EResourceTransitionPipeline TransitionPipeline, FUnorderedAccessViewRHIParamRef* InUAVs, int32 NumUAVs, FRHIComputeFence* WriteComputeFenceRHI)
+void FVulkanCommandListContext::RHITransitionResources(EResourceTransitionAccess TransitionType, EResourceTransitionPipeline TransitionPipeline, FRHIUnorderedAccessView** InUAVs, int32 NumUAVs, FRHIComputeFence* WriteComputeFenceRHI)
 {
 	FPendingTransition PendingTransition;
 	if (NumUAVs > 0)
diff --git a/Engine/Source/Runtime/VulkanRHI/Private/VulkanTexture.cpp b/Engine/Source/Runtime/VulkanRHI/Private/VulkanTexture.cpp
--- a/Engine/Source/Runtime/VulkanRHI/Private/VulkanTexture.cpp
+++ b/Engine/Source/Runtime/VulkanRHI/Private/VulkanTexture.cpp
@@ -2090,7 +2090,7 @@ void FVulkanDynamicRHI::RHIBindDebugLabelName(FTextureRHIParamRef TextureRHI, co
 	TextureRHI->SetName(DebugName);
 }
 
-void FVulkanDynamicRHI::RHIBindDebugLabelName(FUnorderedAccessViewRHIParamRef UnorderedAccessViewRHI, const TCHAR* Name)
+void FVulkanDynamicRHI::RHIBindDebugLabelName(FRHIUnorderedAccessView* UnorderedAccessViewRHI, const TCHAR* Name)
 {
 #if VULKAN_ENABLE_DUMP_LAYER || VULKAN_ENABLE_API_DUMP
 	//if (Device->SupportsDebugMarkers())
diff --git a/Engine/Source/Runtime/VulkanRHI/Private/VulkanUAV.cpp b/Engine/Source/Runtime/VulkanRHI/Private/VulkanUAV.cpp
--- a/Engine/Source/Runtime/VulkanRHI/Private/VulkanUAV.cpp
+++ b/Engine/Source/Runtime/VulkanRHI/Private/VulkanUAV.cpp
@@ -322,7 +322,7 @@ FShaderResourceViewRHIRef FVulkanDynamicRHI::RHICreateShaderResourceView(FIndexB
 	return SRV;
 }
 
-void FVulkanDynamicRHI::RHIUpdateShaderResourceView(FShaderResourceViewRHIParamRef SRV, FVertexBufferRHIParamRef VertexBuffer, uint32 Stride, uint8 Format)
+void FVulkanDynamicRHI::RHIUpdateShaderResourceView(FRHIShaderResourceView* SRV, FVertexBufferRHIParamRef VertexBuffer, uint32 Stride, uint8 Format)
 {
 	FVulkanShaderResourceView* SRVVk = ResourceCast(SRV);
 	check(SRVVk && SRVVk->GetParent() == Device);
@@ -337,7 +337,7 @@ void FVulkanDynamicRHI::RHIUpdateShaderResourceView(FShaderResourceViewRHIParamR
 	}
 }
 
-void FVulkanDynamicRHI::RHIUpdateShaderResourceView(FShaderResourceViewRHIParamRef SRV, FIndexBufferRHIParamRef IndexBuffer)
+void FVulkanDynamicRHI::RHIUpdateShaderResourceView(FRHIShaderResourceView* SRV, FIndexBufferRHIParamRef IndexBuffer)
 {
 	FVulkanShaderResourceView* SRVVk = ResourceCast(SRV);
 	check(SRVVk && SRVVk->GetParent() == Device);
@@ -352,7 +352,7 @@ void FVulkanDynamicRHI::RHIUpdateShaderResourceView(FShaderResourceViewRHIParamR
 	}
 }
 
-void FVulkanCommandListContext::RHIClearTinyUAV(FUnorderedAccessViewRHIParamRef UnorderedAccessViewRHI, const uint32* Values)
+void FVulkanCommandListContext::RHIClearTinyUAV(FRHIUnorderedAccessView* UnorderedAccessViewRHI, const uint32* Values)
 {
 	FVulkanUnorderedAccessView* UnorderedAccessView = ResourceCast(UnorderedAccessViewRHI);
 	FVulkanCmdBuffer* CmdBuffer = CommandBufferManager->GetActiveCmdBuffer();
diff --git a/Engine/Source/Runtime/VulkanRHI/Public/VulkanDynamicRHI.h b/Engine/Source/Runtime/VulkanRHI/Public/VulkanDynamicRHI.h
--- a/Engine/Source/Runtime/VulkanRHI/Public/VulkanDynamicRHI.h
+++ b/Engine/Source/Runtime/VulkanRHI/Public/VulkanDynamicRHI.h
@@ -78,8 +78,8 @@ public:
 	virtual FShaderResourceViewRHIRef RHICreateShaderResourceView(FStructuredBufferRHIParamRef StructuredBuffer) final override;
 	virtual FShaderResourceViewRHIRef RHICreateShaderResourceView(FVertexBufferRHIParamRef VertexBuffer, uint32 Stride, uint8 Format) final override;
 	virtual FShaderResourceViewRHIRef RHICreateShaderResourceView(FIndexBufferRHIParamRef Buffer) final override;
-	virtual void RHIUpdateShaderResourceView(FShaderResourceViewRHIParamRef SRV, FVertexBufferRHIParamRef VertexBuffer, uint32 Stride, uint8 Format) final override;
-	virtual void RHIUpdateShaderResourceView(FShaderResourceViewRHIParamRef SRV, FIndexBufferRHIParamRef IndexBuffer) final override;
+	virtual void RHIUpdateShaderResourceView(FRHIShaderResourceView* SRV, FVertexBufferRHIParamRef VertexBuffer, uint32 Stride, uint8 Format) final override;
+	virtual void RHIUpdateShaderResourceView(FRHIShaderResourceView* SRV, FIndexBufferRHIParamRef IndexBuffer) final override;
 	virtual uint64 RHICalcTexture2DPlatformSize(uint32 SizeX, uint32 SizeY, uint8 Format, uint32 NumMips, uint32 NumSamples, uint32 Flags, uint32& OutAlign) final override;
 	virtual uint64 RHICalcTexture3DPlatformSize(uint32 SizeX, uint32 SizeY, uint32 SizeZ, uint8 Format, uint32 NumMips, uint32 Flags, uint32& OutAlign) final override;
 	virtual uint64 RHICalcTextureCubePlatformSize(uint32 Size, uint8 Format, uint32 NumMips, uint32 Flags, uint32& OutAlign) final override;
@@ -123,7 +123,7 @@ public:
 	virtual void* RHILockTextureCubeFace(FTextureCubeRHIParamRef Texture, uint32 FaceIndex, uint32 ArrayIndex, uint32 MipIndex, EResourceLockMode LockMode, uint32& DestStride, bool bLockWithinMiptail) final override;
 	virtual void RHIUnlockTextureCubeFace(FTextureCubeRHIParamRef Texture, uint32 FaceIndex, uint32 ArrayIndex, uint32 MipIndex, bool bLockWithinMiptail) final override;
 	virtual void RHIBindDebugLabelName(FTextureRHIParamRef Texture, const TCHAR* Name) final override;
-	virtual void RHIBindDebugLabelName(FUnorderedAccessViewRHIParamRef UnorderedAccessViewRHI, const TCHAR* Name) final override;
+	virtual void RHIBindDebugLabelName(FRHIUnorderedAccessView* UnorderedAccessViewRHI, const TCHAR* Name) final override;
 	virtual void RHIReadSurfaceData(FTextureRHIParamRef Texture, FIntRect Rect, TArray<FColor>& OutData, FReadSurfaceDataFlags InFlags) final override;
 	virtual void RHIReadSurfaceData(FTextureRHIParamRef Texture, FIntRect Rect, TArray<FLinearColor>& OutData, FReadSurfaceDataFlags InFlags) final override;
 	virtual void RHIMapStagingSurface(FTextureRHIParamRef Texture, void*& OutData, int32& OutWidth, int32& OutHeight) final override;
diff --git a/Engine/Source/Runtime/Windows/D3D11RHI/Private/D3D11Commands.cpp b/Engine/Source/Runtime/Windows/D3D11RHI/Private/D3D11Commands.cpp
--- a/Engine/Source/Runtime/Windows/D3D11RHI/Private/D3D11Commands.cpp
+++ b/Engine/Source/Runtime/Windows/D3D11RHI/Private/D3D11Commands.cpp
@@ -113,7 +113,7 @@ void FD3D11DynamicRHI::RHIEndUpdateMultiFrameResource(FTextureRHIParamRef RHITex
 	RHIPopEvent();	
 }
 
-void FD3D11DynamicRHI::RHIBeginUpdateMultiFrameResource(FUnorderedAccessViewRHIParamRef UAVRHI)
+void FD3D11DynamicRHI::RHIBeginUpdateMultiFrameResource(FRHIUnorderedAccessView* UAVRHI)
 {
 	if (!IsRHIDeviceNVIDIA() || GNumAlternateFrameRenderingGroups == 1) return;
 
@@ -137,7 +137,7 @@ void FD3D11DynamicRHI::RHIBeginUpdateMultiFrameResource(FUnorderedAccessViewRHIP
 	RHIPopEvent();
 }
 
-void FD3D11DynamicRHI::RHIEndUpdateMultiFrameResource(FUnorderedAccessViewRHIParamRef UAVRHI)
+void FD3D11DynamicRHI::RHIEndUpdateMultiFrameResource(FRHIUnorderedAccessView* UAVRHI)
 {
 	if (!IsRHIDeviceNVIDIA() || GNumAlternateFrameRenderingGroups == 1) return;
 
@@ -435,7 +435,7 @@ void FD3D11DynamicRHI::RHISetShaderTexture(FComputeShaderRHIParamRef ComputeShad
 	}
 }
 
-void FD3D11DynamicRHI::RHISetUAVParameter(FComputeShaderRHIParamRef ComputeShaderRHI,uint32 UAVIndex,FUnorderedAccessViewRHIParamRef UAVRHI)
+void FD3D11DynamicRHI::RHISetUAVParameter(FComputeShaderRHIParamRef ComputeShaderRHI,uint32 UAVIndex, FRHIUnorderedAccessView* UAVRHI)
 {
 	//VALIDATE_BOUND_SHADER(ComputeShaderRHI);
 
@@ -460,7 +460,7 @@ void FD3D11DynamicRHI::RHISetUAVParameter(FComputeShaderRHIParamRef ComputeShade
 	Direct3DDeviceIMContext->CSSetUnorderedAccessViews(UAVIndex,1,&D3D11UAV, &InitialCount );
 }
 
-void FD3D11DynamicRHI::RHISetUAVParameter(FComputeShaderRHIParamRef ComputeShaderRHI,uint32 UAVIndex,FUnorderedAccessViewRHIParamRef UAVRHI, uint32 InitialCount )
+void FD3D11DynamicRHI::RHISetUAVParameter(FComputeShaderRHIParamRef ComputeShaderRHI,uint32 UAVIndex, FRHIUnorderedAccessView* UAVRHI, uint32 InitialCount )
 {
 	//VALIDATE_BOUND_SHADER(ComputeShaderRHI);
 
@@ -483,7 +483,7 @@ void FD3D11DynamicRHI::RHISetUAVParameter(FComputeShaderRHIParamRef ComputeShade
 	Direct3DDeviceIMContext->CSSetUnorderedAccessViews(UAVIndex,1,&D3D11UAV, &InitialCount );
 }
 
-void FD3D11DynamicRHI::RHISetShaderResourceViewParameter(FPixelShaderRHIParamRef PixelShaderRHI,uint32 TextureIndex,FShaderResourceViewRHIParamRef SRVRHI)
+void FD3D11DynamicRHI::RHISetShaderResourceViewParameter(FPixelShaderRHIParamRef PixelShaderRHI,uint32 TextureIndex, FRHIShaderResourceView* SRVRHI)
 {
 	VALIDATE_BOUND_SHADER(PixelShaderRHI);
 
@@ -501,7 +501,7 @@ void FD3D11DynamicRHI::RHISetShaderResourceViewParameter(FPixelShaderRHIParamRef
 	SetShaderResourceView<SF_Pixel>(Resource, D3D11SRV, TextureIndex, NAME_None);
 }
 
-void FD3D11DynamicRHI::RHISetShaderResourceViewParameter(FVertexShaderRHIParamRef VertexShaderRHI,uint32 TextureIndex,FShaderResourceViewRHIParamRef SRVRHI)
+void FD3D11DynamicRHI::RHISetShaderResourceViewParameter(FVertexShaderRHIParamRef VertexShaderRHI,uint32 TextureIndex, FRHIShaderResourceView* SRVRHI)
 {
 	VALIDATE_BOUND_SHADER(VertexShaderRHI);
 
@@ -519,7 +519,7 @@ void FD3D11DynamicRHI::RHISetShaderResourceViewParameter(FVertexShaderRHIParamRe
 	SetShaderResourceView<SF_Vertex>(Resource, D3D11SRV, TextureIndex, NAME_None);
 }
 
-void FD3D11DynamicRHI::RHISetShaderResourceViewParameter(FComputeShaderRHIParamRef ComputeShaderRHI,uint32 TextureIndex,FShaderResourceViewRHIParamRef SRVRHI)
+void FD3D11DynamicRHI::RHISetShaderResourceViewParameter(FComputeShaderRHIParamRef ComputeShaderRHI,uint32 TextureIndex, FRHIShaderResourceView* SRVRHI)
 {
 	//VALIDATE_BOUND_SHADER(ComputeShaderRHI);
 
@@ -537,7 +537,7 @@ void FD3D11DynamicRHI::RHISetShaderResourceViewParameter(FComputeShaderRHIParamR
 	SetShaderResourceView<SF_Compute>(Resource, D3D11SRV, TextureIndex, NAME_None);
 }
 
-void FD3D11DynamicRHI::RHISetShaderResourceViewParameter(FHullShaderRHIParamRef HullShaderRHI,uint32 TextureIndex,FShaderResourceViewRHIParamRef SRVRHI)
+void FD3D11DynamicRHI::RHISetShaderResourceViewParameter(FHullShaderRHIParamRef HullShaderRHI,uint32 TextureIndex, FRHIShaderResourceView* SRVRHI)
 {
 	VALIDATE_BOUND_SHADER(HullShaderRHI);
 
@@ -555,7 +555,7 @@ void FD3D11DynamicRHI::RHISetShaderResourceViewParameter(FHullShaderRHIParamRef
 	SetShaderResourceView<SF_Hull>(Resource, D3D11SRV, TextureIndex, NAME_None);
 }
 
-void FD3D11DynamicRHI::RHISetShaderResourceViewParameter(FDomainShaderRHIParamRef DomainShaderRHI,uint32 TextureIndex,FShaderResourceViewRHIParamRef SRVRHI)
+void FD3D11DynamicRHI::RHISetShaderResourceViewParameter(FDomainShaderRHIParamRef DomainShaderRHI,uint32 TextureIndex, FRHIShaderResourceView* SRVRHI)
 {
 	VALIDATE_BOUND_SHADER(DomainShaderRHI);
 
@@ -573,7 +573,7 @@ void FD3D11DynamicRHI::RHISetShaderResourceViewParameter(FDomainShaderRHIParamRe
 	SetShaderResourceView<SF_Domain>(Resource, D3D11SRV, TextureIndex, NAME_None);
 }
 
-void FD3D11DynamicRHI::RHISetShaderResourceViewParameter(FGeometryShaderRHIParamRef GeometryShaderRHI,uint32 TextureIndex,FShaderResourceViewRHIParamRef SRVRHI)
+void FD3D11DynamicRHI::RHISetShaderResourceViewParameter(FGeometryShaderRHIParamRef GeometryShaderRHI,uint32 TextureIndex, FRHIShaderResourceView* SRVRHI)
 {
 	VALIDATE_BOUND_SHADER(GeometryShaderRHI);
 
@@ -943,7 +943,7 @@ void FD3D11DynamicRHI::RHISetRenderTargets(
 	const FRHIRenderTargetView* NewRenderTargetsRHI,
 	const FRHIDepthRenderTargetView* NewDepthStencilTargetRHI,
 	uint32 NewNumUAVs,
-	const FUnorderedAccessViewRHIParamRef* UAVs
+	FRHIUnorderedAccessView* const* UAVs
 	)
 {
 	FD3D11TextureBase* NewDepthStencilTarget = GetD3D11TextureFromRHITexture(NewDepthStencilTargetRHI ? NewDepthStencilTargetRHI->Texture : nullptr);
@@ -1159,8 +1159,8 @@ void FD3D11DynamicRHI::RHISetRenderTargets(
 
 void FD3D11DynamicRHI::RHISetRenderTargetsAndClear(const FRHISetRenderTargetsInfo& RenderTargetsInfo)
 {
-	// Here convert to FUnorderedAccessViewRHIParamRef* in order to call RHISetRenderTargets
-	FUnorderedAccessViewRHIParamRef UAVs[MaxSimultaneousUAVs] = {};
+	// Here convert to FRHIUnorderedAccessView** in order to call RHISetRenderTargets
+	FRHIUnorderedAccessView* UAVs[MaxSimultaneousUAVs] = {};
 	for (int32 UAVIndex = 0; UAVIndex < RenderTargetsInfo.NumUAVs; ++UAVIndex)
 	{
 		UAVs[UAVIndex] = RenderTargetsInfo.UnorderedAccessView[UAVIndex].GetReference();
@@ -1910,7 +1910,7 @@ void FD3D11DynamicRHI::RHITransitionResources(EResourceTransitionAccess Transiti
 	}
 }
 
-void FD3D11DynamicRHI::RHITransitionResources(EResourceTransitionAccess TransitionType, EResourceTransitionPipeline TransitionPipeline, FUnorderedAccessViewRHIParamRef* InUAVs, int32 InNumUAVs, FRHIComputeFence* WriteFence)
+void FD3D11DynamicRHI::RHITransitionResources(EResourceTransitionAccess TransitionType, EResourceTransitionPipeline TransitionPipeline, FRHIUnorderedAccessView** InUAVs, int32 InNumUAVs, FRHIComputeFence* WriteFence)
 {
 	for (int32 i = 0; i < InNumUAVs; ++i)
 	{
