Show warning message instead of crash when d3d debug options are missing but -d3ddebug is enabled.

diff --git a/Engine/Source/Runtime/Windows/D3D11RHI/Private/Windows/WindowsD3D11Device.cpp b/Engine/Source/Runtime/Windows/D3D11RHI/Private/Windows/WindowsD3D11Device.cpp
--- a/Engine/Source/Runtime/Windows/D3D11RHI/Private/Windows/WindowsD3D11Device.cpp
+++ b/Engine/Source/Runtime/Windows/D3D11RHI/Private/Windows/WindowsD3D11Device.cpp
@@ -262,6 +262,59 @@ static D3D_FEATURE_LEVEL GetAllowedD3DFeatureLevel()
 	return AllowedFeatureLevel;
 }
 
+/**
+* Attempts to create a D3D11 device to test if optional debugging features are installed.
+*/
+static void SafeTestForOptionalGraphicsTools(IDXGIAdapter* Adapter, D3D_FEATURE_LEVEL& FeatureLevel, uint32& DeviceFlags, D3D_FEATURE_LEVEL& OutFeatureLevel, int32 NumAllowedFeatureLevels)
+{
+	if((DeviceFlags & D3D11_CREATE_DEVICE_DEBUG) != D3D11_CREATE_DEVICE_DEBUG)
+	{
+		return;
+	}
+
+	ID3D11Device* D3DDevice = NULL;
+	ID3D11DeviceContext* D3DDeviceContext = NULL;
+
+	static bool bIsWin10 = FWindowsPlatformMisc::VerifyWindowsVersion(10, 0);
+	static bool bPrintErrorMessage = true;
+
+	// Test device creation.
+	HRESULT Result = D3D11CreateDevice(
+		Adapter,
+		D3D_DRIVER_TYPE_UNKNOWN,
+		NULL,
+		DeviceFlags,
+		& FeatureLevel,
+		NumAllowedFeatureLevels,
+		D3D11_SDK_VERSION,
+		&D3DDevice,
+		&OutFeatureLevel,
+		&D3DDeviceContext
+	);
+
+	// If Optional Graphics tools were enabled then return true.
+	if (SUCCEEDED(Result))
+	{
+		D3DDevice->Release();
+		D3DDeviceContext->Release();
+		return;
+	}
+	// If creation failed due to missing Optional Features in Windows 10.
+	else if (DXGI_ERROR_SDK_COMPONENT_MISSING == Result && bIsWin10)
+	{
+		// Only print the debug error message once and don't clutter the log.
+		if (bPrintErrorMessage) 
+		{
+			UE_LOG(LogD3D11RHI, Warning,
+				TEXT("-d3ddebug was ignored as optional Graphics Tools still need to be installed. Install them through the Manage Optional Features in windows."));
+			bPrintErrorMessage = false;
+		}
+
+		// Remove the debug option from the device flag.
+		DeviceFlags = DeviceFlags & (~D3D11_CREATE_DEVICE_DEBUG);
+	}
+}
+
 /**
  * Attempts to create a D3D11 device for the adapter using at most MaxFeatureLevel.
  * If creation is successful, true is returned and the supported feature level is set in OutFeatureLevel.
@@ -271,7 +324,6 @@ static bool SafeTestD3D11CreateDevice(IDXGIAdapter* Adapter,D3D_FEATURE_LEVEL Ma
 	ID3D11Device* D3DDevice = NULL;
 	ID3D11DeviceContext* D3DDeviceContext = NULL;
 	uint32 DeviceFlags = D3D11_CREATE_DEVICE_SINGLETHREADED;
-
 	// Use a debug device if specified on the command line.
 	if(D3D11RHI_ShouldCreateWithD3DDebug())
 	{
@@ -302,12 +354,13 @@ static bool SafeTestD3D11CreateDevice(IDXGIAdapter* Adapter,D3D_FEATURE_LEVEL Ma
 		return false;
 	}
 
+	SafeTestForOptionalGraphicsTools(Adapter, RequestedFeatureLevels[FirstAllowedFeatureLevel], DeviceFlags, *OutFeatureLevel, NumAllowedFeatureLevels);
+
 	__try
 	{
 		// We don't want software renderer. Ideally we specify D3D_DRIVER_TYPE_HARDWARE on creation but
 		// when we specify an adapter we need to specify D3D_DRIVER_TYPE_UNKNOWN (otherwise the call fails).
 		// We cannot check the device type later (seems this is missing functionality in D3D).
-
 		if(SUCCEEDED(D3D11CreateDevice(
 			Adapter,
 			D3D_DRIVER_TYPE_UNKNOWN,
@@ -319,7 +372,7 @@ static bool SafeTestD3D11CreateDevice(IDXGIAdapter* Adapter,D3D_FEATURE_LEVEL Ma
 			&D3DDevice,
 			OutFeatureLevel,
 			&D3DDeviceContext
-			)))
+		))) 
 		{
 			D3DDevice->Release();
 			D3DDeviceContext->Release();
@@ -1490,6 +1543,8 @@ void FD3D11DynamicRHI::InitD3DDevice()
 		
 		if (!bDeviceCreated)
 		{
+			SafeTestForOptionalGraphicsTools(Adapter, FeatureLevel, DeviceFlags, ActualFeatureLevel, 1);
+			
 			// Creating the Direct3D device.
 			VERIFYD3D11RESULT(D3D11CreateDevice(
 				Adapter,
