Make use of MIP-level information in D3D11 GetStagingTexture.

#rnx

diff --git a/Engine/Source/Runtime/Windows/D3D11RHI/Private/D3D11RenderTarget.cpp b/Engine/Source/Runtime/Windows/D3D11RHI/Private/D3D11RenderTarget.cpp
--- a/Engine/Source/Runtime/Windows/D3D11RHI/Private/D3D11RenderTarget.cpp
+++ b/Engine/Source/Runtime/Windows/D3D11RHI/Private/D3D11RenderTarget.cpp
@@ -509,11 +509,15 @@ TRefCountPtr<ID3D11Texture2D> FD3D11DynamicRHI::GetStagingTexture(FTextureRHIPar
 	StagingRectOUT.Max = FIntPoint(SizeX,SizeY);
 
 	// Copy the data to a staging resource.
-	uint32 Subresource = InFlags.GetMip();
+	uint32 Subresource = 0;
 	if( SourceDesc.MiscFlags == D3D11_RESOURCE_MISC_TEXTURECUBE )
 	{
 		uint32 D3DFace = GetD3D11CubeFace(InFlags.GetCubeFace());
-		Subresource = D3D11CalcSubresource(0,D3DFace,1);
+		Subresource = D3D11CalcSubresource(InFlags.GetMip(),D3DFace,TextureRHI->GetNumMips());
+	}
+	else
+	{
+		Subresource = D3D11CalcSubresource(InFlags.GetMip(), 0, TextureRHI->GetNumMips());
 	}
 
 	D3D11_BOX* RectPtr = NULL; // API prefers NULL for entire texture.
