Consolidate convert surface texture code between vulkan/metal/d3d RHIs

diff --git a/Engine/Source/Runtime/Apple/MetalRHI/Private/MetalRenderTarget.cpp b/Engine/Source/Runtime/Apple/MetalRHI/Private/MetalRenderTarget.cpp
--- a/Engine/Source/Runtime/Apple/MetalRHI/Private/MetalRenderTarget.cpp
+++ b/Engine/Source/Runtime/Apple/MetalRHI/Private/MetalRenderTarget.cpp
@@ -11,6 +11,7 @@
 #include "ResolveShader.h"
 #include "PipelineStateCache.h"
 #include "Math/PackedVector.h"
+#include "RHISurfaceDataConversion.h"
 
 static FResolveRect GetDefaultRect(const FResolveRect& Rect, uint32 DefaultWidth, uint32 DefaultHeight)
 {
@@ -144,269 +145,63 @@ struct FMetalRGBA16
 	uint16 A;
 };
 
-// todo: this should be available for all RHI
+void FMetalDynamicRHI::RHIReadSurfaceData(FRHITexture* TextureRHI, FIntRect InRect, TArray<FLinearColor>& OutData, FReadSurfaceDataFlags InFlags)
+{
+	// Use our current surface read implementation and convert to linear - should refactor to make optimal
+	TArray<FColor> OutDataUnConverted;
+	RHIReadSurfaceData(TextureRHI, InRect, OutDataUnConverted, InFlags);
+
+	OutData.Empty();
+	OutData.AddUninitialized(OutDataUnConverted.Num());
+
+	for (uint32 i = 0; i < OutDataUnConverted.Num(); ++i)
+	{
+		OutData[i] = OutDataUnConverted[i].ReinterpretAsLinear();
+	}
+}
+
 static void ConvertSurfaceDataToFColor(EPixelFormat Format, uint32 Width, uint32 Height, uint8 *In, uint32 SrcPitch, FColor* Out, FReadSurfaceDataFlags InFlags)
 {
 	bool bLinearToGamma = InFlags.GetLinearToGamma();
-	
-	if(Format == PF_G16 || Format == PF_R16_UINT || Format == PF_R16_SINT)
+	if (Format == PF_G16 || Format == PF_R16_UINT || Format == PF_R16_SINT)
 	{
-		// e.g. shadow maps
-		for(uint32 Y = 0; Y < Height; Y++)
-		{
-			uint16* SrcPtr = (uint16*)(In + Y * SrcPitch);
-			FColor* DestPtr = Out + Y * Width;
-			
-			for(uint32 X = 0; X < Width; X++)
-			{
-				uint16 Value16 = *SrcPtr;
-				float Value = Value16 / (float)(0xffff);
-				
-				*DestPtr = FLinearColor(Value, Value, Value).Quantize();
-				++SrcPtr;
-				++DestPtr;
-			}
-		}
+		ConvertRawR16DataToFColor(Width, Height, In, SrcPitch, Out);
 	}
-	else if(Format == PF_R8G8B8A8)
+	else if (Format == PF_R8G8B8A8)
 	{
-		// Read the data out of the buffer, converting it from ABGR to ARGB.
-		for(uint32 Y = 0; Y < Height; Y++)
-		{
-			FColor* SrcPtr = (FColor*)(In + Y * SrcPitch);
-			FColor* DestPtr = Out + Y * Width;
-			for(uint32 X = 0; X < Width; X++)
-			{
-				*DestPtr = FColor(SrcPtr->B,SrcPtr->G,SrcPtr->R,SrcPtr->A);
-				++SrcPtr;
-				++DestPtr;
-			}
-		}
+		ConvertRawR8G8B8A8DataToFColor(Width, Height, In, SrcPitch, Out);
 	}
-	else if(Format == PF_B8G8R8A8)
+	else if (Format == PF_B8G8R8A8)
 	{
-		const auto DestPitch = sizeof(FColor) * Width;
-		if (SrcPitch == DestPitch)
-		{
-			FMemory::Memcpy(Out, In, DestPitch * Height);
-		}
-		else
-		{
-			for(uint32 Y = 0; Y < Height; Y++)
-			{
-				FColor* SrcPtr = (FColor*)(In + Y * SrcPitch);
-				FColor* DestPtr = Out + Y * Width;
-				
-				// Need to copy row wise since the Pitch might not match the Width.
-				FMemory::Memcpy(DestPtr, SrcPtr, sizeof(FColor) * Width);
-			}
-		}
+		ConvertRawB8G8R8A8DataToFColor(Width, Height, In, SrcPitch, Out);
 	}
-	else if(Format == PF_A2B10G10R10)
+	else if (Format == PF_A2B10G10R10)
 	{
-		// Read the data out of the buffer, converting it from R10G10B10A2 to FColor.
-		for(uint32 Y = 0; Y < Height; Y++)
-		{
-			FMetalR10G10B10A2* SrcPtr = (FMetalR10G10B10A2*)(In + Y * SrcPitch);
-			FColor* DestPtr = Out + Y * Width;
-			for(uint32 X = 0; X < Width; X++)
-			{
-				*DestPtr = FLinearColor(
-										(float)SrcPtr->R / 1023.0f,
-										(float)SrcPtr->G / 1023.0f,
-										(float)SrcPtr->B / 1023.0f,
-										(float)SrcPtr->A / 3.0f
-										).Quantize();
-				++SrcPtr;
-				++DestPtr;
-			}
-		}
+		ConvertRawR10G10B10A2DataToFColor(Width, Height, In, SrcPitch, Out);
 	}
-	else if(Format == PF_FloatRGBA)
+	else if (Format == PF_FloatRGBA)
 	{
-		FPlane	MinValue(0.0f,0.0f,0.0f,0.0f),
-		MaxValue(1.0f,1.0f,1.0f,1.0f);
-		
-		check(sizeof(FFloat16)==sizeof(uint16));
-		
-		for(uint32 Y = 0; Y < Height; Y++)
-		{
-			FFloat16* SrcPtr = (FFloat16*)(In + Y * SrcPitch);
-			
-			for(uint32 X = 0; X < Width; X++)
-			{
-				MinValue.X = FMath::Min<float>(SrcPtr[0],MinValue.X);
-				MinValue.Y = FMath::Min<float>(SrcPtr[1],MinValue.Y);
-				MinValue.Z = FMath::Min<float>(SrcPtr[2],MinValue.Z);
-				MinValue.W = FMath::Min<float>(SrcPtr[3],MinValue.W);
-				MaxValue.X = FMath::Max<float>(SrcPtr[0],MaxValue.X);
-				MaxValue.Y = FMath::Max<float>(SrcPtr[1],MaxValue.Y);
-				MaxValue.Z = FMath::Max<float>(SrcPtr[2],MaxValue.Z);
-				MaxValue.W = FMath::Max<float>(SrcPtr[3],MaxValue.W);
-				SrcPtr += 4;
-			}
-		}
-		
-		for(uint32 Y = 0; Y < Height; Y++)
-		{
-			FFloat16* SrcPtr = (FFloat16*)(In + Y * SrcPitch);
-			FColor* DestPtr = Out + Y * Width;
-			
-			for(uint32 X = 0; X < Width; X++)
-			{
-				FColor NormalizedColor =
-				FLinearColor(
-							 (SrcPtr[0] - MinValue.X) / (MaxValue.X - MinValue.X),
-							 (SrcPtr[1] - MinValue.Y) / (MaxValue.Y - MinValue.Y),
-							 (SrcPtr[2] - MinValue.Z) / (MaxValue.Z - MinValue.Z),
-							 (SrcPtr[3] - MinValue.W) / (MaxValue.W - MinValue.W)
-							 ).ToFColor(bLinearToGamma);
-				FMemory::Memcpy(DestPtr++,&NormalizedColor,sizeof(FColor));
-				SrcPtr += 4;
-			}
-		}
+		ConvertRawR16G16B16A16FDataToFColor(Width, Height, In, SrcPitch, Out, bLinearToGamma);
 	}
 	else if (Format == PF_FloatR11G11B10)
 	{
-		check(sizeof(FFloat3Packed) == sizeof(uint32));
-		
-		for(uint32 Y = 0; Y < Height; Y++)
-		{
-			FFloat3Packed* SrcPtr = (FFloat3Packed*)(In + Y * SrcPitch);
-			FColor* DestPtr = Out + Y * Width;
-			
-			for(uint32 X = 0; X < Width; X++)
-			{
-				FLinearColor Value = (*SrcPtr).ToLinearColor();
-				
-				FColor NormalizedColor = Value.ToFColor(bLinearToGamma);
-				FMemory::Memcpy(DestPtr++, &NormalizedColor, sizeof(FColor));
-				++SrcPtr;
-			}
-		}
+		ConvertRawR11G11B10DataToFColor(Width, Height, In, SrcPitch, Out, bLinearToGamma);
 	}
 	else if (Format == PF_A32B32G32R32F)
 	{
-		FPlane MinValue(0.0f,0.0f,0.0f,0.0f);
-		FPlane MaxValue(1.0f,1.0f,1.0f,1.0f);
-		
-		for(uint32 Y = 0; Y < Height; Y++)
-		{
-			float* SrcPtr = (float*)(In + Y * SrcPitch);
-			
-			for(uint32 X = 0; X < Width; X++)
-			{
-				MinValue.X = FMath::Min<float>(SrcPtr[0],MinValue.X);
-				MinValue.Y = FMath::Min<float>(SrcPtr[1],MinValue.Y);
-				MinValue.Z = FMath::Min<float>(SrcPtr[2],MinValue.Z);
-				MinValue.W = FMath::Min<float>(SrcPtr[3],MinValue.W);
-				MaxValue.X = FMath::Max<float>(SrcPtr[0],MaxValue.X);
-				MaxValue.Y = FMath::Max<float>(SrcPtr[1],MaxValue.Y);
-				MaxValue.Z = FMath::Max<float>(SrcPtr[2],MaxValue.Z);
-				MaxValue.W = FMath::Max<float>(SrcPtr[3],MaxValue.W);
-				SrcPtr += 4;
-			}
-		}
-		
-		for(uint32 Y = 0; Y < Height; Y++)
-		{
-			float* SrcPtr = (float*)(In + Y * SrcPitch);
-			FColor* DestPtr = Out + Y * Width;
-			
-			for(uint32 X = 0; X < Width; X++)
-			{
-				FColor NormalizedColor =
-				FLinearColor(
-							 (SrcPtr[0] - MinValue.X) / (MaxValue.X - MinValue.X),
-							 (SrcPtr[1] - MinValue.Y) / (MaxValue.Y - MinValue.Y),
-							 (SrcPtr[2] - MinValue.Z) / (MaxValue.Z - MinValue.Z),
-							 (SrcPtr[3] - MinValue.W) / (MaxValue.W - MinValue.W)
-							 ).ToFColor(bLinearToGamma);
-				FMemory::Memcpy(DestPtr++,&NormalizedColor,sizeof(FColor));
-				SrcPtr += 4;
-			}
-		}
+		ConvertRawR32G32B32A32DataToFColor(Width, Height, In, SrcPitch, Out, bLinearToGamma);
 	}
-	else if(Format == PF_A16B16G16R16)
+	else if (Format == PF_A16B16G16R16)
 	{
-		// Read the data out of the buffer, converting it to FColor.
-		for(uint32 Y = 0; Y < Height; Y++)
-		{
-			FMetalRGBA16* SrcPtr = (FMetalRGBA16*)(In + Y * SrcPitch);
-			FColor* DestPtr = Out + Y * Width;
-			for(uint32 X = 0; X < Width; X++)
-			{
-				*DestPtr = FLinearColor(
-										(float)SrcPtr->R / 65535.0f,
-										(float)SrcPtr->G / 65535.0f,
-										(float)SrcPtr->B / 65535.0f,
-										(float)SrcPtr->A / 65535.0f
-										).Quantize();
-				++SrcPtr;
-				++DestPtr;
-			}
-		}
+		ConvertRawR16G16B16A16DataToFColor(Width, Height, In, SrcPitch, Out);
 	}
-	else if(Format == PF_G16R16)
+	else if (Format == PF_G16R16)
 	{
-		// Read the data out of the buffer, converting it to FColor.
-		for(uint32 Y = 0; Y < Height; Y++)
-		{
-			FMetalRG16* SrcPtr = (FMetalRG16*)(In + Y * SrcPitch);
-			FColor* DestPtr = Out + Y * Width;
-			for(uint32 X = 0; X < Width; X++)
-			{
-				*DestPtr = FLinearColor(
-										(float)SrcPtr->R / 65535.0f,
-										(float)SrcPtr->G / 65535.0f,
-										0).Quantize();
-				++SrcPtr;
-				++DestPtr;
-			}
-		}
+		ConvertRawR16G16DataToFColor(Width, Height, In, SrcPitch, Out);
 	}
-	else if(Format == PF_DepthStencil)
+	else if (Format == PF_DepthStencil)
 	{
-		// Depth
-		if(!InFlags.GetOutputStencil())
-		{
-			bool const bDepth32 = ((mtlpp::PixelFormat)GPixelFormats[Format].PlatformFormat == mtlpp::PixelFormat::Depth32Float_Stencil8);
-			
-			for (uint32 Y = 0; Y < Height; Y++)
-			{
-				uint32* SrcPtr = (uint32*)(In + Y * SrcPitch);
-				FColor* DestPtr = Out + Y * Width;
-				
-				for (uint32 X = 0; X < Width; X++)
-				{
-					float DeviceZ = bDepth32 ? *SrcPtr : (*SrcPtr & 0xffffff) / (float)(1 << 24);
-					float LinearValue = FMath::Min(InFlags.ComputeNormalizedDepth(DeviceZ), 1.0f);
-					FColor NormalizedColor = FLinearColor(LinearValue, LinearValue, LinearValue, 0).ToFColor(bLinearToGamma);
-					
-					FMemory::Memcpy(DestPtr++, &NormalizedColor, sizeof(FColor));
-					++SrcPtr;
-				}
-			}
-		}
-		// Stencil
-		else
-		{
-			// Depth stencil
-			for (uint32 Y = 0; Y < Height; Y++)
-			{
-				uint8* SrcPtr = (uint8*)(In + Y * SrcPitch);
-				FColor* DestPtr = Out + Y * Width;
-				
-				for (uint32 X = 0; X < Width; X++)
-				{
-					uint8 DeviceStencil = *SrcPtr;
-					FColor NormalizedColor = FColor(DeviceStencil, DeviceStencil, DeviceStencil, 0xFF);
-					
-					FMemory::Memcpy(DestPtr++, &NormalizedColor, sizeof(FColor));
-					++SrcPtr;
-				}
-			}
-		}
+		ConvertRawD32S8DataToFColor(Width, Height, In, SrcPitch, Out, InFlags);
 	}
 	else
 	{
@@ -415,21 +210,6 @@ static void ConvertSurfaceDataToFColor(EPixelFormat Format, uint32 Width, uint32
 	}
 }
 
-void FMetalDynamicRHI::RHIReadSurfaceData(FRHITexture* TextureRHI, FIntRect InRect, TArray<FLinearColor>& OutData, FReadSurfaceDataFlags InFlags)
-{
-	// Use our current surface read implemtation and convert to linear - should refactor to make optimal
-	TArray<FColor> OutDataUnConverted;
-	RHIReadSurfaceData(TextureRHI, InRect, OutDataUnConverted, InFlags);
-	
-	OutData.Empty();
-	OutData.AddUninitialized(OutDataUnConverted.Num());
-	
-	for(uint32 i = 0;i < OutDataUnConverted.Num();++i)
-	{
-		OutData[i] = OutDataUnConverted[i].ReinterpretAsLinear();
-	}
-}
-
 void FMetalDynamicRHI::RHIReadSurfaceData(FRHITexture* TextureRHI, FIntRect Rect, TArray<FColor>& OutData, FReadSurfaceDataFlags InFlags)
 {
 	@autoreleasepool {
diff --git a/Engine/Source/Runtime/D3D12RHI/Private/D3D12RenderTarget.cpp b/Engine/Source/Runtime/D3D12RHI/Private/D3D12RenderTarget.cpp
--- a/Engine/Source/Runtime/D3D12RHI/Private/D3D12RenderTarget.cpp
+++ b/Engine/Source/Runtime/D3D12RHI/Private/D3D12RenderTarget.cpp
@@ -12,6 +12,7 @@
 #include "SceneUtils.h"
 #include "PipelineStateCache.h"
 #include "Math/PackedVector.h"
+#include "RHISurfaceDataConversion.h"
 
 static inline DXGI_FORMAT ConvertTypelessToUnorm(DXGI_FORMAT Format)
 {
@@ -809,273 +810,55 @@ struct FD3DRGBA16
 	uint16 A;
 };
 
-// todo: this should be available for all RHI
-static void ConvertRAWSurfaceDataToFColor(DXGI_FORMAT Format, uint32 Width, uint32 Height, uint8 *In, uint32 SrcPitch, FColor* Out, FReadSurfaceDataFlags InFlags)
+/** Convert D3D format type to general pixel format type*/
+static void ConvertDXGIToFColor(DXGI_FORMAT Format, uint32 Width, uint32 Height, uint8 *In, uint32 SrcPitch, FColor* Out, FReadSurfaceDataFlags InFlags)
 {
 	bool bLinearToGamma = InFlags.GetLinearToGamma();
-
-	if (Format == DXGI_FORMAT_R16_TYPELESS)
-	{
-		// e.g. shadow maps
-		for (uint32 Y = 0; Y < Height; Y++)
-		{
-			uint16* SrcPtr = (uint16*)(In + Y * SrcPitch);
-			FColor* DestPtr = Out + Y * Width;
-
-			for (uint32 X = 0; X < Width; X++)
-			{
-				uint16 Value16 = *SrcPtr;
-				float Value = Value16 / (float)(0xffff);
-
-				*DestPtr = FLinearColor(Value, Value, Value).Quantize();
-				++SrcPtr;
-				++DestPtr;
-			}
-		}
-	}
-	else if (Format == DXGI_FORMAT_R8G8B8A8_TYPELESS || Format == DXGI_FORMAT_R8G8B8A8_UNORM)
-	{
-		// Read the data out of the buffer, converting it from ABGR to ARGB.
-		for (uint32 Y = 0; Y < Height; Y++)
-		{
-			FColor* SrcPtr = (FColor*)(In + Y * SrcPitch);
-			FColor* DestPtr = Out + Y * Width;
-			for (uint32 X = 0; X < Width; X++)
-			{
-				*DestPtr = FColor(SrcPtr->B, SrcPtr->G, SrcPtr->R, SrcPtr->A);
-				++SrcPtr;
-				++DestPtr;
-			}
-		}
-	}
-	else if (Format == DXGI_FORMAT_B8G8R8A8_TYPELESS || Format == DXGI_FORMAT_B8G8R8A8_UNORM || Format == DXGI_FORMAT_B8G8R8A8_UNORM_SRGB)
-	{
-		for (uint32 Y = 0; Y < Height; Y++)
-		{
-			FColor* SrcPtr = (FColor*)(In + Y * SrcPitch);
-			FColor* DestPtr = Out + Y * Width;
-
-			// Need to copy row wise since the Pitch might not match the Width.
-			FMemory::Memcpy(DestPtr, SrcPtr, sizeof(FColor) * Width);
-		}
-	}
-	else if (Format == DXGI_FORMAT_R10G10B10A2_UNORM)
-	{
-		// Read the data out of the buffer, converting it from R10G10B10A2 to FColor.
-		for (uint32 Y = 0; Y < Height; Y++)
-		{
-			FD3DR10G10B10A2* SrcPtr = (FD3DR10G10B10A2*)(In + Y * SrcPitch);
-			FColor* DestPtr = Out + Y * Width;
-			for (uint32 X = 0; X < Width; X++)
-			{
-				*DestPtr = FLinearColor(
-					(float)SrcPtr->R / 1023.0f,
-					(float)SrcPtr->G / 1023.0f,
-					(float)SrcPtr->B / 1023.0f,
-					(float)SrcPtr->A / 3.0f
-					).Quantize();
-				++SrcPtr;
-				++DestPtr;
-			}
-		}
-	}
-	else if (Format == DXGI_FORMAT_R16G16B16A16_FLOAT)
-	{
-		FPlane	MinValue(0.0f, 0.0f, 0.0f, 0.0f),
-			MaxValue(1.0f, 1.0f, 1.0f, 1.0f);
-
-		check(sizeof(FFloat16) == sizeof(uint16));
-
-		for (uint32 Y = 0; Y < Height; Y++)
-		{
-			FFloat16* SrcPtr = (FFloat16*)(In + Y * SrcPitch);
-
-			for (uint32 X = 0; X < Width; X++)
-			{
-				MinValue.X = FMath::Min<float>(SrcPtr[0], MinValue.X);
-				MinValue.Y = FMath::Min<float>(SrcPtr[1], MinValue.Y);
-				MinValue.Z = FMath::Min<float>(SrcPtr[2], MinValue.Z);
-				MinValue.W = FMath::Min<float>(SrcPtr[3], MinValue.W);
-				MaxValue.X = FMath::Max<float>(SrcPtr[0], MaxValue.X);
-				MaxValue.Y = FMath::Max<float>(SrcPtr[1], MaxValue.Y);
-				MaxValue.Z = FMath::Max<float>(SrcPtr[2], MaxValue.Z);
-				MaxValue.W = FMath::Max<float>(SrcPtr[3], MaxValue.W);
-				SrcPtr += 4;
-			}
-		}
-
-		for (uint32 Y = 0; Y < Height; Y++)
-		{
-			FFloat16* SrcPtr = (FFloat16*)(In + Y * SrcPitch);
-			FColor* DestPtr = Out + Y * Width;
-
-			for (uint32 X = 0; X < Width; X++)
-			{
-				FColor NormalizedColor =
-					FLinearColor(
-						(SrcPtr[0] - MinValue.X) / (MaxValue.X - MinValue.X),
-						(SrcPtr[1] - MinValue.Y) / (MaxValue.Y - MinValue.Y),
-						(SrcPtr[2] - MinValue.Z) / (MaxValue.Z - MinValue.Z),
-						(SrcPtr[3] - MinValue.W) / (MaxValue.W - MinValue.W)
-						).ToFColor(bLinearToGamma);
-				FMemory::Memcpy(DestPtr++, &NormalizedColor, sizeof(FColor));
-				SrcPtr += 4;
-			}
-		}
-	}
-	else if (Format == DXGI_FORMAT_R11G11B10_FLOAT)
-	{
-		check(sizeof(FFloat3Packed) == sizeof(uint32));
-
-		for (uint32 Y = 0; Y < Height; Y++)
-		{
-			FFloat3Packed* SrcPtr = (FFloat3Packed*)(In + Y * SrcPitch);
-			FColor* DestPtr = Out + Y * Width;
-
-			for (uint32 X = 0; X < Width; X++)
-			{
-				FLinearColor Value = (*SrcPtr).ToLinearColor();
-
-				FColor NormalizedColor = Value.ToFColor(bLinearToGamma);
-				FMemory::Memcpy(DestPtr++, &NormalizedColor, sizeof(FColor));
-				++SrcPtr;
-			}
-		}
-	}
-	else if (Format == DXGI_FORMAT_R32G32B32A32_FLOAT)
-	{
-		FPlane MinValue(0.0f, 0.0f, 0.0f, 0.0f);
-		FPlane MaxValue(1.0f, 1.0f, 1.0f, 1.0f);
-
-		for (uint32 Y = 0; Y < Height; Y++)
-		{
-			float* SrcPtr = (float*)(In + Y * SrcPitch);
-
-			for (uint32 X = 0; X < Width; X++)
-			{
-				MinValue.X = FMath::Min<float>(SrcPtr[0], MinValue.X);
-				MinValue.Y = FMath::Min<float>(SrcPtr[1], MinValue.Y);
-				MinValue.Z = FMath::Min<float>(SrcPtr[2], MinValue.Z);
-				MinValue.W = FMath::Min<float>(SrcPtr[3], MinValue.W);
-				MaxValue.X = FMath::Max<float>(SrcPtr[0], MaxValue.X);
-				MaxValue.Y = FMath::Max<float>(SrcPtr[1], MaxValue.Y);
-				MaxValue.Z = FMath::Max<float>(SrcPtr[2], MaxValue.Z);
-				MaxValue.W = FMath::Max<float>(SrcPtr[3], MaxValue.W);
-				SrcPtr += 4;
-			}
-		}
-
-		for (uint32 Y = 0; Y < Height; Y++)
-		{
-			float* SrcPtr = (float*)In;
-			FColor* DestPtr = Out + Y * Width;
-
-			for (uint32 X = 0; X < Width; X++)
-			{
-				FColor NormalizedColor =
-					FLinearColor(
-						(SrcPtr[0] - MinValue.X) / (MaxValue.X - MinValue.X),
-						(SrcPtr[1] - MinValue.Y) / (MaxValue.Y - MinValue.Y),
-						(SrcPtr[2] - MinValue.Z) / (MaxValue.Z - MinValue.Z),
-						(SrcPtr[3] - MinValue.W) / (MaxValue.W - MinValue.W)
-						).ToFColor(bLinearToGamma);
-				FMemory::Memcpy(DestPtr++, &NormalizedColor, sizeof(FColor));
-				SrcPtr += 4;
-			}
-		}
-	}
-	else if (Format == DXGI_FORMAT_R24G8_TYPELESS)
-	{
-		// Depth stencil
-		for (uint32 Y = 0; Y < Height; Y++)
-		{
-			uint32* SrcPtr = (uint32*)In;
-			FColor* DestPtr = Out + Y * Width;
-
-			for (uint32 X = 0; X < Width; X++)
-			{
-				FColor NormalizedColor;
-				if (InFlags.GetOutputStencil())
-				{
-					uint8 DeviceStencil = (*SrcPtr & 0xFF000000) >> 24;
-					NormalizedColor = FColor(DeviceStencil, DeviceStencil, DeviceStencil, 0xFF);
-				}
-				else
-				{
-					float DeviceZ = (*SrcPtr & 0xffffff) / (float)(1 << 24);
-					float LinearValue = FMath::Min(InFlags.ComputeNormalizedDepth(DeviceZ), 1.0f);
-					NormalizedColor = FLinearColor(LinearValue, LinearValue, LinearValue, 0).ToFColor(bLinearToGamma);
-				}
-
-				FMemory::Memcpy(DestPtr++, &NormalizedColor, sizeof(FColor));
-				++SrcPtr;
-			}
-		}
-	}
-	// Changing Depth Buffers to 32 bit on Dingo as D24S8 is actually implemented as a 32 bit buffer in the hardware
-	else if (Format == DXGI_FORMAT_R32G8X24_TYPELESS)
-	{
-		// Depth stencil
-		for (uint32 Y = 0; Y < Height; Y++)
-		{
-			float* SrcPtr = (float*)(In + Y * SrcPitch);
-			FColor* DestPtr = Out + Y * Width;
-
-			for (uint32 X = 0; X < Width; X++)
-			{
-				float DeviceZ = (*SrcPtr);
-
-				float LinearValue = FMath::Min(InFlags.ComputeNormalizedDepth(DeviceZ), 1.0f);
-
-				FColor NormalizedColor = FLinearColor(LinearValue, LinearValue, LinearValue, 0).ToFColor(bLinearToGamma);
-				FMemory::Memcpy(DestPtr++, &NormalizedColor, sizeof(FColor));
-				SrcPtr+=1; // todo: copies only depth, need to check how this format is read
-				UE_LOG(LogD3D12RHI, Warning, TEXT("CPU read of R32G8X24 is not tested and may not function."));
-			}
-		}
-	}
-	else if (Format == DXGI_FORMAT_R16G16B16A16_UNORM)
-	{
-		// Read the data out of the buffer, converting it to FColor.
-		for (uint32 Y = 0; Y < Height; Y++)
-		{
-			FD3DRGBA16* SrcPtr = (FD3DRGBA16*)(In + Y * SrcPitch);
-			FColor* DestPtr = Out + Y * Width;
-			for (uint32 X = 0; X < Width; X++)
-			{
-				*DestPtr = FLinearColor(
-					(float)SrcPtr->R / 65535.0f,
-					(float)SrcPtr->G / 65535.0f,
-					(float)SrcPtr->B / 65535.0f,
-					(float)SrcPtr->A / 65535.0f
-					).Quantize();
-				++SrcPtr;
-				++DestPtr;
-			}
-		}
-	}
-	else if (Format == DXGI_FORMAT_R16G16_UNORM)
-	{
-		// Read the data out of the buffer, converting it to FColor.
-		for (uint32 Y = 0; Y < Height; Y++)
-		{
-			FD3DRG16* SrcPtr = (FD3DRG16*)(In + Y * SrcPitch);
-			FColor* DestPtr = Out + Y * Width;
-			for (uint32 X = 0; X < Width; X++)
-			{
-				*DestPtr = FLinearColor(
-					(float)SrcPtr->R / 65535.0f,
-					(float)SrcPtr->G / 65535.0f,
-					0).Quantize();
-				++SrcPtr;
-				++DestPtr;
-			}
-		}
-	}
-	else
+	switch (Format)
 	{
-		// not supported yet
-		check(0);
+	case DXGI_FORMAT_R16_TYPELESS:
+		ConvertRawR16DataToFColor(Width, Height, In, SrcPitch, Out);
+		break;
+	case DXGI_FORMAT_R8G8B8A8_TYPELESS:
+	case DXGI_FORMAT_R8G8B8A8_UNORM:
+	case DXGI_FORMAT_R8G8B8A8_UNORM_SRGB:
+		ConvertRawR8G8B8A8DataToFColor(Width, Height, In, SrcPitch, Out);
+		break;
+	case DXGI_FORMAT_B8G8R8A8_TYPELESS:
+	case DXGI_FORMAT_B8G8R8A8_UNORM:
+	case DXGI_FORMAT_B8G8R8A8_UNORM_SRGB:
+		ConvertRawB8G8R8A8DataToFColor(Width, Height, In, SrcPitch, Out);
+		break;
+	case DXGI_FORMAT_R10G10B10A2_UNORM:
+		ConvertRawR10G10B10A2DataToFColor(Width, Height, In, SrcPitch, Out);
+		break;
+	case DXGI_FORMAT_R16G16B16A16_FLOAT:
+		ConvertRawR16G16B16A16FDataToFColor(Width, Height, In, SrcPitch, Out, bLinearToGamma);
+		break;
+	case DXGI_FORMAT_R11G11B10_FLOAT:
+		ConvertRawR11G11B10DataToFColor(Width, Height, In, SrcPitch, Out, bLinearToGamma);
+		break;
+	case DXGI_FORMAT_R32G32B32A32_FLOAT:
+		ConvertRawR32G32B32A32DataToFColor(Width, Height, In, SrcPitch, Out, bLinearToGamma);
+		break;
+	case DXGI_FORMAT_R24G8_TYPELESS:
+		ConvertRawR24G8DataToFColor(Width, Height, In, SrcPitch, Out, InFlags);
+		break;
+	case DXGI_FORMAT_R32G8X24_TYPELESS:
+		ConvertRawR32DataToFColor(Width, Height, In, SrcPitch, Out, InFlags);
+		break;
+	case DXGI_FORMAT_R16G16B16A16_UNORM:
+		ConvertRawR16G16B16A16DataToFColor(Width, Height, In, SrcPitch, Out);
+		break;
+	case DXGI_FORMAT_R16G16_UNORM:
+		ConvertRawR16G16DataToFColor(Width, Height, In, SrcPitch, Out);
+		break;
+	case DXGI_FORMAT_R8_UNORM:
+		ConvertRawR8DataToFColor(Width, Height, In, SrcPitch, Out);
+		break;
+	default:
+		checkf(0, TEXT("Unknown surface format!"));
+		break;
 	}
 }
 
@@ -1134,8 +917,8 @@ void FD3D12DynamicRHI::RHIReadSurfaceData(FRHITexture* TextureRHI, FIntRect InRe
 	FPixelFormatInfo FormatInfo = GPixelFormats[TextureRHI->GetFormat()];
 	uint32 BytesPerPixel = FormatInfo.BlockBytes;
 	uint32 SrcPitch = SizeX * BytesPerPixel;
-
-	ConvertRAWSurfaceDataToFColor((DXGI_FORMAT)FormatInfo.PlatformFormat, SizeX, SizeY, OutDataRaw.GetData(), SrcPitch, OutData.GetData(), InFlags);
+	
+	ConvertDXGIToFColor((DXGI_FORMAT)FormatInfo.PlatformFormat, SizeX, SizeY, OutDataRaw.GetData(), SrcPitch, OutData.GetData(), InFlags);
 }
 
 void FD3D12DynamicRHI::ReadSurfaceDataMSAARaw(FRHICommandList_RecursiveHazardous& RHICmdList, FRHITexture* TextureRHI, FIntRect InRect, TArray<uint8>& OutData, FReadSurfaceDataFlags InFlags)
diff --git a/Engine/Source/Runtime/RHI/Public/RHISurfaceDataConversion.h b/Engine/Source/Runtime/RHI/Public/RHISurfaceDataConversion.h
new file mode 100644
--- /dev/null
+++ b/Engine/Source/Runtime/RHI/Public/RHISurfaceDataConversion.h
@@ -0,0 +1,650 @@
+// Copyright 1998-2019 Epic Games, Inc. All Rights Reserved.
+
+/*=============================================================================
+	RHISurfaceDataConversion.h: RHI surface data conversions.
+=============================================================================*/
+
+#pragma once
+
+#include "Math/PackedVector.h"
+
+/** Helper for accessing R10G10B10A2 colors. */
+struct FRHIR10G10B10A2
+{
+	uint32 R : 10;
+	uint32 G : 10;
+	uint32 B : 10;
+	uint32 A : 2;
+};
+
+/** Helper for accessing R16G16B16A16 colors. */
+struct FRHIRGBA16
+{
+	uint16 R;
+	uint16 G;
+	uint16 B;
+	uint16 A;
+};
+
+/** Helper for accessing R16G16 colors. */
+struct FRHIRG16
+{
+	uint16 R;
+	uint16 G;
+};
+
+static inline void ConvertRawR16DataToFColor(uint32 Width, uint32 Height, uint8 *In, uint32 SrcPitch, FColor* Out)
+{
+	// e.g. shadow maps
+	for (uint32 Y = 0; Y < Height; Y++)
+	{
+		uint16* SrcPtr = (uint16*)(In + Y * SrcPitch);
+		FColor* DestPtr = Out + Y * Width;
+
+		for (uint32 X = 0; X < Width; X++)
+		{
+			uint16 Value16 = *SrcPtr;
+			float Value = Value16 / (float)(0xffff);
+
+			*DestPtr = FLinearColor(Value, Value, Value).Quantize();
+			++SrcPtr;
+			++DestPtr;
+		}
+	}
+}
+
+static inline void ConvertRawR8G8B8A8DataToFColor(uint32 Width, uint32 Height, uint8 *In, uint32 SrcPitch, FColor* Out)
+{
+	for (uint32 Y = 0; Y < Height; Y++)
+	{
+		FColor* SrcPtr = (FColor*)(In + Y * SrcPitch);
+		FColor* DestPtr = Out + Y * Width;
+		for (uint32 X = 0; X < Width; X++)
+		{
+			*DestPtr = FColor(SrcPtr->B, SrcPtr->G, SrcPtr->R, SrcPtr->A);
+			++SrcPtr;
+			++DestPtr;
+		}
+	}
+}
+
+static inline void ConvertRawB8G8R8A8DataToFColor(uint32 Width, uint32 Height, uint8 *In, uint32 SrcPitch, FColor* Out)
+{
+	// todo: check if the source and dst pitch is the same to avoid loop below.
+	for (uint32 Y = 0; Y < Height; Y++)
+	{
+		FColor* SrcPtr = (FColor*)(In + Y * SrcPitch);
+		FColor* DestPtr = Out + Y * Width;
+
+		// Need to copy row wise since the Pitch might not match the Width.
+		FMemory::Memcpy(DestPtr, SrcPtr, sizeof(FColor) * Width);
+	}
+}
+
+static inline void ConvertRawR10G10B10A2DataToFColor(uint32 Width, uint32 Height, uint8 *In, uint32 SrcPitch, FColor* Out)
+{
+	for (uint32 Y = 0; Y < Height; Y++)
+	{
+		FRHIR10G10B10A2* SrcPtr = (FRHIR10G10B10A2*)(In + Y * SrcPitch);
+		FColor* DestPtr = Out + Y * Width;
+		for (uint32 X = 0; X < Width; X++)
+		{
+			*DestPtr = FLinearColor(
+				(float)SrcPtr->R / 1023.0f,
+				(float)SrcPtr->G / 1023.0f,
+				(float)SrcPtr->B / 1023.0f,
+				(float)SrcPtr->A / 3.0f
+			).Quantize();
+			++SrcPtr;
+			++DestPtr;
+		}
+	}
+}
+
+static inline void ConvertRawR16G16B16A16FDataToFColor(uint32 Width, uint32 Height, uint8 *In, uint32 SrcPitch, FColor* Out, bool LinearToGamma)
+{
+	FPlane	MinValue(0.0f, 0.0f, 0.0f, 0.0f),
+		MaxValue(1.0f, 1.0f, 1.0f, 1.0f);
+
+	check(sizeof(FFloat16) == sizeof(uint16));
+
+	for (uint32 Y = 0; Y < Height; Y++)
+	{
+		FFloat16* SrcPtr = (FFloat16*)(In + Y * SrcPitch);
+
+		for (uint32 X = 0; X < Width; X++)
+		{
+			MinValue.X = FMath::Min<float>(SrcPtr[0], MinValue.X);
+			MinValue.Y = FMath::Min<float>(SrcPtr[1], MinValue.Y);
+			MinValue.Z = FMath::Min<float>(SrcPtr[2], MinValue.Z);
+			MinValue.W = FMath::Min<float>(SrcPtr[3], MinValue.W);
+			MaxValue.X = FMath::Max<float>(SrcPtr[0], MaxValue.X);
+			MaxValue.Y = FMath::Max<float>(SrcPtr[1], MaxValue.Y);
+			MaxValue.Z = FMath::Max<float>(SrcPtr[2], MaxValue.Z);
+			MaxValue.W = FMath::Max<float>(SrcPtr[3], MaxValue.W);
+			SrcPtr += 4;
+		}
+	}
+
+	for (uint32 Y = 0; Y < Height; Y++)
+	{
+		FFloat16* SrcPtr = (FFloat16*)(In + Y * SrcPitch);
+		FColor* DestPtr = Out + Y * Width;
+
+		for (uint32 X = 0; X < Width; X++)
+		{
+			FColor NormalizedColor =
+				FLinearColor(
+				(SrcPtr[0] - MinValue.X) / (MaxValue.X - MinValue.X),
+					(SrcPtr[1] - MinValue.Y) / (MaxValue.Y - MinValue.Y),
+					(SrcPtr[2] - MinValue.Z) / (MaxValue.Z - MinValue.Z),
+					(SrcPtr[3] - MinValue.W) / (MaxValue.W - MinValue.W)
+				).ToFColor(LinearToGamma);
+			FMemory::Memcpy(DestPtr++, &NormalizedColor, sizeof(FColor));
+			SrcPtr += 4;
+		}
+	}
+}
+
+static inline void ConvertRawR11G11B10DataToFColor(uint32 Width, uint32 Height, uint8 *In, uint32 SrcPitch, FColor* Out, bool LinearToGamma)
+{
+	check(sizeof(FFloat3Packed) == sizeof(uint32));
+
+	for (uint32 Y = 0; Y < Height; Y++)
+	{
+		FFloat3Packed* SrcPtr = (FFloat3Packed*)(In + Y * SrcPitch);
+		FColor* DestPtr = Out + Y * Width;
+
+		for (uint32 X = 0; X < Width; X++)
+		{
+			FLinearColor Value = (*SrcPtr).ToLinearColor();
+
+			FColor NormalizedColor = Value.ToFColor(LinearToGamma);
+			FMemory::Memcpy(DestPtr++, &NormalizedColor, sizeof(FColor));
+			++SrcPtr;
+		}
+	}
+}
+
+static inline void ConvertRawR32G32B32A32DataToFColor(uint32 Width, uint32 Height, uint8 *In, uint32 SrcPitch, FColor* Out, bool LinearToGamma)
+{
+	FPlane MinValue(0.0f, 0.0f, 0.0f, 0.0f);
+	FPlane MaxValue(1.0f, 1.0f, 1.0f, 1.0f);
+
+	for (uint32 Y = 0; Y < Height; Y++)
+	{
+		float* SrcPtr = (float*)(In + Y * SrcPitch);
+
+		for (uint32 X = 0; X < Width; X++)
+		{
+			MinValue.X = FMath::Min<float>(SrcPtr[0], MinValue.X);
+			MinValue.Y = FMath::Min<float>(SrcPtr[1], MinValue.Y);
+			MinValue.Z = FMath::Min<float>(SrcPtr[2], MinValue.Z);
+			MinValue.W = FMath::Min<float>(SrcPtr[3], MinValue.W);
+			MaxValue.X = FMath::Max<float>(SrcPtr[0], MaxValue.X);
+			MaxValue.Y = FMath::Max<float>(SrcPtr[1], MaxValue.Y);
+			MaxValue.Z = FMath::Max<float>(SrcPtr[2], MaxValue.Z);
+			MaxValue.W = FMath::Max<float>(SrcPtr[3], MaxValue.W);
+			SrcPtr += 4;
+		}
+	}
+
+	for (uint32 Y = 0; Y < Height; Y++)
+	{
+		float* SrcPtr = (float*)In;
+		FColor* DestPtr = Out + Y * Width;
+
+		for (uint32 X = 0; X < Width; X++)
+		{
+			FColor NormalizedColor =
+				FLinearColor(
+				(SrcPtr[0] - MinValue.X) / (MaxValue.X - MinValue.X),
+					(SrcPtr[1] - MinValue.Y) / (MaxValue.Y - MinValue.Y),
+					(SrcPtr[2] - MinValue.Z) / (MaxValue.Z - MinValue.Z),
+					(SrcPtr[3] - MinValue.W) / (MaxValue.W - MinValue.W)
+				).ToFColor(LinearToGamma);
+			FMemory::Memcpy(DestPtr++, &NormalizedColor, sizeof(FColor));
+			SrcPtr += 4;
+		}
+	}
+}
+
+static inline void ConvertRawR24G8DataToFColor(uint32 Width, uint32 Height, uint8 *In, uint32 SrcPitch, FColor* Out, FReadSurfaceDataFlags InFlags)
+{
+	bool bLinearToGamma = InFlags.GetLinearToGamma();
+	// Depth stencil
+	for (uint32 Y = 0; Y < Height; Y++)
+	{
+		uint32* SrcPtr = (uint32 *)In;
+		FColor* DestPtr = Out + Y * Width;
+
+		for (uint32 X = 0; X < Width; X++)
+		{
+			FColor NormalizedColor;
+			if (InFlags.GetOutputStencil())
+			{
+				uint8 DeviceStencil = (*SrcPtr & 0xFF000000) >> 24;
+				NormalizedColor = FColor(DeviceStencil, DeviceStencil, DeviceStencil, 0xFF);
+			}
+			else
+			{
+				float DeviceZ = (*SrcPtr & 0xffffff) / (float)(1 << 24);
+				float LinearValue = FMath::Min(InFlags.ComputeNormalizedDepth(DeviceZ), 1.0f);
+				NormalizedColor = FLinearColor(LinearValue, LinearValue, LinearValue, 0).ToFColor(bLinearToGamma);
+			}
+
+			FMemory::Memcpy(DestPtr++, &NormalizedColor, sizeof(FColor));
+			++SrcPtr;
+		}
+	}
+}
+
+static inline void ConvertRawR32DataToFColor(uint32 Width, uint32 Height, uint8 *In, uint32 SrcPitch, FColor* Out, FReadSurfaceDataFlags InFlags)
+{
+	bool bLinearToGamma = InFlags.GetLinearToGamma();
+	for (uint32 Y = 0; Y < Height; Y++)
+	{
+		float* SrcPtr = (float *)(In + Y * SrcPitch);
+		FColor* DestPtr = Out + Y * Width;
+
+		for (uint32 X = 0; X < Width; X++)
+		{
+			float DeviceZ = (*SrcPtr);
+
+			float LinearValue = FMath::Min(InFlags.ComputeNormalizedDepth(DeviceZ), 1.0f);
+
+			FColor NormalizedColor = FLinearColor(LinearValue, LinearValue, LinearValue, 0).ToFColor(bLinearToGamma);
+			FMemory::Memcpy(DestPtr++, &NormalizedColor, sizeof(FColor));
+			SrcPtr += 1; // todo: copies only depth, need to check how this format is read
+			UE_LOG(LogRHI, Warning, TEXT("CPU read of R32G8X24 is not tested and may not function."));
+		}
+	}
+}
+
+static inline void ConvertRawR16G16B16A16DataToFColor(uint32 Width, uint32 Height, uint8 *In, uint32 SrcPitch, FColor* Out)
+{
+	for (uint32 Y = 0; Y < Height; Y++)
+	{
+		FRHIRGBA16* SrcPtr = (FRHIRGBA16*)(In + Y * SrcPitch);
+		FColor* DestPtr = Out + Y * Width;
+		for (uint32 X = 0; X < Width; X++)
+		{
+			*DestPtr = FLinearColor(
+				(float)SrcPtr->R / 65535.0f,
+				(float)SrcPtr->G / 65535.0f,
+				(float)SrcPtr->B / 65535.0f,
+				(float)SrcPtr->A / 65535.0f
+			).Quantize();
+			++SrcPtr;
+			++DestPtr;
+		}
+	}
+}
+
+static inline void ConvertRawR16G16DataToFColor(uint32 Width, uint32 Height, uint8 *In, uint32 SrcPitch, FColor* Out)
+{
+	for (uint32 Y = 0; Y < Height; Y++)
+	{
+		FRHIRG16* SrcPtr = (FRHIRG16*)(In + Y * SrcPitch);
+		FColor* DestPtr = Out + Y * Width;
+		for (uint32 X = 0; X < Width; X++)
+		{
+			*DestPtr = FLinearColor(
+				(float)SrcPtr->R / 65535.0f,
+				(float)SrcPtr->G / 65535.0f,
+				0).Quantize();
+			++SrcPtr;
+			++DestPtr;
+		}
+	}
+}
+
+static inline void ConvertRawR8DataToFColor(uint32 Width, uint32 Height, uint8 *In, uint32 SrcPitch, FColor* Out)
+{
+	for (uint32 Y = 0; Y < Height; Y++)
+	{
+		uint8* SrcPtr = (uint8*)(In + Y * SrcPitch);
+		FColor* DestPtr = Out + Y * Width;
+		for (uint32 X = 0; X < Width; X++)
+		{
+			*DestPtr = FColor(*SrcPtr, *SrcPtr, *SrcPtr, *SrcPtr);
+			++SrcPtr;
+			++DestPtr;
+		}
+	}
+}
+
+static inline void ConvertRawD32S8DataToFColor(uint32 Width, uint32 Height, uint8 *In, uint32 SrcPitch, FColor* Out, FReadSurfaceDataFlags InFlags)
+{
+	bool bLinearToGamma = InFlags.GetLinearToGamma();
+	// Depth
+	if (!InFlags.GetOutputStencil())
+	{
+		for (uint32 Y = 0; Y < Height; Y++)
+		{
+			uint32* SrcPtr = (uint32*)(In + Y * SrcPitch);
+			FColor* DestPtr = Out + Y * Width;
+
+			for (uint32 X = 0; X < Width; X++)
+			{
+				float DeviceZ = *SrcPtr;
+				float LinearValue = FMath::Min(InFlags.ComputeNormalizedDepth(DeviceZ), 1.0f);
+				FColor NormalizedColor = FLinearColor(LinearValue, LinearValue, LinearValue, 0).ToFColor(bLinearToGamma);
+
+				FMemory::Memcpy(DestPtr++, &NormalizedColor, sizeof(FColor));
+				++SrcPtr;
+			}
+		}
+	}
+	// Stencil
+	else
+	{
+		// Depth stencil
+		for (uint32 Y = 0; Y < Height; Y++)
+		{
+			uint8* SrcPtr = (uint8*)(In + Y * SrcPitch);
+			FColor* DestPtr = Out + Y * Width;
+
+			for (uint32 X = 0; X < Width; X++)
+			{
+				uint8 DeviceStencil = *SrcPtr;
+				FColor NormalizedColor = FColor(DeviceStencil, DeviceStencil, DeviceStencil, 0xFF);
+
+				FMemory::Memcpy(DestPtr++, &NormalizedColor, sizeof(FColor));
+				++SrcPtr;
+			}
+		}
+	}
+}
+
+// Linear functions
+
+static inline void ConvertRawR16DataToFLinearColor(uint32 Width, uint32 Height, uint8 *In, uint32 SrcPitch, FLinearColor* Out)
+{
+	// e.g. shadow maps
+	for (uint32 Y = 0; Y < Height; Y++)
+	{
+		uint16* SrcPtr = (uint16*)(In + Y * SrcPitch);
+		FLinearColor* DestPtr = Out + Y * Width;
+
+		for (uint32 X = 0; X < Width; X++)
+		{
+			uint16 Value16 = *SrcPtr;
+			float Value = Value16 / (float)(0xffff);
+
+			*DestPtr = FLinearColor(Value, Value, Value);
+			++SrcPtr;
+			++DestPtr;
+		}
+	}
+}
+
+static inline void ConvertRawR8G8B8A8DataToFLinearColor(uint32 Width, uint32 Height, uint8 *In, uint32 SrcPitch, FLinearColor* Out)
+{
+	// Read the data out of the buffer, converting it from ABGR to ARGB.
+	for (uint32 Y = 0; Y < Height; Y++)
+	{
+		FColor* SrcPtr = (FColor*)(In + Y * SrcPitch);
+		FLinearColor* DestPtr = Out + Y * Width;
+		for (uint32 X = 0; X < Width; X++)
+		{
+			FColor sRGBColor = FColor(SrcPtr->B, SrcPtr->G, SrcPtr->R, SrcPtr->A);
+			*DestPtr = FLinearColor(sRGBColor);
+			++SrcPtr;
+			++DestPtr;
+		}
+	}
+}
+
+static inline void ConvertRawB8G8R8A8DataToFLinearColor(uint32 Width, uint32 Height, uint8 *In, uint32 SrcPitch, FLinearColor* Out)
+{
+	for (uint32 Y = 0; Y < Height; Y++)
+	{
+		FColor* SrcPtr = (FColor*)(In + Y * SrcPitch);
+		FLinearColor* DestPtr = Out + Y * Width;
+		for (uint32 X = 0; X < Width; X++)
+		{
+			FColor sRGBColor = FColor(SrcPtr->R, SrcPtr->G, SrcPtr->B, SrcPtr->A);
+			*DestPtr = FLinearColor(sRGBColor);
+			++SrcPtr;
+			++DestPtr;
+		}
+	}
+}
+
+static inline void ConvertRawA2B10G10R10DataToFLinearColor(uint32 Width, uint32 Height, uint8 *In, uint32 SrcPitch, FLinearColor* Out)
+{
+	// Read the data out of the buffer, converting it from R10G10B10A2 to FLinearColor.
+	for (uint32 Y = 0; Y < Height; Y++)
+	{
+		FRHIR10G10B10A2* SrcPtr = (FRHIR10G10B10A2*)(In + Y * SrcPitch);
+		FLinearColor* DestPtr = Out + Y * Width;
+		for (uint32 X = 0; X < Width; X++)
+		{
+			*DestPtr = FLinearColor(
+				(float)SrcPtr->R / 1023.0f,
+				(float)SrcPtr->G / 1023.0f,
+				(float)SrcPtr->B / 1023.0f,
+				(float)SrcPtr->A / 3.0f
+			);
+			++SrcPtr;
+			++DestPtr;
+		}
+	}
+}
+
+static inline void ConvertRawR16G16B16A16FDataToFLinearColor(uint32 Width, uint32 Height, uint8 *In, uint32 SrcPitch, FLinearColor* Out, FReadSurfaceDataFlags InFlags)
+{
+	if (InFlags.GetCompressionMode() == RCM_MinMax)
+	{
+		for (uint32 Y = 0; Y < Height; Y++)
+		{
+			FFloat16* SrcPtr = (FFloat16*)(In + Y * SrcPitch);
+			FLinearColor* DestPtr = Out + Y * Width;
+
+			for (uint32 X = 0; X < Width; X++)
+			{
+				*DestPtr = FLinearColor((float)SrcPtr[0], (float)SrcPtr[1], (float)SrcPtr[2], (float)SrcPtr[3]);
+				++DestPtr;
+				SrcPtr += 4;
+			}
+		}
+	}
+	else
+	{
+		FPlane	MinValue(0.0f, 0.0f, 0.0f, 0.0f);
+		FPlane	MaxValue(1.0f, 1.0f, 1.0f, 1.0f);
+
+		check(sizeof(FFloat16) == sizeof(uint16));
+
+		for (uint32 Y = 0; Y < Height; Y++)
+		{
+			FFloat16* SrcPtr = (FFloat16*)(In + Y * SrcPitch);
+
+			for (uint32 X = 0; X < Width; X++)
+			{
+				MinValue.X = FMath::Min<float>(SrcPtr[0], MinValue.X);
+				MinValue.Y = FMath::Min<float>(SrcPtr[1], MinValue.Y);
+				MinValue.Z = FMath::Min<float>(SrcPtr[2], MinValue.Z);
+				MinValue.W = FMath::Min<float>(SrcPtr[3], MinValue.W);
+				MaxValue.X = FMath::Max<float>(SrcPtr[0], MaxValue.X);
+				MaxValue.Y = FMath::Max<float>(SrcPtr[1], MaxValue.Y);
+				MaxValue.Z = FMath::Max<float>(SrcPtr[2], MaxValue.Z);
+				MaxValue.W = FMath::Max<float>(SrcPtr[3], MaxValue.W);
+				SrcPtr += 4;
+			}
+		}
+
+		for (uint32 Y = 0; Y < Height; Y++)
+		{
+			FFloat16* SrcPtr = (FFloat16*)(In + Y * SrcPitch);
+			FLinearColor* DestPtr = Out + Y * Width;
+
+			for (uint32 X = 0; X < Width; X++)
+			{
+				*DestPtr = FLinearColor(
+					(SrcPtr[0] - MinValue.X) / (MaxValue.X - MinValue.X),
+					(SrcPtr[1] - MinValue.Y) / (MaxValue.Y - MinValue.Y),
+					(SrcPtr[2] - MinValue.Z) / (MaxValue.Z - MinValue.Z),
+					(SrcPtr[3] - MinValue.W) / (MaxValue.W - MinValue.W)
+				);
+				++DestPtr;
+				SrcPtr += 4;
+			}
+		}
+	}
+}
+
+static inline void ConvertRawRR11G11B10DataToFLinearColor(uint32 Width, uint32 Height, uint8 *In, uint32 SrcPitch, FLinearColor* Out)
+{
+	check(sizeof(FFloat3Packed) == sizeof(uint32));
+
+	for (uint32 Y = 0; Y < Height; Y++)
+	{
+		FFloat3Packed* SrcPtr = (FFloat3Packed*)(In + Y * SrcPitch);
+		FLinearColor* DestPtr = Out + Y * Width;
+
+		for (uint32 X = 0; X < Width; X++)
+		{
+			*DestPtr = (*SrcPtr).ToLinearColor();
+			++DestPtr;
+			++SrcPtr;
+		}
+	}
+}
+
+static inline void ConvertRawR32G32B32A32DataToFLinearColor(uint32 Width, uint32 Height, uint8 *In, uint32 SrcPitch, FLinearColor* Out, FReadSurfaceDataFlags InFlags)
+{
+	if (InFlags.GetCompressionMode() == RCM_MinMax)
+	{
+		// Copy data directly, respecting existing min-max values
+		FLinearColor* SrcPtr = (FLinearColor*)In;
+		FLinearColor* DestPtr = (FLinearColor*)Out;
+		const int32 ImageSize = sizeof(FLinearColor) * Height * Width;
+
+		FMemory::Memcpy(DestPtr, SrcPtr, ImageSize);
+	}
+	else
+	{
+		// Normalize data
+		FPlane MinValue(0.0f, 0.0f, 0.0f, 0.0f);
+		FPlane MaxValue(1.0f, 1.0f, 1.0f, 1.0f);
+
+		for (uint32 Y = 0; Y < Height; Y++)
+		{
+			float* SrcPtr = (float*)(In + Y * SrcPitch);
+
+			for (uint32 X = 0; X < Width; X++)
+			{
+				MinValue.X = FMath::Min<float>(SrcPtr[0], MinValue.X);
+				MinValue.Y = FMath::Min<float>(SrcPtr[1], MinValue.Y);
+				MinValue.Z = FMath::Min<float>(SrcPtr[2], MinValue.Z);
+				MinValue.W = FMath::Min<float>(SrcPtr[3], MinValue.W);
+				MaxValue.X = FMath::Max<float>(SrcPtr[0], MaxValue.X);
+				MaxValue.Y = FMath::Max<float>(SrcPtr[1], MaxValue.Y);
+				MaxValue.Z = FMath::Max<float>(SrcPtr[2], MaxValue.Z);
+				MaxValue.W = FMath::Max<float>(SrcPtr[3], MaxValue.W);
+				SrcPtr += 4;
+			}
+		}
+
+		float* SrcPtr = (float*)In;
+
+		for (uint32 Y = 0; Y < Height; Y++)
+		{
+			FLinearColor* DestPtr = Out + Y * Width;
+
+			for (uint32 X = 0; X < Width; X++)
+			{
+				*DestPtr = FLinearColor(
+					(SrcPtr[0] - MinValue.X) / (MaxValue.X - MinValue.X),
+					(SrcPtr[1] - MinValue.Y) / (MaxValue.Y - MinValue.Y),
+					(SrcPtr[2] - MinValue.Z) / (MaxValue.Z - MinValue.Z),
+					(SrcPtr[3] - MinValue.W) / (MaxValue.W - MinValue.W)
+				);
+				++DestPtr;
+				SrcPtr += 4;
+			}
+		}
+	}
+}
+
+static inline void ConvertRawR24G8DataToFLinearColor(uint32 Width, uint32 Height, uint8 *In, uint32 SrcPitch, FLinearColor* Out, FReadSurfaceDataFlags InFlags)
+{
+	// Depth stencil
+	for (uint32 Y = 0; Y < Height; Y++)
+	{
+		uint32* SrcPtr = (uint32 *)In;
+		FLinearColor* DestPtr = Out + Y * Width;
+
+		for (uint32 X = 0; X < Width; X++)
+		{
+			float DeviceStencil = 0.0f;
+			DeviceStencil = (float)((*SrcPtr & 0xFF000000) >> 24) / 255.0f;
+			float DeviceZ = (*SrcPtr & 0xffffff) / (float)(1 << 24);
+			float LinearValue = FMath::Min(InFlags.ComputeNormalizedDepth(DeviceZ), 1.0f);
+			*DestPtr = FLinearColor(LinearValue, DeviceStencil, 0.0f, 0.0f);
+			++DestPtr;
+			++SrcPtr;
+		}
+	}
+}
+
+static inline void ConvertRawR32DataToFLinearColor(uint32 Width, uint32 Height, uint8 *In, uint32 SrcPitch, FLinearColor* Out, FReadSurfaceDataFlags InFlags)
+{
+	// Depth stencil
+	for (uint32 Y = 0; Y < Height; Y++)
+	{
+		uint8* SrcStart = (uint8 *)(In + Y * SrcPitch);
+		FLinearColor* DestPtr = Out + Y * Width;
+
+		for (uint32 X = 0; X < Width; X++)
+		{
+			float DeviceZ = *((float *)(SrcStart));
+			float LinearValue = FMath::Min(InFlags.ComputeNormalizedDepth(DeviceZ), 1.0f);
+			float DeviceStencil = (float)(*(SrcStart + 4)) / 255.0f;
+			*DestPtr = FLinearColor(LinearValue, DeviceStencil, 0.0f, 0.0f);
+			SrcStart += 8; //64 bit format with the last 24 bit ignore
+		}
+	}
+}
+
+static inline void ConvertRawR16G16B16A16DataToFLinearColor(uint32 Width, uint32 Height, uint8 *In, uint32 SrcPitch, FLinearColor* Out)
+{
+	// Read the data out of the buffer, converting it to FLinearColor.
+	for (uint32 Y = 0; Y < Height; Y++)
+	{
+		FRHIRGBA16* SrcPtr = (FRHIRGBA16*)(In + Y * SrcPitch);
+		FLinearColor* DestPtr = Out + Y * Width;
+		for (uint32 X = 0; X < Width; X++)
+		{
+			*DestPtr = FLinearColor(
+				(float)SrcPtr->R / 65535.0f,
+				(float)SrcPtr->G / 65535.0f,
+				(float)SrcPtr->B / 65535.0f,
+				(float)SrcPtr->A / 65535.0f
+			);
+			++SrcPtr;
+			++DestPtr;
+		}
+	}
+}
+
+static inline void ConvertRawR16G16DataToFLinearColor(uint32 Width, uint32 Height, uint8 *In, uint32 SrcPitch, FLinearColor* Out)
+{
+	// Read the data out of the buffer, converting it to FLinearColor.
+	for (uint32 Y = 0; Y < Height; Y++)
+	{
+		FRHIRG16* SrcPtr = (FRHIRG16*)(In + Y * SrcPitch);
+		FLinearColor* DestPtr = Out + Y * Width;
+		for (uint32 X = 0; X < Width; X++)
+		{
+			*DestPtr = FLinearColor(
+				(float)SrcPtr->R / 65535.0f,
+				(float)SrcPtr->G / 65535.0f,
+				0);
+			++SrcPtr;
+			++DestPtr;
+		}
+	}
+}
\ No newline at end of file
diff --git a/Engine/Source/Runtime/VulkanRHI/Private/VulkanRenderTarget.cpp b/Engine/Source/Runtime/VulkanRHI/Private/VulkanRenderTarget.cpp
--- a/Engine/Source/Runtime/VulkanRHI/Private/VulkanRenderTarget.cpp
+++ b/Engine/Source/Runtime/VulkanRHI/Private/VulkanRenderTarget.cpp
@@ -9,6 +9,7 @@
 #include "VulkanPendingState.h"
 #include "VulkanContext.h"
 #include "SceneUtils.h"
+#include "RHISurfaceDataConversion.h"
 
 static int32 GSubmitOnCopyToResolve = 0;
 static FAutoConsoleVariableRef CVarVulkanSubmitOnCopyToResolve(
@@ -751,6 +752,7 @@ void FVulkanCommandListContext::RHICopyToResolveTarget(FRHITexture* SourceTextur
 	}
 }
 
+
 void FVulkanDynamicRHI::RHIReadSurfaceData(FRHITexture* TextureRHI, FIntRect Rect, TArray<FColor>& OutData, FReadSurfaceDataFlags InFlags)
 {
 	FRHITexture2D* TextureRHI2D = TextureRHI->GetTexture2D();
@@ -839,96 +841,38 @@ void FVulkanDynamicRHI::RHIReadSurfaceData(FRHITexture* TextureRHI, FIntRect Rec
 
 	if (Texture2D->Surface.StorageFormat == VK_FORMAT_R16G16B16A16_SFLOAT)
 	{
-		for (int32 Row = Rect.Min.Y; Row < Rect.Max.Y; ++Row)
-		{
-			FFloat16Color* Src = (FFloat16Color*)StagingBuffer->GetMappedPointer() + Row * TextureRHI2D->GetSizeX() + Rect.Min.X;
-			for (int32 Col = Rect.Min.X; Col < Rect.Max.X; ++Col)
-			{
-				Dest->R = (uint8)(uint32)FMath::Clamp<int32>((int32)(Src->R.GetFloat() * 255.0f), 0, 255);
-				Dest->G = (uint8)(uint32)FMath::Clamp<int32>((int32)(Src->G.GetFloat() * 255.0f), 0, 255);
-				Dest->B = (uint8)(uint32)FMath::Clamp<int32>((int32)(Src->B.GetFloat() * 255.0f), 0, 255);
-				Dest->A = (uint8)(uint32)FMath::Clamp<int32>((int32)(Src->A.GetFloat() * 255.0f), 0, 255);
-				Dest++;
-				Src++;
-			}
-		}
+		uint32 PixelByteSize = 8u;
+		uint8* In = (uint8*)StagingBuffer->GetMappedPointer() + (Rect.Min.Y * TextureRHI2D->GetSizeX() + Rect.Min.X) * PixelByteSize;
+		uint32 SrcPitch = TextureRHI2D->GetSizeX() * PixelByteSize;
+		ConvertRawR16G16B16A16FDataToFColor(TextureRHI2D->GetSizeX(), TextureRHI2D->GetSizeY(), In, SrcPitch, Dest, false);
 	}
 	else if (Texture2D->Surface.StorageFormat == VK_FORMAT_A2B10G10R10_UNORM_PACK32)
 	{
-		struct FR10G10B10A2
-		{
-			uint32 R : 10;
-			uint32 G : 10;
-			uint32 B : 10;
-			uint32 A : 2;
-		};
-		for (int32 Row = Rect.Min.Y; Row < Rect.Max.Y; ++Row)
-		{
-			FR10G10B10A2* Src = (FR10G10B10A2*)StagingBuffer->GetMappedPointer() + Row * TextureRHI2D->GetSizeX() + Rect.Min.X;
-			for (int32 Col = Rect.Min.X; Col < Rect.Max.X; ++Col)
-			{
-				*Dest = FLinearColor(
-					(float)Src->R / 1023.0f,
-					(float)Src->G / 1023.0f,
-					(float)Src->B / 1023.0f,
-					(float)Src->A / 3.0f
-				).Quantize();
-				++Dest;
-				++Src;
-			}
-		}
+		uint32 PixelByteSize = 4u;
+		uint8* In = (uint8*)StagingBuffer->GetMappedPointer() + (Rect.Min.Y * TextureRHI2D->GetSizeX() + Rect.Min.X) * PixelByteSize;
+		uint32 SrcPitch = TextureRHI2D->GetSizeX() * PixelByteSize;
+		ConvertRawR10G10B10A2DataToFColor(TextureRHI2D->GetSizeX(), TextureRHI2D->GetSizeY(), In, SrcPitch, Dest);
 	}
 	else if (Texture2D->Surface.StorageFormat == VK_FORMAT_R8G8B8A8_UNORM)
 	{
-		for (int32 Row = Rect.Min.Y; Row < Rect.Max.Y; ++Row)
-		{
-			FColor* Src = (FColor*)StagingBuffer->GetMappedPointer() + Row * TextureRHI2D->GetSizeX() + Rect.Min.X;
-			for (int32 Col = Rect.Min.X; Col < Rect.Max.X; ++Col)
-			{
-				Dest->R = Src->B;
-				Dest->G = Src->G;
-				Dest->B = Src->R;
-				Dest->A = Src->A;
-				Dest++;
-				Src++;
-			}
-		}
+		uint32 PixelByteSize = 4u;
+		uint8* In = (uint8*)StagingBuffer->GetMappedPointer() + (Rect.Min.Y * TextureRHI2D->GetSizeX() + Rect.Min.X) * PixelByteSize;
+		uint32 SrcPitch = TextureRHI2D->GetSizeX() * PixelByteSize;
+		ConvertRawR8G8B8A8DataToFColor(TextureRHI2D->GetSizeX(), TextureRHI2D->GetSizeY(), In, SrcPitch, Dest);
 	}
 	else if (Texture2D->Surface.StorageFormat == VK_FORMAT_R16G16B16A16_UNORM)
 	{
-		struct FRGBA16U
-		{
-			uint16 R;
-			uint16 G;
-			uint16 B;
-			uint16 A;
-		};
-		for (int32 Row = Rect.Min.Y; Row < Rect.Max.Y; ++Row)
-		{
-			FRGBA16U* Src = (FRGBA16U*)StagingBuffer->GetMappedPointer() + Row * TextureRHI2D->GetSizeX() + Rect.Min.X;
-			for (int32 Col = Rect.Min.X; Col < Rect.Max.X; ++Col)
-			{
-				*Dest = FLinearColor(
-					(float)Src->R / 65535.0f,
-					(float)Src->G / 65535.0f,
-					(float)Src->B / 65535.0f,
-					(float)Src->A / 65535.0f
-				).Quantize();
-				Dest++;
-				Src++;
-			}
-		}
+		uint32 PixelByteSize = 8u;
+		uint8* In = (uint8*)StagingBuffer->GetMappedPointer() + (Rect.Min.Y * TextureRHI2D->GetSizeX() + Rect.Min.X) * PixelByteSize;
+		uint32 SrcPitch = TextureRHI2D->GetSizeX() * PixelByteSize;
+		ConvertRawR16G16B16A16DataToFColor(Texture2D->GetSizeX(), Texture2D->GetSizeY(), In, SrcPitch, Dest);
 	}
 	else if (Texture2D->Surface.StorageFormat == VK_FORMAT_B8G8R8A8_UNORM)
 	{
-		FColor* Src = (FColor*)StagingBuffer->GetMappedPointer() + Rect.Min.Y * TextureRHI2D->GetSizeX() + Rect.Min.X;
-		for (int32 Row = Rect.Min.Y; Row < Rect.Max.Y; ++Row)
-		{
-			const int32 NumCols = Rect.Max.X - Rect.Min.X;
-			FMemory::Memcpy(Dest, Src, NumCols * sizeof(FColor));
-			Src += TextureRHI2D->GetSizeX();
-			Dest += NumCols;
-		}
+		uint32 PixelByteSize = 4u;
+		uint8* In = (uint8*)StagingBuffer->GetMappedPointer() + (Rect.Min.Y * TextureRHI2D->GetSizeX() + Rect.Min.X) * PixelByteSize;
+		uint32 SrcPitch = TextureRHI2D->GetSizeX() * PixelByteSize;
+		ConvertRawR8G8B8A8DataToFColor(TextureRHI2D->GetSizeX(), TextureRHI2D->GetSizeY(), In, SrcPitch, Dest);
 	}
 
 	Device->GetStagingManager().ReleaseBuffer(CmdBuffer, StagingBuffer);
diff --git a/Engine/Source/Runtime/Windows/D3D11RHI/Private/D3D11RenderTarget.cpp b/Engine/Source/Runtime/Windows/D3D11RHI/Private/D3D11RenderTarget.cpp
--- a/Engine/Source/Runtime/Windows/D3D11RHI/Private/D3D11RenderTarget.cpp
+++ b/Engine/Source/Runtime/Windows/D3D11RHI/Private/D3D11RenderTarget.cpp
@@ -11,6 +11,7 @@
 #include "ResolveShader.h"
 #include "PipelineStateCache.h"
 #include "Math/PackedVector.h"
+#include "RHISurfaceDataConversion.h"
 
 static inline DXGI_FORMAT ConvertTypelessToUnorm(DXGI_FORMAT Format)
 {
@@ -613,288 +614,55 @@ struct FD3DRGBA16
 	uint16 A;
 };
 
-// todo: this should be available for all RHI
-static void ConvertRAWSurfaceDataToFColor(DXGI_FORMAT Format, uint32 Width, uint32 Height, uint8 *In, uint32 SrcPitch, FColor* Out, FReadSurfaceDataFlags InFlags)
+/** Convert D3D format type to general pixel format type*/
+static void ConvertDXGIToFColor(DXGI_FORMAT Format, uint32 Width, uint32 Height, uint8 *In, uint32 SrcPitch, FColor* Out, FReadSurfaceDataFlags InFlags)
 {
 	bool bLinearToGamma = InFlags.GetLinearToGamma();
-
-	if(Format == DXGI_FORMAT_R16_TYPELESS)
-	{
-		// e.g. shadow maps
-		for(uint32 Y = 0; Y < Height; Y++)
-		{
-			uint16* SrcPtr = (uint16*)(In + Y * SrcPitch);
-			FColor* DestPtr = Out + Y * Width;
-
-			for(uint32 X = 0; X < Width; X++)
-			{
-				uint16 Value16 = *SrcPtr;
-				float Value = Value16 / (float)(0xffff);
-
-				*DestPtr = FLinearColor(Value, Value, Value).Quantize();
-				++SrcPtr;
-				++DestPtr;
-			}
-		}
-	}
-	else if(Format == DXGI_FORMAT_R8G8B8A8_TYPELESS || Format == DXGI_FORMAT_R8G8B8A8_UNORM || Format == DXGI_FORMAT_R8G8B8A8_UNORM_SRGB)
-	{
-		// Read the data out of the buffer, converting it from ABGR to ARGB.
-		for(uint32 Y = 0; Y < Height; Y++)
-		{
-			FColor* SrcPtr = (FColor*)(In + Y * SrcPitch);
-			FColor* DestPtr = Out + Y * Width;
-			for(uint32 X = 0; X < Width; X++)
-			{
-				*DestPtr = FColor(SrcPtr->B,SrcPtr->G,SrcPtr->R,SrcPtr->A);
-				++SrcPtr;
-				++DestPtr;
-			}
-		}
-	}
-	else if(Format == DXGI_FORMAT_B8G8R8A8_TYPELESS || Format == DXGI_FORMAT_B8G8R8A8_UNORM || Format == DXGI_FORMAT_B8G8R8A8_UNORM_SRGB)
-	{
-		for(uint32 Y = 0; Y < Height; Y++)
-		{
-			FColor* SrcPtr = (FColor*)(In + Y * SrcPitch);
-			FColor* DestPtr = Out + Y * Width;
-
-			// Need to copy row wise since the Pitch might not match the Width.
-			FMemory::Memcpy(DestPtr, SrcPtr, sizeof(FColor) * Width);
-		}
-	}
-	else if(Format == DXGI_FORMAT_R10G10B10A2_UNORM)
-	{
-		// Read the data out of the buffer, converting it from R10G10B10A2 to FColor.
-		for(uint32 Y = 0; Y < Height; Y++)
-		{
-			FD3DR10G10B10A2* SrcPtr = (FD3DR10G10B10A2*)(In + Y * SrcPitch);
-			FColor* DestPtr = Out + Y * Width;
-			for(uint32 X = 0; X < Width; X++)
-			{
-				*DestPtr = FLinearColor(
-					(float)SrcPtr->R / 1023.0f,
-					(float)SrcPtr->G / 1023.0f,
-					(float)SrcPtr->B / 1023.0f,
-					(float)SrcPtr->A / 3.0f
-					).Quantize();
-				++SrcPtr;
-				++DestPtr;
-			}
-		}
-	}
-	else if(Format == DXGI_FORMAT_R16G16B16A16_FLOAT)
-	{
-		FPlane	MinValue(0.0f,0.0f,0.0f,0.0f),
-			MaxValue(1.0f,1.0f,1.0f,1.0f);
-
-		check(sizeof(FFloat16)==sizeof(uint16));
-
-		for(uint32 Y = 0; Y < Height; Y++)
-		{
-			FFloat16* SrcPtr = (FFloat16*)(In + Y * SrcPitch);
-
-			for(uint32 X = 0; X < Width; X++)
-			{
-				MinValue.X = FMath::Min<float>(SrcPtr[0],MinValue.X);
-				MinValue.Y = FMath::Min<float>(SrcPtr[1],MinValue.Y);
-				MinValue.Z = FMath::Min<float>(SrcPtr[2],MinValue.Z);
-				MinValue.W = FMath::Min<float>(SrcPtr[3],MinValue.W);
-				MaxValue.X = FMath::Max<float>(SrcPtr[0],MaxValue.X);
-				MaxValue.Y = FMath::Max<float>(SrcPtr[1],MaxValue.Y);
-				MaxValue.Z = FMath::Max<float>(SrcPtr[2],MaxValue.Z);
-				MaxValue.W = FMath::Max<float>(SrcPtr[3],MaxValue.W);
-				SrcPtr += 4;
-			}
-		}
-
-		for(uint32 Y = 0; Y < Height; Y++)
-		{
-			FFloat16* SrcPtr = (FFloat16*)(In + Y * SrcPitch);
-			FColor* DestPtr = Out + Y * Width;
-
-			for(uint32 X = 0; X < Width; X++)
-			{
-				FColor NormalizedColor =
-					FLinearColor(
-					(SrcPtr[0] - MinValue.X) / (MaxValue.X - MinValue.X),
-					(SrcPtr[1] - MinValue.Y) / (MaxValue.Y - MinValue.Y),
-					(SrcPtr[2] - MinValue.Z) / (MaxValue.Z - MinValue.Z),
-					(SrcPtr[3] - MinValue.W) / (MaxValue.W - MinValue.W)
-					).ToFColor(bLinearToGamma);
-				FMemory::Memcpy(DestPtr++,&NormalizedColor,sizeof(FColor));
-				SrcPtr += 4;
-			}
-		}
-	}
-	else if (Format == DXGI_FORMAT_R11G11B10_FLOAT)
-	{
-		check(sizeof(FFloat3Packed) == sizeof(uint32));
-
-		for(uint32 Y = 0; Y < Height; Y++)
-		{
-			FFloat3Packed* SrcPtr = (FFloat3Packed*)(In + Y * SrcPitch);
-			FColor* DestPtr = Out + Y * Width;
-
-			for(uint32 X = 0; X < Width; X++)
-			{
-				FLinearColor Value = (*SrcPtr).ToLinearColor();
-
-				FColor NormalizedColor = Value.ToFColor(bLinearToGamma);
-				FMemory::Memcpy(DestPtr++, &NormalizedColor, sizeof(FColor));
-				++SrcPtr;
-			}
-		}
-	}
-	else if (Format == DXGI_FORMAT_R32G32B32A32_FLOAT)
-	{
-		FPlane MinValue(0.0f,0.0f,0.0f,0.0f);
-		FPlane MaxValue(1.0f,1.0f,1.0f,1.0f);
-
-		for(uint32 Y = 0; Y < Height; Y++)
-		{
-			float* SrcPtr = (float*)(In + Y * SrcPitch);
-
-			for(uint32 X = 0; X < Width; X++)
-			{
-				MinValue.X = FMath::Min<float>(SrcPtr[0],MinValue.X);
-				MinValue.Y = FMath::Min<float>(SrcPtr[1],MinValue.Y);
-				MinValue.Z = FMath::Min<float>(SrcPtr[2],MinValue.Z);
-				MinValue.W = FMath::Min<float>(SrcPtr[3],MinValue.W);
-				MaxValue.X = FMath::Max<float>(SrcPtr[0],MaxValue.X);
-				MaxValue.Y = FMath::Max<float>(SrcPtr[1],MaxValue.Y);
-				MaxValue.Z = FMath::Max<float>(SrcPtr[2],MaxValue.Z);
-				MaxValue.W = FMath::Max<float>(SrcPtr[3],MaxValue.W);
-				SrcPtr += 4;
-			}
-		}
-		
-		for(uint32 Y = 0; Y < Height; Y++)
-		{
-			float* SrcPtr = (float*)In;
-			FColor* DestPtr = Out + Y * Width;
-
-			for(uint32 X = 0; X < Width; X++)
-			{
-				FColor NormalizedColor =
-					FLinearColor(
-					(SrcPtr[0] - MinValue.X) / (MaxValue.X - MinValue.X),
-					(SrcPtr[1] - MinValue.Y) / (MaxValue.Y - MinValue.Y),
-					(SrcPtr[2] - MinValue.Z) / (MaxValue.Z - MinValue.Z),
-					(SrcPtr[3] - MinValue.W) / (MaxValue.W - MinValue.W)
-					).ToFColor(bLinearToGamma);
-				FMemory::Memcpy(DestPtr++,&NormalizedColor,sizeof(FColor));
-				SrcPtr += 4;
-			}
-		}
-	}
-	else if (Format == DXGI_FORMAT_R24G8_TYPELESS)
-	{
-		// Depth stencil
-		for(uint32 Y = 0; Y < Height; Y++)
-		{
-			uint32* SrcPtr = (uint32 *)In;
-			FColor* DestPtr = Out + Y * Width;
-			
-			for(uint32 X = 0; X < Width; X++)
-			{
-				FColor NormalizedColor;
-				if (InFlags.GetOutputStencil())
-				{
-					uint8 DeviceStencil = (*SrcPtr & 0xFF000000) >> 24;
-					NormalizedColor = FColor(DeviceStencil, DeviceStencil, DeviceStencil, 0xFF);
-				}
-				else
-				{
-					float DeviceZ = (*SrcPtr & 0xffffff) / (float)(1<<24);
-					float LinearValue = FMath::Min(InFlags.ComputeNormalizedDepth(DeviceZ), 1.0f);
-					NormalizedColor = FLinearColor(LinearValue, LinearValue, LinearValue, 0).ToFColor(bLinearToGamma);
-				}
-								
-				FMemory::Memcpy(DestPtr++, &NormalizedColor, sizeof(FColor));
-				++SrcPtr;
-			}
-		}
-	}
-	// Changing Depth Buffers to 32 bit on Dingo as D24S8 is actually implemented as a 32 bit buffer in the hardware
-	else if (Format == DXGI_FORMAT_R32G8X24_TYPELESS )
-	{
-		// Depth stencil
-		for(uint32 Y = 0; Y < Height; Y++)
-		{
-			float* SrcPtr = (float *)(In + Y * SrcPitch);
-			FColor* DestPtr = Out + Y * Width;
-
-			for(uint32 X = 0; X < Width; X++)
-			{
-				float DeviceZ = (*SrcPtr);
-
-				float LinearValue = FMath::Min(InFlags.ComputeNormalizedDepth(DeviceZ), 1.0f);
-
-				FColor NormalizedColor = FLinearColor(LinearValue, LinearValue, LinearValue, 0).ToFColor(bLinearToGamma);
-				FMemory::Memcpy(DestPtr++, &NormalizedColor, sizeof(FColor));
-				SrcPtr+=1; // todo: copies only depth, need to check how this format is read
-				UE_LOG(LogD3D11RHI, Warning, TEXT("CPU read of R32G8X24 is not tested and may not function."));
-			}
-		}
-	}
-	else if(Format == DXGI_FORMAT_R16G16B16A16_UNORM)
-	{
-		// Read the data out of the buffer, converting it to FColor.
-		for(uint32 Y = 0; Y < Height; Y++)
-		{
-			FD3DRGBA16* SrcPtr = (FD3DRGBA16*)(In + Y * SrcPitch);
-			FColor* DestPtr = Out + Y * Width;
-			for(uint32 X = 0; X < Width; X++)
-			{
-				*DestPtr = FLinearColor(
-					(float)SrcPtr->R / 65535.0f,
-					(float)SrcPtr->G / 65535.0f,
-					(float)SrcPtr->B / 65535.0f,
-					(float)SrcPtr->A / 65535.0f
-					).Quantize();
-				++SrcPtr;
-				++DestPtr;
-			}
-		}
-	}
-	else if(Format == DXGI_FORMAT_R16G16_UNORM)
-	{
-		// Read the data out of the buffer, converting it to FColor.
-		for(uint32 Y = 0; Y < Height; Y++)
-		{
-			FD3DRG16* SrcPtr = (FD3DRG16*)(In + Y * SrcPitch);
-			FColor* DestPtr = Out + Y * Width;
-			for(uint32 X = 0; X < Width; X++)
-			{
-				*DestPtr = FLinearColor(
-					(float)SrcPtr->R / 65535.0f,
-					(float)SrcPtr->G / 65535.0f,
-					0).Quantize();
-				++SrcPtr;
-				++DestPtr;
-			}
-		}
-	}
-	else if (Format == DXGI_FORMAT_R8_UNORM)
-	{
-		// Read the data out of the buffer, converting it from ABGR to ARGB.
-		for (uint32 Y = 0; Y < Height; Y++)
-		{
-			uint8* SrcPtr = (uint8*)(In + Y * SrcPitch);
-			FColor* DestPtr = Out + Y * Width;
-			for (uint32 X = 0; X < Width; X++)
-			{
-				*DestPtr = FColor(*SrcPtr, 0, 0, 0);
-				++SrcPtr;
-				++DestPtr;
-			}
-		}
-	}
-	else
+	switch (Format)
 	{
-		// not supported yet
-		check(0);
+		case DXGI_FORMAT_R16_TYPELESS:
+			ConvertRawR16DataToFColor(Width, Height, In, SrcPitch, Out);
+			break;
+		case DXGI_FORMAT_R8G8B8A8_TYPELESS:
+		case DXGI_FORMAT_R8G8B8A8_UNORM:
+		case DXGI_FORMAT_R8G8B8A8_UNORM_SRGB:
+			ConvertRawR8G8B8A8DataToFColor(Width, Height, In, SrcPitch, Out);
+			break;
+		case DXGI_FORMAT_B8G8R8A8_TYPELESS:
+		case DXGI_FORMAT_B8G8R8A8_UNORM:
+		case DXGI_FORMAT_B8G8R8A8_UNORM_SRGB:
+			ConvertRawB8G8R8A8DataToFColor(Width, Height, In, SrcPitch, Out);
+			break;
+		case DXGI_FORMAT_R10G10B10A2_UNORM:
+			ConvertRawR10G10B10A2DataToFColor(Width, Height, In, SrcPitch, Out);
+			break;
+		case DXGI_FORMAT_R16G16B16A16_FLOAT:
+			ConvertRawR16G16B16A16FDataToFColor(Width, Height, In, SrcPitch, Out, bLinearToGamma);
+			break;
+		case DXGI_FORMAT_R11G11B10_FLOAT:
+			ConvertRawR11G11B10DataToFColor(Width, Height, In, SrcPitch, Out, bLinearToGamma);
+			break;
+		case DXGI_FORMAT_R32G32B32A32_FLOAT:
+			ConvertRawR32G32B32A32DataToFColor(Width, Height, In, SrcPitch, Out, bLinearToGamma);
+			break;
+		case DXGI_FORMAT_R24G8_TYPELESS:
+			ConvertRawR24G8DataToFColor(Width, Height, In, SrcPitch, Out, InFlags);
+			break;
+		case DXGI_FORMAT_R32G8X24_TYPELESS:
+			ConvertRawR32DataToFColor(Width, Height, In, SrcPitch, Out, InFlags);
+			break;
+		case DXGI_FORMAT_R16G16B16A16_UNORM:
+			ConvertRawR16G16B16A16DataToFColor(Width, Height, In, SrcPitch, Out);
+			break;
+		case DXGI_FORMAT_R16G16_UNORM:
+			ConvertRawR16G16DataToFColor(Width, Height, In, SrcPitch, Out);
+			break;
+		case DXGI_FORMAT_R8_UNORM:
+			ConvertRawR8DataToFColor(Width, Height, In, SrcPitch, Out);
+			break;
+		default:
+			checkf(0, TEXT("Unknown surface format!"));
+			break;
 	}
 }
 
@@ -938,7 +706,7 @@ void FD3D11DynamicRHI::RHIReadSurfaceData(FRHITexture* TextureRHI,FIntRect InRec
 	uint32 BytesPerPixel = ComputeBytesPerPixel(TextureDesc.Format);
 	uint32 SrcPitch = SizeX * BytesPerPixel;
 
-	ConvertRAWSurfaceDataToFColor(TextureDesc.Format, SizeX, SizeY, OutDataRaw.GetData(), SrcPitch, OutData.GetData(), InFlags);
+	ConvertDXGIToFColor(TextureDesc.Format, SizeX, SizeY, OutDataRaw.GetData(), SrcPitch, OutData.GetData(), InFlags);
 }
 
 void FD3D11DynamicRHI::ReadSurfaceDataMSAARaw(FRHICommandList_RecursiveHazardous& RHICmdList, FRHITexture* TextureRHI,FIntRect InRect,TArray<uint8>& OutData, FReadSurfaceDataFlags InFlags)
@@ -1169,285 +937,51 @@ void FD3D11DynamicRHI::RHIReadSurfaceFloatData(FRHITexture* TextureRHI,FIntRect
 
 static void ConvertRAWSurfaceDataToFLinearColor(EPixelFormat Format, uint32 Width, uint32 Height, uint8 *In, uint32 SrcPitch, FLinearColor* Out, FReadSurfaceDataFlags InFlags)
 {
+	bool bLinearToGamma = InFlags.GetLinearToGamma();
 	if (Format == PF_R16F || Format == PF_R16F_FILTER)
 	{
-		// e.g. shadow maps
-		for (uint32 Y = 0; Y < Height; Y++)
-		{
-			uint16* SrcPtr = (uint16*)(In + Y * SrcPitch);
-			FLinearColor* DestPtr = Out + Y * Width;
-
-			for (uint32 X = 0; X < Width; X++)
-			{
-				uint16 Value16 = *SrcPtr;
-				float Value = Value16 / (float)(0xffff);
-
-				*DestPtr = FLinearColor(Value, Value, Value);
-				++SrcPtr;
-				++DestPtr;
-			}
-		}
+		ConvertRawR16DataToFLinearColor(Width, Height, In, SrcPitch, Out);
 	}
 	else if (Format == PF_R8G8B8A8)
 	{
-		// Read the data out of the buffer, converting it from ABGR to ARGB.
-		for (uint32 Y = 0; Y < Height; Y++)
-		{
-			FColor* SrcPtr = (FColor*)(In + Y * SrcPitch);
-			FLinearColor* DestPtr = Out + Y * Width;
-			for (uint32 X = 0; X < Width; X++)
-			{
-				FColor sRGBColor = FColor(SrcPtr->B, SrcPtr->G, SrcPtr->R, SrcPtr->A);
-				*DestPtr = FLinearColor(sRGBColor);
-				++SrcPtr;
-				++DestPtr;
-			}
-		}
+		ConvertRawR8G8B8A8DataToFLinearColor(Width, Height, In, SrcPitch, Out);
 	}
 	else if (Format == PF_B8G8R8A8)
 	{
-		for (uint32 Y = 0; Y < Height; Y++)
-		{
-			FColor* SrcPtr = (FColor*)(In + Y * SrcPitch);
-			FLinearColor* DestPtr = Out + Y * Width;
-			for (uint32 X = 0; X < Width; X++)
-			{
-				FColor sRGBColor = FColor(SrcPtr->R, SrcPtr->G, SrcPtr->B, SrcPtr->A);
-				*DestPtr = FLinearColor(sRGBColor);
-				++SrcPtr;
-				++DestPtr;
-			}
-		}
+		ConvertRawB8G8R8A8DataToFLinearColor(Width, Height, In, SrcPitch, Out);
 	}
 	else if (Format == PF_A2B10G10R10)
 	{
-		// Read the data out of the buffer, converting it from R10G10B10A2 to FLinearColor.
-		for (uint32 Y = 0; Y < Height; Y++)
-		{
-			FD3DR10G10B10A2* SrcPtr = (FD3DR10G10B10A2*)(In + Y * SrcPitch);
-			FLinearColor* DestPtr = Out + Y * Width;
-			for (uint32 X = 0; X < Width; X++)
-			{
-				*DestPtr = FLinearColor(
-					(float)SrcPtr->R / 1023.0f,
-					(float)SrcPtr->G / 1023.0f,
-					(float)SrcPtr->B / 1023.0f,
-					(float)SrcPtr->A / 3.0f
-					);
-				++SrcPtr;
-				++DestPtr;
-			}
-		}
+		ConvertRawA2B10G10R10DataToFLinearColor(Width, Height, In, SrcPitch, Out);
 	}
 	else if (Format == PF_FloatRGBA)
 	{
-		if (InFlags.GetCompressionMode() == RCM_MinMax)
-		{
-			for (uint32 Y = 0; Y < Height; Y++)
-			{
-				FFloat16* SrcPtr = (FFloat16*)(In + Y * SrcPitch);
-				FLinearColor* DestPtr = Out + Y * Width;
-
-				for (uint32 X = 0; X < Width; X++)
-				{
-					*DestPtr = FLinearColor((float)SrcPtr[0], (float)SrcPtr[1], (float)SrcPtr[2], (float)SrcPtr[3]);
-					++DestPtr;
-					SrcPtr += 4;
-				}
-			}
-		}
-		else
-		{
-			FPlane	MinValue(0.0f, 0.0f, 0.0f, 0.0f);
-			FPlane	MaxValue(1.0f, 1.0f, 1.0f, 1.0f);
-
-			check(sizeof(FFloat16) == sizeof(uint16));
-
-			for (uint32 Y = 0; Y < Height; Y++)
-			{
-				FFloat16* SrcPtr = (FFloat16*)(In + Y * SrcPitch);
-
-				for (uint32 X = 0; X < Width; X++)
-				{
-					MinValue.X = FMath::Min<float>(SrcPtr[0], MinValue.X);
-					MinValue.Y = FMath::Min<float>(SrcPtr[1], MinValue.Y);
-					MinValue.Z = FMath::Min<float>(SrcPtr[2], MinValue.Z);
-					MinValue.W = FMath::Min<float>(SrcPtr[3], MinValue.W);
-					MaxValue.X = FMath::Max<float>(SrcPtr[0], MaxValue.X);
-					MaxValue.Y = FMath::Max<float>(SrcPtr[1], MaxValue.Y);
-					MaxValue.Z = FMath::Max<float>(SrcPtr[2], MaxValue.Z);
-					MaxValue.W = FMath::Max<float>(SrcPtr[3], MaxValue.W);
-					SrcPtr += 4;
-				}
-			}
-
-			for (uint32 Y = 0; Y < Height; Y++)
-			{
-				FFloat16* SrcPtr = (FFloat16*)(In + Y * SrcPitch);
-				FLinearColor* DestPtr = Out + Y * Width;
-
-				for (uint32 X = 0; X < Width; X++)
-				{
-					*DestPtr = FLinearColor(
-						(SrcPtr[0] - MinValue.X) / (MaxValue.X - MinValue.X),
-						(SrcPtr[1] - MinValue.Y) / (MaxValue.Y - MinValue.Y),
-						(SrcPtr[2] - MinValue.Z) / (MaxValue.Z - MinValue.Z),
-						(SrcPtr[3] - MinValue.W) / (MaxValue.W - MinValue.W)
-						);
-					++DestPtr;
-					SrcPtr += 4;
-				}
-			}
-		}
+		ConvertRawR16G16B16A16FDataToFLinearColor(Width, Height, In, SrcPitch, Out, InFlags);
 	}
 	else if (Format == PF_FloatRGB || Format == PF_FloatR11G11B10)
 	{
-		check(sizeof(FFloat3Packed) == sizeof(uint32));
-
-		for (uint32 Y = 0; Y < Height; Y++)
-		{
-			FFloat3Packed* SrcPtr = (FFloat3Packed*)(In + Y * SrcPitch);
-			FLinearColor* DestPtr = Out + Y * Width;
-
-			for (uint32 X = 0; X < Width; X++)
-			{
-				*DestPtr = (*SrcPtr).ToLinearColor();
-				++DestPtr;
-				++SrcPtr;
-			}
-		}
+		ConvertRawRR11G11B10DataToFLinearColor(Width, Height, In, SrcPitch, Out);
 	}
 	else if (Format == PF_A32B32G32R32F)
 	{
-		if (InFlags.GetCompressionMode() == RCM_MinMax)
-		{
-			// Copy data directly, respecting existing min-max values
-			FLinearColor* SrcPtr = (FLinearColor*)In;
-			FLinearColor* DestPtr = (FLinearColor*)Out;
-			const int32 ImageSize = sizeof(FLinearColor) * Height * Width;
-
-			FMemory::Memcpy(DestPtr, SrcPtr, ImageSize);
-		}
-		else
-		{
-			// Normalize data
-			FPlane MinValue(0.0f, 0.0f, 0.0f, 0.0f);
-			FPlane MaxValue(1.0f, 1.0f, 1.0f, 1.0f);
-
-			for (uint32 Y = 0; Y < Height; Y++)
-			{
-				float* SrcPtr = (float*)(In + Y * SrcPitch);
-
-				for (uint32 X = 0; X < Width; X++)
-				{
-					MinValue.X = FMath::Min<float>(SrcPtr[0], MinValue.X);
-					MinValue.Y = FMath::Min<float>(SrcPtr[1], MinValue.Y);
-					MinValue.Z = FMath::Min<float>(SrcPtr[2], MinValue.Z);
-					MinValue.W = FMath::Min<float>(SrcPtr[3], MinValue.W);
-					MaxValue.X = FMath::Max<float>(SrcPtr[0], MaxValue.X);
-					MaxValue.Y = FMath::Max<float>(SrcPtr[1], MaxValue.Y);
-					MaxValue.Z = FMath::Max<float>(SrcPtr[2], MaxValue.Z);
-					MaxValue.W = FMath::Max<float>(SrcPtr[3], MaxValue.W);
-					SrcPtr += 4;
-				} 
-			}
-
-			float* SrcPtr = (float*)In;
-
-			for (uint32 Y = 0; Y < Height; Y++)
-			{
-				FLinearColor* DestPtr = Out + Y * Width;
-
-				for (uint32 X = 0; X < Width; X++)
-				{
-					*DestPtr = FLinearColor(
-							(SrcPtr[0] - MinValue.X) / (MaxValue.X - MinValue.X),
-							(SrcPtr[1] - MinValue.Y) / (MaxValue.Y - MinValue.Y),
-							(SrcPtr[2] - MinValue.Z) / (MaxValue.Z - MinValue.Z),
-							(SrcPtr[3] - MinValue.W) / (MaxValue.W - MinValue.W)
-							);
-					++DestPtr;
-					SrcPtr += 4;
-				}
-			}
-		}
+		ConvertRawR32G32B32A32DataToFLinearColor(Width, Height, In, SrcPitch, Out, InFlags);
 	}
-	else if (Format == PF_DepthStencil || Format == PF_D24)
+	else if (Format == PF_D24)
 	{
-		// Depth stencil
-		for (uint32 Y = 0; Y < Height; Y++)
-		{
-			uint32* SrcPtr = (uint32 *)In;
-			FLinearColor* DestPtr = Out + Y * Width;
-
-			for (uint32 X = 0; X < Width; X++)
-			{
-				float DeviceStencil = 0.0f;
-				DeviceStencil = (float)((*SrcPtr & 0xFF000000) >> 24)/255.0f;
-				float DeviceZ = (*SrcPtr & 0xffffff) / (float)(1 << 24);
-				float LinearValue = FMath::Min(InFlags.ComputeNormalizedDepth(DeviceZ), 1.0f);
-				*DestPtr = FLinearColor(LinearValue, DeviceStencil, 0.0f, 0.0f);
-				++DestPtr;
-				++SrcPtr;
-			}
-		}
+		ConvertRawR24G8DataToFLinearColor(Width, Height, In, SrcPitch, Out, InFlags);
 	}
 	// Changing Depth Buffers to 32 bit on Dingo as D24S8 is actually implemented as a 32 bit buffer in the hardware
 	else if (Format == PF_DepthStencil)
 	{
-		// Depth stencil
-		for (uint32 Y = 0; Y < Height; Y++)
-		{
-			uint8* SrcStart = (uint8 *)(In + Y * SrcPitch);
-			FLinearColor* DestPtr = Out + Y * Width;
-
-			for (uint32 X = 0; X < Width; X++)
-			{
-				float DeviceZ = *((float *)(SrcStart));
-				float LinearValue = FMath::Min(InFlags.ComputeNormalizedDepth(DeviceZ), 1.0f);
-				float DeviceStencil = (float)(*(SrcStart+4)) / 255.0f;
-				*DestPtr = FLinearColor(LinearValue, DeviceStencil, 0.0f, 0.0f);
-				SrcStart += 8; //64 bit format with the last 24 bit ignore
-			}
-		}
+		ConvertRawR32DataToFLinearColor(Width, Height, In, SrcPitch, Out, InFlags);
 	}
 	else if (Format == PF_A16B16G16R16)
 	{
-		// Read the data out of the buffer, converting it to FLinearColor.
-		for (uint32 Y = 0; Y < Height; Y++)
-		{
-			FD3DRGBA16* SrcPtr = (FD3DRGBA16*)(In + Y * SrcPitch);
-			FLinearColor* DestPtr = Out + Y * Width;
-			for (uint32 X = 0; X < Width; X++)
-			{
-				*DestPtr = FLinearColor(
-					(float)SrcPtr->R / 65535.0f,
-					(float)SrcPtr->G / 65535.0f,
-					(float)SrcPtr->B / 65535.0f,
-					(float)SrcPtr->A / 65535.0f
-					);
-				++SrcPtr;
-				++DestPtr;
-			}
-		}
+		ConvertRawR16G16B16A16DataToFLinearColor(Width, Height, In, SrcPitch, Out);
 	}
 	else if (Format == PF_G16R16)
 	{
-		// Read the data out of the buffer, converting it to FLinearColor.
-		for (uint32 Y = 0; Y < Height; Y++)
-		{
-			FD3DRG16* SrcPtr = (FD3DRG16*)(In + Y * SrcPitch);
-			FLinearColor* DestPtr = Out + Y * Width;
-			for (uint32 X = 0; X < Width; X++)
-			{
-				*DestPtr = FLinearColor(
-					(float)SrcPtr->R / 65535.0f,
-					(float)SrcPtr->G / 65535.0f,
-					0);
-				++SrcPtr;
-				++DestPtr;
-			}
-		}
+		ConvertRawR16G16DataToFLinearColor(Width, Height, In, SrcPitch, Out);
 	}
 	else
 	{
