Fix motion blur assert about unsupported format

diff --git a/Engine/Source/Runtime/Windows/D3D11RHI/Private/D3D11Device.cpp b/Engine/Source/Runtime/Windows/D3D11RHI/Private/D3D11Device.cpp
--- a/Engine/Source/Runtime/Windows/D3D11RHI/Private/D3D11Device.cpp
+++ b/Engine/Source/Runtime/Windows/D3D11RHI/Private/D3D11Device.cpp
@@ -148,6 +148,7 @@ FD3D11DynamicRHI::FD3D11DynamicRHI(IDXGIFactory1* InDXGIFactory1,D3D_FEATURE_LEV
 	}
 	GPixelFormats[ PF_ShadowDepth	].PlatformFormat	= DXGI_FORMAT_R16_TYPELESS;
 	GPixelFormats[ PF_ShadowDepth	].BlockBytes		= 2;
+	GPixelFormats[ PF_ShadowDepth	].Supported			= true;
 	GPixelFormats[ PF_R32_FLOAT		].PlatformFormat	= DXGI_FORMAT_R32_FLOAT;
 	GPixelFormats[ PF_G16R16		].PlatformFormat	= DXGI_FORMAT_R16G16_UNORM;
 	GPixelFormats[ PF_G16R16F		].PlatformFormat	= DXGI_FORMAT_R16G16_FLOAT;
