Fixed merge compilation issue.

diff --git a/Engine/Source/Runtime/Windows/D3D11RHI/Private/D3D11Util.cpp b/Engine/Source/Runtime/Windows/D3D11RHI/Private/D3D11Util.cpp
--- a/Engine/Source/Runtime/Windows/D3D11RHI/Private/D3D11Util.cpp
+++ b/Engine/Source/Runtime/Windows/D3D11RHI/Private/D3D11Util.cpp
@@ -258,15 +258,6 @@ void VerifyD3D11Result(HRESULT D3DResult,const ANSICHAR* Code,const ANSICHAR* Fi
 	UE_LOG(LogD3D11RHI, Fatal,TEXT("%s failed \n at %s:%u \n with error %s"),ANSI_TO_TCHAR(Code),ANSI_TO_TCHAR(Filename),Line,*ErrorString);
 }
 
-void VerifyD3D11ResultNoExit(HRESULT D3DResult, const ANSICHAR* Code, const ANSICHAR* Filename, uint32 Line, ID3D11Device* Device)
-{
-	check(FAILED(D3DResult));
-
-	const FString& ErrorString = GetD3D11ErrorString(D3DResult, Device);
-
-	UE_LOG(LogD3D11RHI, Error, TEXT("%s failed \n at %s:%u \n with error %s Error Code List: https://docs.microsoft.com/en-us/windows/desktop/direct3ddxgi/dxgi-error"), ANSI_TO_TCHAR(Code), ANSI_TO_TCHAR(Filename), Line, *ErrorString);
-}
-
 void VerifyD3D11ShaderResult(FRHIShader* Shader, HRESULT D3DResult, const ANSICHAR* Code, const ANSICHAR* Filename, uint32 Line, ID3D11Device* Device)
 {
 	check(FAILED(D3DResult));
