Fix for static analysis warning

diff --git a/Engine/Source/Runtime/Windows/AudioMixerXAudio2/Private/AudioMixerPlatformXAudio2.cpp b/Engine/Source/Runtime/Windows/AudioMixerXAudio2/Private/AudioMixerPlatformXAudio2.cpp
--- a/Engine/Source/Runtime/Windows/AudioMixerXAudio2/Private/AudioMixerPlatformXAudio2.cpp
+++ b/Engine/Source/Runtime/Windows/AudioMixerXAudio2/Private/AudioMixerPlatformXAudio2.cpp
@@ -120,7 +120,7 @@ namespace Audio
 			case E_INVALIDARG:								return TEXT("E_INVALIDARG");
 			case E_OUTOFMEMORY:								return TEXT("E_OUTOFMEMORY");
 #endif
-			case 0xe000020b:								return TEXT("ERROR_NO_SUCH_DEVINST");
+			case HRESULT(0xe000020b):						return TEXT("ERROR_NO_SUCH_DEVINST");
 			default:										
 			{
 				// We don't know this error, ask this system if it does.
