Fix crash changing to/from fullscreen mode.

diff --git a/Engine/Source/Runtime/Windows/D3D11RHI/Private/D3D11Viewport.cpp b/Engine/Source/Runtime/Windows/D3D11RHI/Private/D3D11Viewport.cpp
--- a/Engine/Source/Runtime/Windows/D3D11RHI/Private/D3D11Viewport.cpp
+++ b/Engine/Source/Runtime/Windows/D3D11RHI/Private/D3D11Viewport.cpp
@@ -269,7 +269,8 @@ void FD3D11Viewport::Resize(uint32 InSizeX, uint32 InSizeY, bool bInIsFullscreen
 
 				if (FAILED(SwapChain->ResizeTarget(&BufferDesc)))
 				{
-					ConditionalResetSwapChain(true);
+					ResetSwapChainInternal(true);
+					VERIFYD3D11RESIZEVIEWPORTRESULT(SwapChain->ResizeBuffers(0, SizeX, SizeY, RenderTargetFormat, D3D11GetSwapChainFlags()), SizeX, SizeY, RenderTargetFormat, D3DRHI->GetDevice());
 				}
 			}
 		}
@@ -284,7 +285,9 @@ void FD3D11Viewport::Resize(uint32 InSizeX, uint32 InSizeY, bool bInIsFullscreen
 		{
 			// Use ConditionalResetSwapChain to call SetFullscreenState, to handle the failure case.
 			// Ignore the viewport's focus state; since Resize is called as the result of a user action we assume authority without waiting for Focus.
-			ConditionalResetSwapChain(true);
+			ResetSwapChainInternal(true);
+			DXGI_FORMAT RenderTargetFormat = GetRenderTargetFormat(PixelFormat);
+			VERIFYD3D11RESIZEVIEWPORTRESULT(SwapChain->ResizeBuffers(0, SizeX, SizeY, RenderTargetFormat, D3D11GetSwapChainFlags()), SizeX, SizeY, RenderTargetFormat, D3DRHI->GetDevice());
 		}
 	}
 
@@ -521,6 +524,7 @@ bool FD3D11Viewport::Present(bool bLockToVsync)
 		// Can't compare BOOL with bool...
 		if ( (!!bSwapChainFullscreenState)  != bIsFullscreen )
 		{
+			bFullscreenLost = true;
 			bIsValid = false;
 		}
 	}
@@ -533,6 +537,10 @@ bool FD3D11Viewport::Present(bool bLockToVsync)
 		DXGIDevice->SetMaximumFrameLatency(MaximumFrameLatency);
 	}
 
+	if (!bIsValid)
+	{
+		return false;
+	}
 	// When desktop composition is enabled, locking to vsync via the Present
 	// call is unreliable. Instead, communicate with the desktop window manager
 	// directly to enable vsync.
diff --git a/Engine/Source/Runtime/Windows/D3D11RHI/Private/HoloLens/HoloLensD3D11Viewport.cpp b/Engine/Source/Runtime/Windows/D3D11RHI/Private/HoloLens/HoloLensD3D11Viewport.cpp
--- a/Engine/Source/Runtime/Windows/D3D11RHI/Private/HoloLens/HoloLensD3D11Viewport.cpp
+++ b/Engine/Source/Runtime/Windows/D3D11RHI/Private/HoloLens/HoloLensD3D11Viewport.cpp
@@ -201,6 +201,23 @@ void FD3D11Viewport::CheckHDRMonitorStatus()
 }
 
 void FD3D11Viewport::ConditionalResetSwapChain(bool bIgnoreFocus)
+{
+	if (!bIsValid)
+	{
+		if (bFullscreenLost)
+		{
+			FlushRenderingCommands();
+			bFullscreenLost = false;
+			Resize(SizeX, SizeY, false, PixelFormat);
+		}
+		else
+		{
+			ResetSwapChainInternal(bIgnoreFocus);
+		}
+	}
+}
+
+void FD3D11Viewport::ResetSwapChainInternal(bool bIgnoreFocus)
 {
 	if(!bIsValid)
 	{
diff --git a/Engine/Source/Runtime/Windows/D3D11RHI/Private/Windows/WindowsD3D11Viewport.cpp b/Engine/Source/Runtime/Windows/D3D11RHI/Private/Windows/WindowsD3D11Viewport.cpp
--- a/Engine/Source/Runtime/Windows/D3D11RHI/Private/Windows/WindowsD3D11Viewport.cpp
+++ b/Engine/Source/Runtime/Windows/D3D11RHI/Private/Windows/WindowsD3D11Viewport.cpp
@@ -51,6 +51,7 @@ FD3D11Viewport::FD3D11Viewport(FD3D11DynamicRHI* InD3DRHI,HWND InWindowHandle,ui
 	SizeX(InSizeX),
 	SizeY(InSizeY),
 	bIsFullscreen(bInIsFullscreen),
+	bFullscreenLost(false),
 	PixelFormat(InPreferredPixelFormat),
 	PixelColorSpace(EColorSpaceAndEOTF::ERec709_sRGB),
 	bIsValid(true),
@@ -181,11 +182,11 @@ FD3D11Viewport::FD3D11Viewport(FD3D11DynamicRHI* InD3DRHI,HWND InWindowHandle,ui
 			IDXGIFactory2* Factory2 = (IDXGIFactory2*)D3DRHI->GetFactory();
 
 			VERIFYD3D11RESULT_EX((Factory2->CreateSwapChainForHwnd(D3DRHI->GetDevice(), WindowHandle, &SwapChainDesc, &FSSwapChainDesc, nullptr, &SwapChain1)), D3DRHI->GetDevice());
-			SwapChain1->QueryInterface(__uuidof(IDXGISwapChain1), (void**)SwapChain.GetInitReference());
+				SwapChain1->QueryInterface(__uuidof(IDXGISwapChain1), (void**)SwapChain.GetInitReference());
 
-			// See if we are running on a HDR monitor 
-			CheckHDRMonitorStatus();
-		}
+				// See if we are running on a HDR monitor 
+				CheckHDRMonitorStatus();
+			}
 
 
 		if (SwapChain == nullptr)
@@ -284,6 +285,23 @@ void FD3D11Viewport::CheckHDRMonitorStatus()
 }
 
 void FD3D11Viewport::ConditionalResetSwapChain(bool bIgnoreFocus)
+{
+	if (!bIsValid)
+	{
+		if (bFullscreenLost)
+		{
+			FlushRenderingCommands();
+			bFullscreenLost = false;
+			Resize(SizeX, SizeY, false, PixelFormat);
+		}
+		else
+		{
+			ResetSwapChainInternal(bIgnoreFocus);
+		}
+	}
+}
+
+void FD3D11Viewport::ResetSwapChainInternal(bool bIgnoreFocus)
 {
 	if (!bIsValid)
 	{
diff --git a/Engine/Source/Runtime/Windows/D3D11RHI/Public/D3D11Viewport.h b/Engine/Source/Runtime/Windows/D3D11RHI/Public/D3D11Viewport.h
--- a/Engine/Source/Runtime/Windows/D3D11RHI/Public/D3D11Viewport.h
+++ b/Engine/Source/Runtime/Windows/D3D11RHI/Public/D3D11Viewport.h
@@ -61,8 +61,8 @@ class D3D11RHI_API FD3D11Viewport : public FRHIViewport
 {
 public:
 
-	FD3D11Viewport(class FD3D11DynamicRHI* InD3DRHI): D3DRHI(InD3DRHI), FrameSyncEvent(InD3DRHI){}
-	FD3D11Viewport(class FD3D11DynamicRHI* InD3DRHI,HWND InWindowHandle,uint32 InSizeX,uint32 InSizeY,bool bInIsFullscreen, EPixelFormat InPreferredPixelFormat);
+	FD3D11Viewport(class FD3D11DynamicRHI* InD3DRHI) : D3DRHI(InD3DRHI), bFullscreenLost(false), FrameSyncEvent(InD3DRHI) {}
+	FD3D11Viewport(class FD3D11DynamicRHI* InD3DRHI, HWND InWindowHandle, uint32 InSizeX, uint32 InSizeY, bool bInIsFullscreen, EPixelFormat InPreferredPixelFormat);
 	~FD3D11Viewport();
 
 	virtual void Resize(uint32 InSizeX, uint32 InSizeY, bool bInIsFullscreen, EPixelFormat PreferredPixelFormat);
@@ -71,7 +71,7 @@ public:
 	 * If the swap chain has been invalidated by DXGI, resets the swap chain to the expected state; otherwise, does nothing.
 	 * Called once/frame by the game thread on all viewports.
 	 * @param bIgnoreFocus - Whether the reset should happen regardless of whether the window is focused.
-     */
+	 */
 	void ConditionalResetSwapChain(bool bIgnoreFocus);
 
 	/**
@@ -82,7 +82,7 @@ public:
 
 	/** Presents the swap chain. 
 	 * Returns true if Present was done by Engine.
- 	 */
+	 */
 	bool Present(bool bLockToVsync);
 
 	// Accessors.
@@ -121,6 +121,8 @@ public:
 
 protected:
 
+	void ResetSwapChainInternal(bool bIgnoreFocus);
+
 	/** Presents the frame synchronizing with DWM. */
 	void PresentWithVsyncDWM();
 
@@ -141,7 +143,8 @@ protected:
 	uint32 SizeX;
 	uint32 SizeY;
 	uint32 BackBufferCount;
-	bool bIsFullscreen;
+	bool bIsFullscreen : 1;
+	bool bFullscreenLost : 1;
 	EPixelFormat PixelFormat;
 	EColorSpaceAndEOTF PixelColorSpace;
 	bool bIsValid;
