Augmenting D3D11 uniform buffer assert with more helpful information.

diff --git a/Engine/Source/Runtime/Windows/D3D11RHI/Private/D3D11Commands.cpp b/Engine/Source/Runtime/Windows/D3D11RHI/Private/D3D11Commands.cpp
--- a/Engine/Source/Runtime/Windows/D3D11RHI/Private/D3D11Commands.cpp
+++ b/Engine/Source/Runtime/Windows/D3D11RHI/Private/D3D11Commands.cpp
@@ -1526,7 +1526,11 @@ void FD3D11DynamicRHI::SetResourcesFromTables(const ShaderType* RESTRICT Shader)
 				UE_LOG(LogD3D11RHI, Error, TEXT("Bound Layout='%s' Shader='%s', Layout CB Size %d %d"), *DebugName, *ShaderName, BufferLayout.ConstantBufferSize, BufferLayout.Resources.Num());
 #endif
 				// this might mean you are accessing a data you haven't bound e.g. GBuffer
-				check(BufferLayout.GetHash() == Shader->ShaderResourceTable.ResourceTableLayoutHashes[BufferIndex]);
+				checkf(BufferLayout.GetHash() == Shader->ShaderResourceTable.ResourceTableLayoutHashes[BufferIndex],
+					TEXT("Uniform buffer bound to slot %u is not what the shader expected:\n")
+					TEXT("\tBound:    Uniform Buffer[%s] with Hash[%u]\n")
+					TEXT("\tExpected: Uniform Buffer[%s] with Hash[%u]"),
+					BufferIndex, *DebugName, BufferLayout.GetHash(), *Shader->UniformBuffers[BufferIndex].GetPlainNameString(), Shader->ShaderResourceTable.ResourceTableLayoutHashes[BufferIndex]);
 			}
 		}
 #endif
