Integrating NvGPU overclock information to crash context.

diff --git a/Engine/Source/Runtime/Windows/D3D11RHI/Private/Windows/WindowsD3D11Device.cpp b/Engine/Source/Runtime/Windows/D3D11RHI/Private/Windows/WindowsD3D11Device.cpp
--- a/Engine/Source/Runtime/Windows/D3D11RHI/Private/Windows/WindowsD3D11Device.cpp
+++ b/Engine/Source/Runtime/Windows/D3D11RHI/Private/Windows/WindowsD3D11Device.cpp
@@ -19,7 +19,7 @@
 #include "HardwareInfo.h"
 #include "Runtime/HeadMountedDisplay/Public/IHeadMountedDisplayModule.h"
 #include "GenericPlatform/GenericPlatformDriver.h"			// FGPUDriverInfo
-
+#include "GenericPlatform/GenericPlatformCrashContext.h"
 #include "dxgi1_3.h"
 #include "RHIValidation.h"
 
@@ -657,6 +657,81 @@ static bool SupportsHDROutput(FD3D11DynamicRHI* D3DRHI)
 	return false;
 }
 
+static bool IsDeviceOverclocked()
+{
+	if (IsRHIDeviceNVIDIA())
+	{
+		NvAPI_Status Status;
+		NvU32 NumGPUs;
+		NvPhysicalGpuHandle GPUHandles[NVAPI_MAX_PHYSICAL_GPUS];
+
+		// Fetch GPU NV handles
+		Status = NvAPI_EnumPhysicalGPUs(GPUHandles, &NumGPUs);
+
+		if (Status != NVAPI_OK)
+		{
+			return false;
+		}
+
+		for (NvU32 GPU = 0; GPU < NumGPUs; ++GPU)
+		{
+			const NvPhysicalGpuHandle& GPUHandle = GPUHandles[GPU];
+
+			NV_GPU_PERF_PSTATES20_INFO PS20Info = {0};
+			PS20Info.version = NV_GPU_PERF_PSTATES20_INFO_VER;
+
+			Status = NvAPI_GPU_GetPstates20(GPUHandle, &PS20Info);
+
+			// Some GPUs use an older struct layout
+			if (Status == NVAPI_INCOMPATIBLE_STRUCT_VERSION)
+			{
+				PS20Info.version = NV_GPU_PERF_PSTATES20_INFO_VER1;
+				Status = NvAPI_GPU_GetPstates20(GPUHandle, &PS20Info);
+			}
+
+			// Assume non-overclocked on failure
+			if (Status != NVAPI_OK)
+			{
+				return false;
+			}
+
+			// Where state was marked editable, check for relevant modifications
+			if (PS20Info.bIsEditable == 0)
+			{
+				continue;
+			}
+
+			NvU32 PState, Clock, Voltage;
+			for (PState = 0; PState < PS20Info.numPstates; PState++)
+			{
+				if (PS20Info.pstates[PState].bIsEditable != 0)
+				{
+					// Are clocks user modified?
+					for (Clock = 0; Clock < PS20Info.numClocks; Clock++)
+					{
+						if (PS20Info.pstates[PState].clocks[Clock].bIsEditable != 0 && PS20Info.pstates[PState].clocks[Clock].freqDelta_kHz.value != 0)
+						{
+							return true;
+						}
+					}
+
+					// Are voltages user modified?
+					for (Voltage = 0; Voltage < PS20Info.numBaseVoltages; Voltage++)
+					{
+						if (PS20Info.pstates[PState].baseVoltages[Voltage].bIsEditable != 0 && PS20Info.pstates[PState].baseVoltages[Voltage].voltDelta_uV.value != 0)
+						{
+							return true;
+						}
+					}
+				}
+			}
+		}
+	}
+
+	// Assume non-overclocked by default
+	return false;
+}
+
 void FD3D11DynamicRHIModule::StartupModule()
 {	
 #if NV_AFTERMATH
@@ -1757,6 +1832,10 @@ void FD3D11DynamicRHI::InitD3DDevice()
 			GRHISupportsHDROutput = SupportsHDROutput(this);
 		}
 
+		// Add device overclock state to crash context
+		const bool bIsGPUOverclocked = IsDeviceOverclocked();
+		FGenericCrashContext::SetEngineData(TEXT("RHI.IsGPUOverclocked"), bIsGPUOverclocked ? TEXT("true") : TEXT("false"));
+		
 		FHardwareInfo::RegisterHardwareInfo( NAME_RHI, TEXT( "D3D11" ) );
 
 		GRHISupportsTextureStreaming = true;
