Resolving compile error when NVAPI_INTERFACE is undefined.

diff --git a/Engine/Source/Runtime/Windows/D3D11RHI/Private/Windows/WindowsD3D11Device.cpp b/Engine/Source/Runtime/Windows/D3D11RHI/Private/Windows/WindowsD3D11Device.cpp
--- a/Engine/Source/Runtime/Windows/D3D11RHI/Private/Windows/WindowsD3D11Device.cpp
+++ b/Engine/Source/Runtime/Windows/D3D11RHI/Private/Windows/WindowsD3D11Device.cpp
@@ -661,6 +661,7 @@ static bool IsDeviceOverclocked()
 {
 	if (IsRHIDeviceNVIDIA())
 	{
+#ifdef NVAPI_INTERFACE
 		NvAPI_Status Status;
 		NvU32 NumGPUs;
 		NvPhysicalGpuHandle GPUHandles[NVAPI_MAX_PHYSICAL_GPUS];
@@ -726,6 +727,7 @@ static bool IsDeviceOverclocked()
 				}
 			}
 		}
+#endif // NVAPI_INTERFACE
 	}
 
 	// Assume non-overclocked by default
