UE4R - More info when can't make an unsupported format texture
#rnx

diff --git a/Engine/Source/Runtime/Windows/D3D11RHI/Private/D3D11Texture.cpp b/Engine/Source/Runtime/Windows/D3D11RHI/Private/D3D11Texture.cpp
--- a/Engine/Source/Runtime/Windows/D3D11RHI/Private/D3D11Texture.cpp
+++ b/Engine/Source/Runtime/Windows/D3D11RHI/Private/D3D11Texture.cpp
@@ -452,7 +452,7 @@ DXGI_FORMAT FD3D11DynamicRHI::GetPlatformTextureResourceFormat(DXGI_FORMAT InFor
 /**
  * Creates a 2D texture optionally guarded by a structured exception handler.
  */
-void SafeCreateTexture2D(ID3D11Device* Direct3DDevice, const D3D11_TEXTURE2D_DESC* TextureDesc, const D3D11_SUBRESOURCE_DATA* SubResourceData, ID3D11Texture2D** OutTexture2D)
+void SafeCreateTexture2D(ID3D11Device* Direct3DDevice, int32 UEFormat, const D3D11_TEXTURE2D_DESC* TextureDesc, const D3D11_SUBRESOURCE_DATA* SubResourceData, ID3D11Texture2D** OutTexture2D)
 {
 #if GUARDED_TEXTURE_CREATES
 	bool bDriverCrash = true;
@@ -461,6 +461,7 @@ void SafeCreateTexture2D(ID3D11Device* Direct3DDevice, const D3D11_TEXTURE2D_DES
 #endif // #if GUARDED_TEXTURE_CREATES
 		VERIFYD3D11CREATETEXTURERESULT(
 			Direct3DDevice->CreateTexture2D(TextureDesc,SubResourceData,OutTexture2D),
+			UEFormat,
 			TextureDesc->Width,
 			TextureDesc->Height,
 			TextureDesc->ArraySize,
@@ -485,13 +486,14 @@ void SafeCreateTexture2D(ID3D11Device* Direct3DDevice, const D3D11_TEXTURE2D_DES
 		if (bDriverCrash)
 		{
 			UE_LOG(LogD3D11RHI,Error,
-				TEXT("Driver crashed while creating texture: %ux%ux%u %s(0x%08x) with %u mips"),
+				TEXT("Driver crashed while creating texture: %ux%ux%u %s(0x%08x) with %u mips, PF_ %d"),
 				TextureDesc->Width,
 				TextureDesc->Height,
 				TextureDesc->ArraySize,
 				GetD3D11TextureFormatString(TextureDesc->Format),
 				(uint32)TextureDesc->Format,
-				TextureDesc->MipLevels
+				TextureDesc->MipLevels,
+				UEFormat
 				);
 		}
 	}
@@ -765,7 +767,7 @@ TD3D11Texture2D<BaseResourceType>* FD3D11DynamicRHI::CreateD3D11Texture2D(uint32
 		else
 #endif
 		{
-			SafeCreateTexture2D(Direct3DDevice, &TextureDesc, CreateInfo.BulkData != NULL ? (const D3D11_SUBRESOURCE_DATA*)SubResourceData.GetData() : NULL, TextureResource.GetInitReference());
+			SafeCreateTexture2D(Direct3DDevice, Format, &TextureDesc, CreateInfo.BulkData != NULL ? (const D3D11_SUBRESOURCE_DATA*)SubResourceData.GetData() : NULL, TextureResource.GetInitReference());
 		}
 
 		if(bCreateRTV)
@@ -1027,6 +1029,7 @@ FD3D11Texture3D* FD3D11DynamicRHI::CreateD3D11Texture3D(uint32 SizeX,uint32 Size
 	const D3D11_SUBRESOURCE_DATA* SubResData = CreateInfo.BulkData != nullptr ? (const D3D11_SUBRESOURCE_DATA*)SubResourceData.GetData() : nullptr;
 	VERIFYD3D11CREATETEXTURERESULT(
 		Direct3DDevice->CreateTexture3D(&TextureDesc, SubResData,TextureResource.GetInitReference()),
+		Format,
 		SizeX,
 		SizeY,
 		SizeZ,
@@ -1195,7 +1198,7 @@ FTexture2DRHIRef FD3D11DynamicRHI::RHIAsyncCreateTexture2D(uint32 SizeX,uint32 S
 		SubResourceData[MipIndex].SysMemSlicePitch = MipSize;
 	}
 
-	SafeCreateTexture2D(Direct3DDevice,&TextureDesc,SubResourceData,TextureResource.GetInitReference());
+	SafeCreateTexture2D(Direct3DDevice, Format, &TextureDesc,SubResourceData,TextureResource.GetInitReference());
 
 	if (TempBufferSize != ZeroBufferSize)
 	{
@@ -1625,6 +1628,7 @@ void* TD3D11Texture2D<RHIResourceType>::Lock(uint32 MipIndex,uint32 ArrayIndex,E
 		TRefCountPtr<ID3D11Texture2D> StagingTexture;
 		VERIFYD3D11CREATETEXTURERESULT(
 			D3DRHI->GetDevice()->CreateTexture2D(&StagingTextureDesc,NULL,StagingTexture.GetInitReference()),
+			RHIResourceType::GetFormat(),
 			this->GetSizeX(),
 			this->GetSizeY(),
 			this->GetSizeZ(),
diff --git a/Engine/Source/Runtime/Windows/D3D11RHI/Private/D3D11Util.cpp b/Engine/Source/Runtime/Windows/D3D11RHI/Private/D3D11Util.cpp
--- a/Engine/Source/Runtime/Windows/D3D11RHI/Private/D3D11Util.cpp
+++ b/Engine/Source/Runtime/Windows/D3D11RHI/Private/D3D11Util.cpp
@@ -280,17 +280,17 @@ void VerifyD3D11ShaderResult(FRHIShader* Shader, HRESULT D3DResult, const ANSICH
 	}
 }
 
-void VerifyD3D11CreateTextureResult(HRESULT D3DResult,const ANSICHAR* Code,const ANSICHAR* Filename,uint32 Line,uint32 SizeX,uint32 SizeY,uint32 SizeZ,uint8 Format,uint32 NumMips,uint32 Flags,
+void VerifyD3D11CreateTextureResult(HRESULT D3DResult, int32 UEFormat,const ANSICHAR* Code,const ANSICHAR* Filename,uint32 Line,uint32 SizeX,uint32 SizeY,uint32 SizeZ,uint8 D3DFormat,uint32 NumMips,uint32 Flags,
 	D3D11_USAGE Usage, uint32 CPUAccessFlags, uint32 MiscFlags, uint32 SampleCount, uint32 SampleQuality,
 	const void* SubResPtr, uint32 SubResPitch, uint32 SubResSlicePitch, ID3D11Device* Device)
 {
 	check(FAILED(D3DResult));
 
 	const FString ErrorString = GetD3D11ErrorString(D3DResult, 0);
-	const TCHAR* D3DFormatString = GetD3D11TextureFormatString((DXGI_FORMAT)Format);
+	const TCHAR* D3DFormatString = GetD3D11TextureFormatString((DXGI_FORMAT)D3DFormat);
 
 	UE_LOG(LogD3D11RHI, Error,
-		TEXT("%s failed \n at %s:%u \n with error %s, \n Size=%ix%ix%i Format=%s(0x%08X), NumMips=%i, Flags=%s, Usage:0x%x, CPUFlags:0x%x, MiscFlags:0x%x, SampleCount:0x%x, SampleQuality:0x%x, SubresPtr:0x%p, SubresPitch:%i, SubresSlicePitch:%i"),
+		TEXT("%s failed \n at %s:%u \n with error %s, \n Size=%ix%ix%i PF=%d D3DFormat=%s(0x%08X), NumMips=%i, Flags=%s, Usage:0x%x, CPUFlags:0x%x, MiscFlags:0x%x, SampleCount:0x%x, SampleQuality:0x%x, SubresPtr:0x%p, SubresPitch:%i, SubresSlicePitch:%i"),
 		ANSI_TO_TCHAR(Code),
 		ANSI_TO_TCHAR(Filename),
 		Line,
@@ -298,8 +298,9 @@ void VerifyD3D11CreateTextureResult(HRESULT D3DResult,const ANSICHAR* Code,const
 		SizeX,
 		SizeY,
 		SizeZ,
+		UEFormat,
 		D3DFormatString,
-		Format,
+		D3DFormat,
 		NumMips,
 		*GetD3D11TextureFlagString(Flags),
 		Usage,
@@ -315,7 +316,7 @@ void VerifyD3D11CreateTextureResult(HRESULT D3DResult,const ANSICHAR* Code,const
 	TerminateOnOutOfMemory(D3DResult, true);
 
 	UE_LOG(LogD3D11RHI, Fatal,
-		TEXT("%s failed \n at %s:%u \n with error %s, \n Size=%ix%ix%i Format=%s(0x%08X), NumMips=%i, Flags=%s, Usage:0x%x, CPUFlags:0x%x, MiscFlags:0x%x, SampleCount:0x%x, SampleQuality:0x%x, SubresPtr:0x%p, SubresPitch:%i, SubresSlicePitch:%i"),
+		TEXT("%s failed \n at %s:%u \n with error %s, \n Size=%ix%ix%i PF=%d Format=%s(0x%08X), NumMips=%i, Flags=%s, Usage:0x%x, CPUFlags:0x%x, MiscFlags:0x%x, SampleCount:0x%x, SampleQuality:0x%x, SubresPtr:0x%p, SubresPitch:%i, SubresSlicePitch:%i"),
 		ANSI_TO_TCHAR(Code),
 		ANSI_TO_TCHAR(Filename),
 		Line,
@@ -323,8 +324,9 @@ void VerifyD3D11CreateTextureResult(HRESULT D3DResult,const ANSICHAR* Code,const
 		SizeX,
 		SizeY,
 		SizeZ,
+		UEFormat,
 		D3DFormatString,
-		Format,
+		D3DFormat,
 		NumMips,
 		*GetD3D11TextureFlagString(Flags),
 		Usage,
diff --git a/Engine/Source/Runtime/Windows/D3D11RHI/Public/D3D11Util.h b/Engine/Source/Runtime/Windows/D3D11RHI/Public/D3D11Util.h
--- a/Engine/Source/Runtime/Windows/D3D11RHI/Public/D3D11Util.h
+++ b/Engine/Source/Runtime/Windows/D3D11RHI/Public/D3D11Util.h
@@ -67,7 +67,7 @@ extern D3D11RHI_API void VerifyD3D11ShaderResult(class FRHIShader* Shader, HRESU
 * @param	Filename - The filename of the source file containing Code.
 * @param	Line - The line number of Code within Filename.	
 */
-extern D3D11RHI_API void VerifyD3D11CreateTextureResult(HRESULT D3DResult,const ANSICHAR* Code,const ANSICHAR* Filename,uint32 Line,
+extern D3D11RHI_API void VerifyD3D11CreateTextureResult(HRESULT D3DResult, int32 UEFormat,const ANSICHAR* Code,const ANSICHAR* Filename,uint32 Line,
 										 uint32 SizeX,uint32 SizeY,uint32 SizeZ,uint8 D3DFormat,uint32 NumMips,uint32 Flags, D3D11_USAGE Usage,
 										 uint32 CPUAccessFlags, uint32 MiscFlags, uint32 SampleCount, uint32 SampleQuality,
 										 const void* SubResPtr, uint32 SubResPitch, uint32 SubResSlicePitch,ID3D11Device* Device);
@@ -84,7 +84,7 @@ extern D3D11RHI_API void VerifyD3D11ResizeViewportResult(HRESULT D3DResult, cons
 #define VERIFYD3D11RESULT_NOEXIT(x)		{HRESULT hr = x; if (FAILED(hr)) { VerifyD3D11ResultNoExit(hr,#x,__FILE__,__LINE__, 0); }}
 #define VERIFYD3D11SHADERRESULT(Result, Shader, Device) {HRESULT hr = (Result); if (FAILED(hr)) { VerifyD3D11ShaderResult(Shader, hr, #Result,__FILE__,__LINE__, Device); }}
 #define VERIFYD3D11RESULT_NOEXIT(x)		{HRESULT hr = x; if (FAILED(hr)) { VerifyD3D11ResultNoExit(hr,#x,__FILE__,__LINE__, 0); }}
-#define VERIFYD3D11CREATETEXTURERESULT(x,SizeX,SizeY,SizeZ,Format,NumMips,Flags,Usage,CPUAccessFlags,MiscFlags,SampleCount,SampleQuality,SubResPtr,SubResPitch,SubResSlicePitch,Device) {HRESULT hr = x; if (FAILED(hr)) { VerifyD3D11CreateTextureResult(hr,#x,__FILE__,__LINE__,SizeX,SizeY,SizeZ,Format,NumMips,Flags,Usage,CPUAccessFlags,MiscFlags,SampleCount,SampleQuality,SubResPtr,SubResPitch,SubResSlicePitch,Device); }}
+#define VERIFYD3D11CREATETEXTURERESULT(x,UEFormat,SizeX,SizeY,SizeZ,Format,NumMips,Flags,Usage,CPUAccessFlags,MiscFlags,SampleCount,SampleQuality,SubResPtr,SubResPitch,SubResSlicePitch,Device) {HRESULT hr = x; if (FAILED(hr)) { VerifyD3D11CreateTextureResult(hr, UEFormat,#x,__FILE__,__LINE__,SizeX,SizeY,SizeZ,Format,NumMips,Flags,Usage,CPUAccessFlags,MiscFlags,SampleCount,SampleQuality,SubResPtr,SubResPitch,SubResSlicePitch,Device); }}
 #define VERIFYD3D11RESIZEVIEWPORTRESULT(x,SizeX,SizeY,Format, Device) {HRESULT hr = x; if (FAILED(hr)) { VerifyD3D11ResizeViewportResult(hr,#x,__FILE__,__LINE__,SizeX,SizeY,Format, Device); }}
 /**
  * Checks that a COM object has the expected number of references.
