Merged CL9966351 from Dev-Rendering

Fix for vsync always enabled when running with iflip on D3D11

diff --git a/Engine/Source/Runtime/Windows/D3D11RHI/Private/D3D11Viewport.cpp b/Engine/Source/Runtime/Windows/D3D11RHI/Private/D3D11Viewport.cpp
--- a/Engine/Source/Runtime/Windows/D3D11RHI/Private/D3D11Viewport.cpp
+++ b/Engine/Source/Runtime/Windows/D3D11RHI/Private/D3D11Viewport.cpp
@@ -22,6 +22,15 @@
 		#include <dwmapi.h>
 #endif	//D3D11_WITH_DWMAPI
 
+#ifndef DXGI_PRESENT_ALLOW_TEARING
+#define DXGI_PRESENT_ALLOW_TEARING          0x00000200UL
+#endif
+
+#ifndef DXGI_SWAP_CHAIN_FLAG_ALLOW_TEARING
+#define DXGI_SWAP_CHAIN_FLAG_ALLOW_TEARING  2048
+#endif
+
+
 /**
  * RHI console variables used by viewports.
  */
@@ -335,7 +344,13 @@ bool FD3D11Viewport::PresentChecked(int32 SyncInterval)
 		if (SwapChain.IsValid())
 		{
 			// Present the back buffer to the viewport window.
-			Result = SwapChain->Present(SyncInterval, 0);
+			uint32 Flags = 0;
+			if( (D3D11GetSwapChainFlags() & DXGI_SWAP_CHAIN_FLAG_ALLOW_TEARING) != 0 && !SyncInterval && !bIsFullscreen )
+			{
+				Flags |= DXGI_PRESENT_ALLOW_TEARING;
+			}
+
+			Result = SwapChain->Present(SyncInterval, Flags);
 		}
 
 		if (IsValidRef(CustomPresent))
diff --git a/Engine/Source/Runtime/Windows/D3D11RHI/Private/Windows/WindowsD3D11Viewport.cpp b/Engine/Source/Runtime/Windows/D3D11RHI/Private/Windows/WindowsD3D11Viewport.cpp
--- a/Engine/Source/Runtime/Windows/D3D11RHI/Private/Windows/WindowsD3D11Viewport.cpp
+++ b/Engine/Source/Runtime/Windows/D3D11RHI/Private/Windows/WindowsD3D11Viewport.cpp
@@ -13,13 +13,10 @@
 
 #include "dxgi1_6.h"
 
-#ifndef DXGI_PRESENT_ALLOW_TEARING
-#define DXGI_PRESENT_ALLOW_TEARING          0x00000200UL
+#ifndef DXGI_SWAP_CHAIN_FLAG_ALLOW_TEARING
 #define DXGI_SWAP_CHAIN_FLAG_ALLOW_TEARING  2048
 #endif
 
-
-
 static bool GSwapFlagsInitialized = false;
 static DXGI_SWAP_EFFECT GSwapEffect = DXGI_SWAP_EFFECT_DISCARD;
 static uint32 GSwapChainFlags = DXGI_SWAP_CHAIN_FLAG_ALLOW_MODE_SWITCH;
