Move the default implementation of GetIndexForDevice to the base class

diff --git a/Engine/Source/Runtime/AudioMixerCore/Private/AudioMixer.cpp b/Engine/Source/Runtime/AudioMixerCore/Private/AudioMixer.cpp
--- a/Engine/Source/Runtime/AudioMixerCore/Private/AudioMixer.cpp
+++ b/Engine/Source/Runtime/AudioMixerCore/Private/AudioMixer.cpp
@@ -341,6 +341,33 @@ namespace Audio
 		bIsDeviceInitialized = true;
 	}
 
+	int32 IAudioMixerPlatformInterface::GetIndexForDevice(const FString& InDeviceName)
+	{
+		uint32 TotalNumDevices = 0;
+
+		if (!GetNumOutputDevices(TotalNumDevices))
+		{
+			return INDEX_NONE;
+		}
+
+		// Iterate through every device and see if
+		for (uint32 DeviceIndex = 0; DeviceIndex < TotalNumDevices; DeviceIndex++)
+		{
+			FAudioPlatformDeviceInfo DeviceInfo;
+			if (GetOutputDeviceInfo(DeviceIndex, DeviceInfo))
+			{
+				// check if the device name matches the input device name:
+				if (DeviceInfo.Name.Contains(InDeviceName))
+				{
+					return DeviceIndex;
+				}
+			}
+		}
+
+		// If we've made it here, we weren't able to find a matching device.
+		return INDEX_NONE;
+	}
+
 	template<typename BufferType>
 	void IAudioMixerPlatformInterface::ApplyAttenuationInternal(BufferType* BufferDataPtr, const int32 NumFrames)
 	{
diff --git a/Engine/Source/Runtime/AudioMixerCore/Public/AudioMixer.h b/Engine/Source/Runtime/AudioMixerCore/Public/AudioMixer.h
--- a/Engine/Source/Runtime/AudioMixerCore/Public/AudioMixer.h
+++ b/Engine/Source/Runtime/AudioMixerCore/Public/AudioMixer.h
@@ -396,9 +396,9 @@ namespace Audio
 		/**
 		 * Can be used to look up the current index for a given device name.
 		 * On most platforms, this index may be invalidated if any devices are added or removed.
-		 * Returns INDEX_NONE if 
+		 * Returns INDEX_NONE if no mapping is found
 		 */
-		virtual int32 GetIndexForDevice(const FString& InDeviceName) { return INDEX_NONE; }
+		virtual int32 GetIndexForDevice(const FString& InDeviceName);
 
 		/** Gets the platform specific audio settings. */
 		virtual FAudioPlatformSettings GetPlatformSettings() const = 0;
@@ -627,4 +627,4 @@ public:
 	virtual bool IsAudioMixerModule() const { return false; }
 	virtual FAudioDevice* CreateAudioDevice() { return nullptr; }
 	virtual Audio::IAudioMixerPlatformInterface* CreateAudioMixerPlatformInterface() { return nullptr; }
-};
\ No newline at end of file
+};
diff --git a/Engine/Source/Runtime/Linux/AudioMixerSDL/Private/AudioMixerSDL.cpp b/Engine/Source/Runtime/Linux/AudioMixerSDL/Private/AudioMixerSDL.cpp
--- a/Engine/Source/Runtime/Linux/AudioMixerSDL/Private/AudioMixerSDL.cpp
+++ b/Engine/Source/Runtime/Linux/AudioMixerSDL/Private/AudioMixerSDL.cpp
@@ -25,6 +25,11 @@ public:
 	{
 		return new Audio::FMixerDevice(new Audio::FAudioMixerPlatformSDL());
 	}
+
+	virtual Audio::IAudioMixerPlatformInterface* CreateAudioMixerPlatformInterface() override
+	{
+		return new Audio::FAudioMixerPlatformSDL();
+	}
 };
 
 IMPLEMENT_MODULE(FAudioMixerModuleSDL, AudioMixerSDL);
diff --git a/Engine/Source/Runtime/Windows/AudioMixerXAudio2/Private/AudioMixerPlatformXAudio2.cpp b/Engine/Source/Runtime/Windows/AudioMixerXAudio2/Private/AudioMixerPlatformXAudio2.cpp
--- a/Engine/Source/Runtime/Windows/AudioMixerXAudio2/Private/AudioMixerPlatformXAudio2.cpp
+++ b/Engine/Source/Runtime/Windows/AudioMixerXAudio2/Private/AudioMixerPlatformXAudio2.cpp
@@ -524,33 +524,6 @@ namespace Audio
 		return true;
 	}
 
-	int32 FMixerPlatformXAudio2::GetIndexForDevice(const FString& InDeviceName)
-	{
-		uint32 TotalNumDevices = 0;
-
-		if (!GetNumOutputDevices(TotalNumDevices))
-		{
-			return INDEX_NONE;
-		}
-
-		// Iterate through every device and see if
-		for (uint32 DeviceIndex = 0; DeviceIndex < TotalNumDevices; DeviceIndex++)
-		{
-			FAudioPlatformDeviceInfo DeviceInfo;
-			if (GetOutputDeviceInfo(DeviceIndex, DeviceInfo))
-			{
-				// check if the device name matches the input device name:
-				if (DeviceInfo.Name.Contains(InDeviceName))
-				{
-					return DeviceIndex;
-				}
-			}
-		}
-
-		// If we've made it here, we weren't able to find a matching device.
-		return INDEX_NONE;
-	}
-
 	bool FMixerPlatformXAudio2::GetDefaultOutputDeviceIndex(uint32& OutDefaultDeviceIndex) const
 	{
 		OutDefaultDeviceIndex = 0;
diff --git a/Engine/Source/Runtime/Windows/AudioMixerXAudio2/Private/AudioMixerPlatformXAudio2.h b/Engine/Source/Runtime/Windows/AudioMixerXAudio2/Private/AudioMixerPlatformXAudio2.h
--- a/Engine/Source/Runtime/Windows/AudioMixerXAudio2/Private/AudioMixerPlatformXAudio2.h
+++ b/Engine/Source/Runtime/Windows/AudioMixerXAudio2/Private/AudioMixerPlatformXAudio2.h
@@ -53,7 +53,6 @@ namespace Audio
 		virtual bool GetNumOutputDevices(uint32& OutNumOutputDevices) override;
 		virtual bool GetOutputDeviceInfo(const uint32 InDeviceIndex, FAudioPlatformDeviceInfo& OutInfo) override;
 		virtual bool GetDefaultOutputDeviceIndex(uint32& OutDefaultDeviceIndex) const override;
-		virtual int32 GetIndexForDevice(const FString& InDeviceName) override;
 		virtual bool OpenAudioStream(const FAudioMixerOpenStreamParams& Params) override;
 		virtual bool CloseAudioStream() override;
 		virtual bool StartAudioStream() override;
