Fix compile error

diff --git a/Engine/Source/Runtime/Windows/AudioMixerXAudio2/Private/AudioMixerPlatformXAudio2.cpp b/Engine/Source/Runtime/Windows/AudioMixerXAudio2/Private/AudioMixerPlatformXAudio2.cpp
--- a/Engine/Source/Runtime/Windows/AudioMixerXAudio2/Private/AudioMixerPlatformXAudio2.cpp
+++ b/Engine/Source/Runtime/Windows/AudioMixerXAudio2/Private/AudioMixerPlatformXAudio2.cpp
@@ -503,13 +503,20 @@ namespace Audio
 	}
 #endif //  #if PLATFORM_WINDOWS
 
-	bool FMixerPlatformXAudio2::GetOutputDeviceInfoInternal(const uint32 InDeviceIndex, FAudioPlatformDeviceInfo& OutInfo)
+	bool FMixerPlatformXAudio2::GetOutputDeviceInfo(const uint32 InDeviceIndex, FAudioPlatformDeviceInfo& OutInfo)
 	{
+		if (!bIsInitialized)
+		{
+			AUDIO_PLATFORM_ERROR(TEXT("XAudio2 was not initialized."));
+			return false;
+		}
+
 #if PLATFORM_WINDOWS
 		IMMDeviceEnumerator* DeviceEnumerator = nullptr;
 		IMMDeviceCollection* DeviceCollection = nullptr;
 		IMMDevice* Device = nullptr;
 		bool bIsDefault = false;
+		bool bSucceeded = false;
 
 		HRESULT Result = CoCreateInstance(__uuidof(MMDeviceEnumerator), NULL, CLSCTX_INPROC_SERVER, IID_PPV_ARGS(&DeviceEnumerator));
 		XAUDIO2_GOTO_CLEANUP_ON_FAIL(Result);
@@ -552,7 +559,7 @@ namespace Audio
 		}
 
 		check(Device);
-		GetMMDeviceInfo(Device, OutInfo);
+		bSucceeded = GetMMDeviceInfo(Device, OutInfo);
 
 		// Fix up if this was a default device
 		OutInfo.bIsSystemDefault = bIsDefault;
@@ -562,7 +569,7 @@ namespace Audio
 		SAFE_RELEASE(DeviceCollection);
 		SAFE_RELEASE(DeviceEnumerator);
 
-		return SUCCEEDED(Result);
+		return bSucceeded && SUCCEEDED(Result);
 
 #elif PLATFORM_HOLOLENS // #if PLATFORM_WINDOWS
 
@@ -724,26 +731,12 @@ namespace Audio
 				UE_LOG(LogAudioMixer, Display, TEXT("%d: %s"), i, EAudioMixerChannel::ToString(OutInfo.OutputChannelArray[i]));
 			}
 		}
-
-#endif // #elif PLATFORM_HOLOLENS // #if PLATFORM_WINDOWS
-	}
-
-	bool FMixerPlatformXAudio2::GetOutputDeviceInfo(const uint32 InDeviceIndex, FAudioPlatformDeviceInfo& OutInfo)
-	{
-		if (!bIsInitialized)
-		{
-			AUDIO_PLATFORM_ERROR(TEXT("XAudio2 was not initialized."));
-			return false;
-		}
-
-#if PLATFORM_HOLOLENS || PLATFORM_WINDOWS
-		return GetOutputDeviceInfoInternal(InDeviceIndex, OutInfo);
-#else
+#else // #elif PLATFORM_HOLOLENS // #if PLATFORM_WINDOWS
 		OutInfo.bIsSystemDefault = true;
 		OutInfo.SampleRate = 44100;
 		OutInfo.DeviceId = 0;
 		OutInfo.Format = EAudioMixerStreamDataFormat::Float;
-		OutInfo.Name = TEXT("XboxOne Audio Device.");
+		OutInfo.Name = TEXT("Audio Device.");
 		OutInfo.NumChannels = 8;
 
 		OutInfo.OutputChannelArray.Add(EAudioMixerChannel::FrontLeft);
@@ -754,7 +747,7 @@ namespace Audio
 		OutInfo.OutputChannelArray.Add(EAudioMixerChannel::BackRight);
 		OutInfo.OutputChannelArray.Add(EAudioMixerChannel::SideLeft);
 		OutInfo.OutputChannelArray.Add(EAudioMixerChannel::SideRight);
-#endif // #else // #if PLATFORM_HOLOLENS || PLATFORM_WINDOWS
+#endif // #else // #elif PLATFORM_HOLOLENS // #if PLATFORM_WINDOWS
 		return true;
 	}
 
diff --git a/Engine/Source/Runtime/Windows/AudioMixerXAudio2/Private/AudioMixerPlatformXAudio2.h b/Engine/Source/Runtime/Windows/AudioMixerXAudio2/Private/AudioMixerPlatformXAudio2.h
--- a/Engine/Source/Runtime/Windows/AudioMixerXAudio2/Private/AudioMixerPlatformXAudio2.h
+++ b/Engine/Source/Runtime/Windows/AudioMixerXAudio2/Private/AudioMixerPlatformXAudio2.h
@@ -92,8 +92,6 @@ namespace Audio
 
 	private:
 
-		bool GetOutputDeviceInfoInternal(const uint32 InDeviceIndex, FAudioPlatformDeviceInfo& OutInfo);
-
 		bool AllowDeviceSwap();
 
 		// Used to teardown and reinitialize XAudio2.
