Fixing compiile error

diff --git a/Engine/Source/Runtime/Windows/AudioMixerXAudio2/Private/AudioMixerPlatformXAudio2.cpp b/Engine/Source/Runtime/Windows/AudioMixerXAudio2/Private/AudioMixerPlatformXAudio2.cpp
--- a/Engine/Source/Runtime/Windows/AudioMixerXAudio2/Private/AudioMixerPlatformXAudio2.cpp
+++ b/Engine/Source/Runtime/Windows/AudioMixerXAudio2/Private/AudioMixerPlatformXAudio2.cpp
@@ -146,14 +146,14 @@ namespace Audio
 		, bIsDeviceOpen(false)
 	{
 #if PLATFORM_WINDOWS || PLATFORM_HOLOLENS
-		FWindowsPlatformMisc::CoInitialize();
+		FPlatformMisc::CoInitialize();
 #endif // #if PLATFORM_WINDOWS || PLATFORM_HOLOLENS
 	}
 
 	FMixerPlatformXAudio2::~FMixerPlatformXAudio2()
 	{
 #if PLATFORM_WINDOWS || PLATFORM_HOLOLENS
-		FWindowsPlatformMisc::CoUninitialize();
+		FPlatformMisc::CoUninitialize();
 #endif // #if PLATFORM_WINDOWS || PLATFORM_HOLOLENS
 	}
 
