Fix compile error
#rnx

diff --git a/Engine/Source/Runtime/Windows/AudioMixerXAudio2/Private/AudioMixerPlatformXAudio2.cpp b/Engine/Source/Runtime/Windows/AudioMixerXAudio2/Private/AudioMixerPlatformXAudio2.cpp
--- a/Engine/Source/Runtime/Windows/AudioMixerXAudio2/Private/AudioMixerPlatformXAudio2.cpp
+++ b/Engine/Source/Runtime/Windows/AudioMixerXAudio2/Private/AudioMixerPlatformXAudio2.cpp
@@ -642,7 +642,7 @@ namespace Audio
 			// case of a non-contiguous subset of channels. For example, if a stream contains left, bass enhance and right, then channel 1 is left, channel 2 is right, 
 			// and channel 3 is bass enhance. This enables the linkage of multi-channel streams to well-defined multi-speaker configurations.
 
-			check(EAudioMixerChannel::ChannelTypeCount == ChannelTypeMap.Num());
+			check(EAudioMixerChannel::ChannelTypeCount == UE_ARRAY_COUNT(ChannelTypeMap));
 			uint32 ChanCount = 0;
 			for (uint32 ChannelTypeIndex = 0; ChannelTypeIndex < EAudioMixerChannel::ChannelTypeCount && ChanCount < (uint32)OutInfo.NumChannels; ++ChannelTypeIndex)
 			{
