Removing CommonAudioPoolSize (unused code that can cause crash in certain scenarios)

diff --git a/Engine/Source/Runtime/Engine/Private/AudioDevice.cpp b/Engine/Source/Runtime/Engine/Private/AudioDevice.cpp
--- a/Engine/Source/Runtime/Engine/Private/AudioDevice.cpp
+++ b/Engine/Source/Runtime/Engine/Private/AudioDevice.cpp
@@ -210,9 +210,6 @@ FAudioDevice::FAudioDevice()
 	: NumStoppingSources(32)
 	, SampleRate(0)
 	, NumPrecacheFrames(MONO_PCM_BUFFER_SAMPLES)
-	, CommonAudioPoolSize(0)
-	, CommonAudioPool(nullptr)
-	, CommonAudioPoolFreeBytes(0)
 	, DeviceHandle(static_cast<Audio::FDeviceId>(INDEX_NONE))
 	, SpatializationPluginInterface(nullptr)
 	, ReverbPluginInterface(nullptr)
@@ -322,8 +319,6 @@ bool FAudioDevice::Init(int32 InMaxSources)
 	// Mixed sample rate is set by the platform
 	SampleRate = PlatformSettings.SampleRate;
 
-	verify(GConfig->GetInt(TEXT("Audio"), TEXT("CommonAudioPoolSize"), CommonAudioPoolSize, GEngineIni));
-
 	// If this is true, skip the initial startup precache so we can do it later in the flow
 	GConfig->GetBool(TEXT("Audio"), TEXT("DeferStartupPrecache"), bDeferStartupPrecache, GEngineIni);
 
@@ -4269,20 +4264,6 @@ void FAudioDevice::Update(bool bGameTicking)
 
 	// send any needed information back to the game thread
 	SendUpdateResultsToGameThread(FirstActiveIndex);
-
-#if !UE_BUILD_SHIPPING
-	// Print statistics for first non initial load allocation.
-	static bool bFirstTime = true;
-	if (bFirstTime && CommonAudioPoolSize != 0)
-	{
-		bFirstTime = false;
-		if (CommonAudioPoolFreeBytes != 0)
-		{
-			UE_LOG(LogAudio, Log, TEXT("Audio pool size mismatch by %d bytes. Please update CommonAudioPoolSize ini setting to %d to avoid waste!"),
-				CommonAudioPoolFreeBytes, CommonAudioPoolSize - CommonAudioPoolFreeBytes);
-		}
-	}
-#endif
 }
 
 void FAudioDevice::SendUpdateResultsToGameThread(const int32 FirstActiveIndex)
diff --git a/Engine/Source/Runtime/Engine/Public/AudioDevice.h b/Engine/Source/Runtime/Engine/Public/AudioDevice.h
--- a/Engine/Source/Runtime/Engine/Public/AudioDevice.h
+++ b/Engine/Source/Runtime/Engine/Public/AudioDevice.h
@@ -1677,15 +1677,6 @@ public:
 	/** The number of frames to precache. */
 	int32 NumPrecacheFrames;
 
-	/** The amount of memory to reserve for always resident sounds */
-	int32 CommonAudioPoolSize;
-
-	/** Pointer to permanent memory allocation stack. */
-	void* CommonAudioPool;
-
-	/** Available size in permanent memory stack */
-	int32 CommonAudioPoolFreeBytes;
-
 	/** The handle for this audio device used in the audio device manager. */
 	Audio::FDeviceId DeviceHandle;
 
diff --git a/Engine/Source/Runtime/Windows/XAudio2/Private/XAudio2Device.cpp b/Engine/Source/Runtime/Windows/XAudio2/Private/XAudio2Device.cpp
--- a/Engine/Source/Runtime/Windows/XAudio2/Private/XAudio2Device.cpp
+++ b/Engine/Source/Runtime/Windows/XAudio2/Private/XAudio2Device.cpp
@@ -268,19 +268,6 @@ bool FXAudio2Device::InitializeHardware()
 	// Set that we initialized our hardware audio device ok so we should use real voices.
 	bIsAudioDeviceHardwareInitialized = true;
 
-	// Initialize permanent memory stack for initial & always loaded sound allocations.
-	if( CommonAudioPoolSize )
-	{
-		UE_LOG(LogAudio, Log, TEXT( "Allocating %g MByte for always resident audio data" ), CommonAudioPoolSize / ( 1024.0f * 1024.0f ) );
-		CommonAudioPoolFreeBytes = CommonAudioPoolSize;
-		CommonAudioPool = ( uint8* )FMemory::Malloc( CommonAudioPoolSize );
-	}
-	else
-	{
-		UE_LOG(LogAudio, Log, TEXT( "CommonAudioPoolSize is set to 0 - disabling persistent pool for audio data" ) );
-		CommonAudioPoolFreeBytes = 0;
-	}
-
 #if WITH_XMA2
 	FXMAAudioInfo::Initialize();
 #endif
@@ -708,41 +695,6 @@ bool FXAudio2Device::Exec( UWorld* InWorld, const TCHAR* Cmd, FOutputDevice& Ar
 	return( false );
 }
 
-/**
- * Allocates memory from permanent pool. This memory will NEVER be freed.
- *
- * @param	Size	Size of allocation.
- *
- * @return pointer to a chunk of memory with size Size
- */
-void* FXAudio2Device::AllocatePermanentMemory( int32 Size, bool& AllocatedInPool )
-{
-	void* Allocation = NULL;
-	
-	// Fall back to using regular allocator if there is not enough space in permanent memory pool.
-	if( Size > CommonAudioPoolFreeBytes )
-	{
-		Allocation = FMemory::Malloc( Size );
-		check( Allocation );
-
-		AllocatedInPool = false;
-	}
-	// Allocate memory from pool.
-	else
-	{
-		uint8* CommonAudioPoolAddress = ( uint8* )CommonAudioPool;
-		Allocation = CommonAudioPoolAddress + ( CommonAudioPoolSize - CommonAudioPoolFreeBytes );
-
-		AllocatedInPool = true;
-	}
-
-	// Decrement available size regardless of whether we allocated from pool or used regular allocator
-	// to allow us to log suggested size at the end of initial loading.
-	CommonAudioPoolFreeBytes -= Size;
-	
-	return( Allocation );
-}
-
 FAudioPlatformSettings FXAudio2Device::GetPlatformSettings() const
 {
 	return FAudioPlatformSettings::GetPlatformSettings(FPlatformProperties::GetRuntimeSettingsClassName());
diff --git a/Engine/Source/Runtime/Windows/XAudio2/Public/XAudio2Device.h b/Engine/Source/Runtime/Windows/XAudio2/Public/XAudio2Device.h
--- a/Engine/Source/Runtime/Windows/XAudio2/Public/XAudio2Device.h
+++ b/Engine/Source/Runtime/Windows/XAudio2/Public/XAudio2Device.h
@@ -188,16 +188,6 @@ protected:
 	/** Creates a mastering voice. Returns false if fails. */
 	struct IXAudio2MasteringVoice* CreateMasteringVoice();
 
-	/**
-     * Allocates memory from permanent pool. This memory will NEVER be freed.
-	 *
-	 * @param	Size	Size of allocation.
-	 * @param	AllocatedInPool	(OUT) True if the allocation occurred from the pool; false if it was a regular physical allocation.
-	 *
-	 * @return pointer to a chunk of memory with size Size
-	 */
-	void* AllocatePermanentMemory( int32 Size, /*OUT*/ bool& AllocatedInPool );
-
 	/** Retrieves the platform settings */
 	virtual FAudioPlatformSettings GetPlatformSettings() const override;
 
