Fix for r.D3D11.UseAllowTearing=0 crashing on startup
revert earlier bFullscreenLost, as it broke exclusive fullscreen.

diff --git a/Engine/Source/Runtime/Windows/D3D11RHI/Private/D3D11Viewport.cpp b/Engine/Source/Runtime/Windows/D3D11RHI/Private/D3D11Viewport.cpp
--- a/Engine/Source/Runtime/Windows/D3D11RHI/Private/D3D11Viewport.cpp
+++ b/Engine/Source/Runtime/Windows/D3D11RHI/Private/D3D11Viewport.cpp
@@ -539,7 +539,6 @@ bool FD3D11Viewport::Present(bool bLockToVsync)
 		// Can't compare BOOL with bool...
 		if ( (!!bSwapChainFullscreenState)  != bIsFullscreen )
 		{
-			bFullscreenLost = true;
 			bIsValid = false;
 		}
 	}
diff --git a/Engine/Source/Runtime/Windows/D3D11RHI/Private/Windows/WindowsD3D11Viewport.cpp b/Engine/Source/Runtime/Windows/D3D11RHI/Private/Windows/WindowsD3D11Viewport.cpp
--- a/Engine/Source/Runtime/Windows/D3D11RHI/Private/Windows/WindowsD3D11Viewport.cpp
+++ b/Engine/Source/Runtime/Windows/D3D11RHI/Private/Windows/WindowsD3D11Viewport.cpp
@@ -178,13 +178,14 @@ FD3D11Viewport::FD3D11Viewport(FD3D11DynamicRHI* InD3DRHI,HWND InWindowHandle,ui
 			IDXGISwapChain1* SwapChain1 = nullptr;
 			IDXGIFactory2* Factory2 = (IDXGIFactory2*)D3DRHI->GetFactory();
 
-			VERIFYD3D11RESULT_EX((Factory2->CreateSwapChainForHwnd(D3DRHI->GetDevice(), WindowHandle, &SwapChainDesc, &FSSwapChainDesc, nullptr, &SwapChain1)), D3DRHI->GetDevice());
+			if(!FAILED(Factory2->CreateSwapChainForHwnd(D3DRHI->GetDevice(), WindowHandle, &SwapChainDesc, &FSSwapChainDesc, nullptr, &SwapChain1)))
+			{
 				SwapChain1->QueryInterface(__uuidof(IDXGISwapChain1), (void**)SwapChain.GetInitReference());
 
 				// See if we are running on a HDR monitor 
 				CheckHDRMonitorStatus();
 			}
-
+		}
 
 		if (SwapChain == nullptr)
 		{
