UE-74862 : Reusable media encoder framework implementation

diff --git a/Engine/Source/Runtime/Windows/WindowsPlatformFeatures/Private/WmfMp4Writer.cpp b/Engine/Source/Runtime/Windows/WindowsPlatformFeatures/Private/WmfMp4Writer.cpp
--- a/Engine/Source/Runtime/Windows/WindowsPlatformFeatures/Private/WmfMp4Writer.cpp
+++ b/Engine/Source/Runtime/Windows/WindowsPlatformFeatures/Private/WmfMp4Writer.cpp
@@ -1,8 +1,7 @@
 // Copyright 1998-2019 Epic Games, Inc. All Rights Reserved.
 
 #include "WmfMp4Writer.h"
-
-#include "GameplayMediaEncoderSample.h"
+#include "Microsoft/AVEncoderIMFSampleWrapper.h"
 
 #if WMFMEDIA_SUPPORTED_PLATFORM
 	#pragma comment(lib, "mfplat")
@@ -23,12 +22,78 @@ bool FWmfMp4Writer::Initialize(const TCHAR* Filename)
 	return true;
 }
 
-bool FWmfMp4Writer::CreateStream(IMFMediaType* StreamType, DWORD& StreamIndex)
+TOptional<DWORD> FWmfMp4Writer::CreateAudioStream(const FString& Codec, const AVEncoder::FAudioEncoderConfig& Config)
 {
-	CHECK_HR(Writer->AddStream(StreamType, &StreamIndex));
+	GUID Format;
+	int AudioBitsPerSample = 0;
+	if (Codec=="aac")
+	{
+		Format = MFAudioFormat_AAC;
+		AudioBitsPerSample = 16;
+	}
+#if 0
+	else if (Codec=="mp3")
+	{
+		Format = MFAudioFormat_MP3;
+		AudioBitsPerSample = 16;
+	}
+#endif
+	else
+	{
+		UE_LOG(WMF, Error, TEXT("FWmfMp4Writer doesn't support codec %s"), *Codec);
+		return {};
+	}
+
+	TRefCountPtr<IMFMediaType> MediaType;
+	CHECK_HR_DEFAULT(MFCreateMediaType(MediaType.GetInitReference()));
+	CHECK_HR_DEFAULT(MediaType->SetGUID(MF_MT_MAJOR_TYPE, MFMediaType_Audio));
+	CHECK_HR_DEFAULT(MediaType->SetGUID(MF_MT_SUBTYPE, Format));
+	CHECK_HR_DEFAULT(MediaType->SetUINT32(MF_MT_AUDIO_BITS_PER_SAMPLE, AudioBitsPerSample));
+	CHECK_HR_DEFAULT(MediaType->SetUINT32(MF_MT_AUDIO_SAMPLES_PER_SECOND, Config.Samplerate));
+	CHECK_HR_DEFAULT(MediaType->SetUINT32(MF_MT_AUDIO_NUM_CHANNELS, Config.NumChannels));
+	CHECK_HR_DEFAULT(MediaType->SetUINT32(MF_MT_AUDIO_AVG_BYTES_PER_SECOND, Config.Bitrate));
+
+	DWORD StreamIndex = 0;
+	CHECK_HR_DEFAULT(Writer->AddStream(MediaType, &StreamIndex));
 	// no transcoding here so input type is the same as output type
-	CHECK_HR(Writer->SetInputMediaType(StreamIndex, StreamType, nullptr));
-	return true;
+	CHECK_HR_DEFAULT(Writer->SetInputMediaType(StreamIndex, MediaType, nullptr));
+	return TOptional<DWORD>(StreamIndex);
+}
+
+TOptional<DWORD> FWmfMp4Writer::CreateVideoStream(const FString& Codec, const AVEncoder::FVideoEncoderConfig& Config)
+{
+	GUID Format;
+	if (Codec == "h264")
+	{
+		Format = MFVideoFormat_H264;
+	}
+#if 0
+	else if (Codec == "h265")
+	{
+		Format = MFVideoFormat_H265;
+	}
+#endif
+	else
+	{
+		UE_LOG(WMF, Error, TEXT("FWmfMp4Writer doesn't support codec %s"), *Codec);
+		return {};
+	}
+
+	TRefCountPtr<IMFMediaType> MediaType;
+	CHECK_HR_DEFAULT(MFCreateMediaType(MediaType.GetInitReference()));
+	CHECK_HR_DEFAULT(MediaType->SetGUID(MF_MT_MAJOR_TYPE, MFMediaType_Video));
+	CHECK_HR_DEFAULT(MediaType->SetGUID(MF_MT_SUBTYPE, Format));
+	CHECK_HR_DEFAULT(MediaType->SetUINT32(MF_MT_AVG_BITRATE, Config.Bitrate));
+	CHECK_HR_DEFAULT(MFSetAttributeRatio(MediaType, MF_MT_FRAME_RATE, Config.Framerate, 1));
+	CHECK_HR_DEFAULT(MFSetAttributeSize(MediaType, MF_MT_FRAME_SIZE, Config.Width, Config.Height));
+	CHECK_HR_DEFAULT(MFSetAttributeRatio(MediaType, MF_MT_PIXEL_ASPECT_RATIO, 1, 1));
+	CHECK_HR_DEFAULT(MediaType->SetUINT32(MF_MT_INTERLACE_MODE, MFVideoInterlace_Progressive));
+
+	DWORD StreamIndex = 0;
+	CHECK_HR_DEFAULT(Writer->AddStream(MediaType, &StreamIndex));
+	// no transcoding here so input type is the same as output type
+	CHECK_HR_DEFAULT(Writer->SetInputMediaType(StreamIndex, MediaType, nullptr));
+	return TOptional<DWORD>(StreamIndex);
 }
 
 bool FWmfMp4Writer::Start()
@@ -37,11 +102,30 @@ bool FWmfMp4Writer::Start()
 	return true;
 }
 
-bool FWmfMp4Writer::Write(const FGameplayMediaEncoderSample& Sample, DWORD StreamIndex)
+bool FWmfMp4Writer::Write(const AVEncoder::FAVPacket& InSample, DWORD StreamIndex)
 {
+	AVEncoder::FIMFSampleWrapper Sample { InSample.Type };
+
+	if (!Sample.CreateSample())
+	{
+		return false;
+	}
+
+	TRefCountPtr<IMFMediaBuffer> WmfBuffer;
+	CHECK_HR_DEFAULT(MFCreateAlignedMemoryBuffer(InSample.Data.Num(), MF_1_BYTE_ALIGNMENT, WmfBuffer.GetInitReference()));
+	uint8* Dst = nullptr;
+	CHECK_HR_DEFAULT(WmfBuffer->Lock(&Dst, nullptr, nullptr));
+	FMemory::Memcpy(Dst, InSample.Data.GetData(), InSample.Data.Num());
+	CHECK_HR_DEFAULT(WmfBuffer->Unlock());
+
+	WmfBuffer->SetCurrentLength(InSample.Data.Num());
+	Sample.GetSample()->AddBuffer(WmfBuffer);
+	Sample.SetTime(InSample.Timestamp);
+	Sample.SetDuration(InSample.Duration);
+
 	CHECK_HR(Writer->WriteSample(StreamIndex, const_cast<IMFSample*>(Sample.GetSample())));
 
-	UE_LOG(MP4, VeryVerbose, TEXT("stream #%d: time %.3f, duration %.3f%s"), static_cast<int>(Sample.GetType()), Sample.GetTime().GetTotalSeconds(), Sample.GetDuration().GetTotalSeconds(), Sample.IsVideoKeyFrame() ? TEXT(", key-frame") : TEXT(""));
+	UE_LOG(MP4, VeryVerbose, TEXT("stream #%d: time %.3f, duration %.3f%s"), StreamIndex, Sample.GetTime().GetTotalSeconds(), Sample.GetDuration().GetTotalSeconds(), InSample.IsVideoKeyFrame() ? TEXT(", key-frame") : TEXT(""));
 
 	return true;
 }
diff --git a/Engine/Source/Runtime/Windows/WindowsPlatformFeatures/Private/WmfMp4Writer.h b/Engine/Source/Runtime/Windows/WindowsPlatformFeatures/Private/WmfMp4Writer.h
--- a/Engine/Source/Runtime/Windows/WindowsPlatformFeatures/Private/WmfMp4Writer.h
+++ b/Engine/Source/Runtime/Windows/WindowsPlatformFeatures/Private/WmfMp4Writer.h
@@ -5,15 +5,26 @@
 #include "WmfPrivate.h"
 
 #include "Math/NumericLimits.h"
-#include "GameplayMediaEncoderSample.h"
+#include "AVEncoder.h"
+#include "Misc/Optional.h"
 
 class FWmfMp4Writer final
 {
 public:
 	bool Initialize(const TCHAR* Filename);
-	bool CreateStream(IMFMediaType* StreamType, DWORD& StreamIndex);
+
+	/**
+	 * Create an audio stream and return the its index on success
+	 */
+	TOptional<DWORD> CreateAudioStream(const FString& Codec, const AVEncoder::FAudioEncoderConfig& Config);
+
+	/**
+	 * Create a video stream and return the its index on success
+	 */
+	TOptional<DWORD> CreateVideoStream(const FString& Codec, const AVEncoder::FVideoEncoderConfig& Config);
+
 	bool Start();
-	bool Write(const FGameplayMediaEncoderSample& Sample, DWORD StreamIndex);
+	bool Write(const AVEncoder::FAVPacket& InSample, DWORD StreamIndex);
 	bool Finalize();
 
 private:
diff --git a/Engine/Source/Runtime/Windows/WindowsPlatformFeatures/Private/WmfPrivate.h b/Engine/Source/Runtime/Windows/WindowsPlatformFeatures/Private/WmfPrivate.h
--- a/Engine/Source/Runtime/Windows/WindowsPlatformFeatures/Private/WmfPrivate.h
+++ b/Engine/Source/Runtime/Windows/WindowsPlatformFeatures/Private/WmfPrivate.h
@@ -43,10 +43,8 @@ inline const FString GetComErrorDescription(HRESULT Res)
 	}
 }
 
-
 #include "Windows/HideWindowsPlatformTypes.h"
 
-
 // macro to deal with COM calls inside a function that returns `false` on error
 #define CHECK_HR(COM_call)\
 	{\
@@ -80,3 +78,15 @@ inline const FString GetComErrorDescription(HRESULT Res)
 		}\
 	}
 
+// macro to deal with COM calls inside a function and return `{}` on error
+// This has the advantage of being able to return any type that can be initialized with `{}`
+#define CHECK_HR_DEFAULT(COM_call)\
+	{\
+		HRESULT Res = COM_call;\
+		if (FAILED(Res))\
+		{\
+			UE_LOG(WMF, Error, TEXT("`" #COM_call "` failed: 0x%X - %s"), Res, *GetComErrorDescription(Res));\
+			return {};\
+		}\
+	}
+
diff --git a/Engine/Source/Runtime/Windows/WindowsPlatformFeatures/Private/WmfRingBuffer.cpp b/Engine/Source/Runtime/Windows/WindowsPlatformFeatures/Private/WmfRingBuffer.cpp
--- a/Engine/Source/Runtime/Windows/WindowsPlatformFeatures/Private/WmfRingBuffer.cpp
+++ b/Engine/Source/Runtime/Windows/WindowsPlatformFeatures/Private/WmfRingBuffer.cpp
@@ -8,14 +8,14 @@ DEFINE_LOG_CATEGORY(WmfRingBuffer);
 
 WINDOWSPLATFORMFEATURES_START
 
-void FWmfRingBuffer::Push(FGameplayMediaEncoderSample&& Sample)
+void FWmfRingBuffer::Push(AVEncoder::FAVPacket&& Sample)
 {
 	check(MaxDuration != 0);
 
 	FScopeLock Lock(&Mutex);
 
-	if (Samples.Num() == 0 && Sample.GetType() == EMediaType::Audio)
-		// start always from video sample for simplicity of `Cleanup` impl
+	// start always from video sample for simplicity of `Cleanup` impl
+	if (Samples.Num() == 0 && Sample.Type == AVEncoder::EPacketType::Audio)
 	{
 		// drop
 	}
@@ -27,15 +27,15 @@ void FWmfRingBuffer::Push(FGameplayMediaEncoderSample&& Sample)
 		// video encoding takes longer than audio thus video encoded samples arrive later and need to be placed
 		// inside the container. Video also is not strictly timestamp-ordered (maybe due to B-frames)
 	{
-		FTimespan TimestampToInsert = Sample.GetTime();
+		FTimespan TimestampToInsert = Sample.Timestamp;
 
 		bool bIsSampleVideoKeyFrame = Sample.IsVideoKeyFrame();
 
 		int i = Samples.Num() - 1;
 		while (true)
 		{
-			if (Samples[i].GetType() == Sample.GetType() // don't change order of same stream cos video is not strictly timestamp-ordered (B-Frames)
-				|| TimestampToInsert >= Samples[i].GetTime())
+			if (Samples[i].Type == Sample.Type // don't change order of same stream cos video is not strictly timestamp-ordered (B-Frames)
+				|| TimestampToInsert >= Samples[i].Timestamp)
 			{
 				Samples.Insert(MoveTemp(Sample), i + 1);
 				break;
@@ -43,7 +43,7 @@ void FWmfRingBuffer::Push(FGameplayMediaEncoderSample&& Sample)
 
 			if (i == 0)
 			{
-				if (Sample.GetType() == EMediaType::Video)
+				if (Sample.Type == AVEncoder::EPacketType::Video)
 				{
 					Samples.Insert(MoveTemp(Sample), 0);
 				}
@@ -93,7 +93,12 @@ void FWmfRingBuffer::Cleanup()
 	Samples.RemoveAt(0, i, false);
 	check(Samples[0].IsVideoKeyFrame());
 
-	UE_LOG(WmfRingBuffer, VeryVerbose, TEXT("%d samples, %.3f s, %.3f - %.3f:%.3f"), Samples.Num(), GetDuration().GetTotalSeconds(), Samples[0].GetTime().GetTotalSeconds(), Samples.Last().GetTime().GetTotalSeconds(), Samples.Last().GetDuration().GetTotalSeconds());
+	UE_LOG(WmfRingBuffer, VeryVerbose, TEXT("%d samples, %.3f s, %.3f - %.3f:%.3f"),
+		Samples.Num(),
+		GetDuration().GetTotalSeconds(),
+		Samples[0].Timestamp.GetTotalSeconds(),
+		Samples.Last().Timestamp.GetTotalSeconds(),
+		Samples.Last().Duration.GetTotalSeconds());
 }
 
 void FWmfRingBuffer::PauseCleanup(bool bPause)
@@ -108,23 +113,20 @@ FTimespan FWmfRingBuffer::GetDuration() const
 		return 0;
 	}
 
-	return Samples.Last().GetTime() - Samples[0].GetTime();
+	// A duration of a video is not "LastFrameCaptureTime - FirstFrameCaptureTime" , but 
+	// (LastFrameCaptureTime - FirstFrameCaptureTime) + LastFrameDuration.
+	return Samples.Last().Timestamp + Samples.Last().Duration - Samples[0].Timestamp;
 }
 
-TArray<FGameplayMediaEncoderSample> FWmfRingBuffer::GetCopy()
+TArray<AVEncoder::FAVPacket> FWmfRingBuffer::GetCopy()
 {
 	FScopeLock Lock(&Mutex);
 
-	TArray<FGameplayMediaEncoderSample> Copy;
-	Copy.Reset(Samples.Num());
-
-	for (const auto& Sample : Samples)
-	{
-		// cloning is required because during saving a copy of ring buffer we modify samples timestamps.
-		// using samples references shared between ring buffer and its copy being saved would cause ring
-		// buffer timestamp logic corruption
-		Copy.Add(Sample.Clone());
-	}
+	TArray<AVEncoder::FAVPacket> Copy;
+	// cloning is required because during saving a copy of ring buffer we modify samples timestamps.
+	// using samples references shared between ring buffer and its copy being saved would cause ring
+	// buffer timestamp logic corruption
+	Copy = Samples;
 
 	return Copy;
 }
diff --git a/Engine/Source/Runtime/Windows/WindowsPlatformFeatures/Private/WmfRingBuffer.h b/Engine/Source/Runtime/Windows/WindowsPlatformFeatures/Private/WmfRingBuffer.h
--- a/Engine/Source/Runtime/Windows/WindowsPlatformFeatures/Private/WmfRingBuffer.h
+++ b/Engine/Source/Runtime/Windows/WindowsPlatformFeatures/Private/WmfRingBuffer.h
@@ -5,7 +5,7 @@
 #include "WmfPrivate.h"
 
 #include "HAL/ThreadSafeBool.h"
-#include "GameplayMediaEncoderSample.h"
+#include "AVEncoder.h"
 
 class FWmfRingBuffer
 {
@@ -22,11 +22,11 @@ public:
 
 	FTimespan GetDuration() const;
 
-	void Push(FGameplayMediaEncoderSample&& Sample);
+	void Push(AVEncoder::FAVPacket&& Sample);
 
 	void PauseCleanup(bool bPause);
 
-	TArray<FGameplayMediaEncoderSample> GetCopy();
+	TArray<AVEncoder::FAVPacket> GetCopy();
 
 	void Reset();
 
@@ -35,7 +35,7 @@ private:
 
 private:
 	FTimespan MaxDuration = 0;
-	TArray<FGameplayMediaEncoderSample> Samples;
+	TArray<AVEncoder::FAVPacket> Samples;
 	FCriticalSection Mutex;
 	FThreadSafeBool bCleanupPaused = false;
 };
diff --git a/Engine/Source/Runtime/Windows/WindowsPlatformFeatures/WindowsPlatformFeatures.Build.cs b/Engine/Source/Runtime/Windows/WindowsPlatformFeatures/WindowsPlatformFeatures.Build.cs
--- a/Engine/Source/Runtime/Windows/WindowsPlatformFeatures/WindowsPlatformFeatures.Build.cs
+++ b/Engine/Source/Runtime/Windows/WindowsPlatformFeatures/WindowsPlatformFeatures.Build.cs
@@ -19,12 +19,14 @@ public class WindowsPlatformFeatures : ModuleRules
 				"CoreUObject",
 				"Engine",
 				"AudioMixer",
-				"GameplayMediaEncoder"
+				"AVEncoder",
+				"GameplayMediaEncoder",
 			}
 		);
 
         PublicDelayLoadDLLs.Add("mfplat.dll");
         PublicDelayLoadDLLs.Add("mfuuid.dll");
         PublicDelayLoadDLLs.Add("Mfreadwrite.dll");
+
     }
 }
