XMAInfo safety refactors

diff --git a/Engine/Source/Runtime/NonRealtimeAudioRenderer/Private/AudioMixerPlatformNonRealtime.cpp b/Engine/Source/Runtime/NonRealtimeAudioRenderer/Private/AudioMixerPlatformNonRealtime.cpp
--- a/Engine/Source/Runtime/NonRealtimeAudioRenderer/Private/AudioMixerPlatformNonRealtime.cpp
+++ b/Engine/Source/Runtime/NonRealtimeAudioRenderer/Private/AudioMixerPlatformNonRealtime.cpp
@@ -399,6 +399,10 @@ namespace Audio
 
 	void FMixerPlatformNonRealtime::OnHardwareUpdate()
 	{
+#if WITH_XMA2
+		FXMAAudioInfo::Tick();
+#endif //WITH_XMA2
+
 		if (RenderEveryTickCvar)
 		{
 			RenderAudio(TickDelta);
diff --git a/Engine/Source/Runtime/Windows/AudioMixerXAudio2/Private/AudioMixerPlatformXAudio2.cpp b/Engine/Source/Runtime/Windows/AudioMixerXAudio2/Private/AudioMixerPlatformXAudio2.cpp
--- a/Engine/Source/Runtime/Windows/AudioMixerXAudio2/Private/AudioMixerPlatformXAudio2.cpp
+++ b/Engine/Source/Runtime/Windows/AudioMixerXAudio2/Private/AudioMixerPlatformXAudio2.cpp
@@ -1299,6 +1299,10 @@ namespace Audio
 
 	void FMixerPlatformXAudio2::OnHardwareUpdate()
 	{
+#if WITH_XMA2
+		FXMAAudioInfo::Tick();
+#endif //WITH_XMA2
+
 		if (bIsUsingNullDevice)
 		{
 			float CurrentTime = FPlatformTime::Seconds();
diff --git a/Engine/Source/Runtime/Windows/XAudio2/Private/XAudio2Device.cpp b/Engine/Source/Runtime/Windows/XAudio2/Private/XAudio2Device.cpp
--- a/Engine/Source/Runtime/Windows/XAudio2/Private/XAudio2Device.cpp
+++ b/Engine/Source/Runtime/Windows/XAudio2/Private/XAudio2Device.cpp
@@ -298,6 +298,10 @@ void FXAudio2Device::TeardownHardware()
 
 void FXAudio2Device::UpdateHardware()
 {
+#if WITH_XMA2
+	FXMAAudioInfo::Tick();
+#endif //WITH_XMA2
+
 	// If the audio device changed, we need to tear down and restart the audio engine state
 	if (DeviceProperties && DeviceProperties->DidAudioDeviceChange())
 	{
