CVar on XMAInfo implementation

diff --git a/Engine/Source/Runtime/NonRealtimeAudioRenderer/Private/AudioMixerPlatformNonRealtime.cpp b/Engine/Source/Runtime/NonRealtimeAudioRenderer/Private/AudioMixerPlatformNonRealtime.cpp
--- a/Engine/Source/Runtime/NonRealtimeAudioRenderer/Private/AudioMixerPlatformNonRealtime.cpp
+++ b/Engine/Source/Runtime/NonRealtimeAudioRenderer/Private/AudioMixerPlatformNonRealtime.cpp
@@ -126,7 +126,7 @@ namespace Audio
 
 #if WITH_XMA2
 		//Initialize our XMA2 decoder context
-		FXMAAudioInfo::Initialize();
+		XMA2_INFO_CALL(FXMAAudioInfo::Initialize());
 #endif //#if WITH_XMA2
 
 		// Load ogg and vorbis dlls if they haven't been loaded yet
@@ -147,7 +147,7 @@ namespace Audio
 		}
 
 #if WITH_XMA2
-		FXMAAudioInfo::Shutdown();
+		XMA2_INFO_CALL(FXMAAudioInfo::Shutdown());
 #endif
 
 		bIsInitialized = false;
@@ -355,7 +355,7 @@ namespace Audio
 #if WITH_XMA2 && USE_XMA2_FOR_STREAMING
 		if (InSoundWave->IsStreaming() && InSoundWave->NumChannels <= 2 )
 		{
-			return new FXMAAudioInfo();
+			return XMA2_INFO_NEW();
 		}
 #endif
 
@@ -378,7 +378,7 @@ namespace Audio
 		static const FName NAME_XMA(TEXT("XMA"));
 		if (FPlatformProperties::RequiresCookedData() ? InSoundWave->HasCompressedData(NAME_XMA) : (InSoundWave->GetCompressedData(NAME_XMA) != nullptr))
 		{
-			return new FXMAAudioInfo();
+			return XMA2_INFO_NEW();
 		}
 #endif // WITH_XMA2
 #endif // PLATFORM_WINDOWS || WITH_XMA2
@@ -400,7 +400,7 @@ namespace Audio
 	void FMixerPlatformNonRealtime::OnHardwareUpdate()
 	{
 #if WITH_XMA2
-		FXMAAudioInfo::Tick();
+		XMA2_INFO_CALL(FXMAAudioInfo::Tick());
 #endif //WITH_XMA2
 
 		if (RenderEveryTickCvar)
diff --git a/Engine/Source/Runtime/Windows/AudioMixerXAudio2/Private/AudioMixerPlatformXAudio2.cpp b/Engine/Source/Runtime/Windows/AudioMixerXAudio2/Private/AudioMixerPlatformXAudio2.cpp
--- a/Engine/Source/Runtime/Windows/AudioMixerXAudio2/Private/AudioMixerPlatformXAudio2.cpp
+++ b/Engine/Source/Runtime/Windows/AudioMixerXAudio2/Private/AudioMixerPlatformXAudio2.cpp
@@ -261,7 +261,7 @@ namespace Audio
 
 #if WITH_XMA2
 		//Initialize our XMA2 decoder context
-		FXMAAudioInfo::Initialize();
+		XMA2_INFO_CALL(FXMAAudioInfo::Initialize());
 #endif //#if WITH_XMA2
 #if WITH_ENGINE
 		// Load ogg and vorbis dlls if they haven't been loaded yet
@@ -284,7 +284,7 @@ namespace Audio
 		SAFE_RELEASE(XAudio2System);
 
 #if WITH_XMA2
-		FXMAAudioInfo::Shutdown();
+		XMA2_INFO_CALL(FXMAAudioInfo::Shutdown());
 #endif
 
 #if PLATFORM_WINDOWS
@@ -1252,7 +1252,7 @@ namespace Audio
 #if WITH_XMA2 && USE_XMA2_FOR_STREAMING
 		if (InSoundWave->IsStreaming() && InSoundWave->NumChannels <= 2 )
 		{
-			return new FXMAAudioInfo();
+			return XMA2_INFO_NEW();
 		}
 #endif
 
@@ -1275,7 +1275,7 @@ namespace Audio
 		static const FName NAME_XMA(TEXT("XMA"));
 		if (FPlatformProperties::RequiresCookedData() ? InSoundWave->HasCompressedData(NAME_XMA) : (InSoundWave->GetCompressedData(NAME_XMA) != nullptr))
 		{
-			return new FXMAAudioInfo();
+			return XMA2_INFO_NEW();
 		}
 #endif // WITH_XMA2
 #endif // WITH_ENGINE
@@ -1300,7 +1300,7 @@ namespace Audio
 	void FMixerPlatformXAudio2::OnHardwareUpdate()
 	{
 #if WITH_XMA2
-		FXMAAudioInfo::Tick();
+		XMA2_INFO_CALL(FXMAAudioInfo::Tick());
 #endif //WITH_XMA2
 
 		if (bIsUsingNullDevice)
diff --git a/Engine/Source/Runtime/Windows/XAudio2/Private/XAudio2Device.cpp b/Engine/Source/Runtime/Windows/XAudio2/Private/XAudio2Device.cpp
--- a/Engine/Source/Runtime/Windows/XAudio2/Private/XAudio2Device.cpp
+++ b/Engine/Source/Runtime/Windows/XAudio2/Private/XAudio2Device.cpp
@@ -269,7 +269,7 @@ bool FXAudio2Device::InitializeHardware()
 	bIsAudioDeviceHardwareInitialized = true;
 
 #if WITH_XMA2
-	FXMAAudioInfo::Initialize();
+	XMA2_INFO_CALL(FXMAAudioInfo::Initialize());
 #endif
 
 	return true;
@@ -284,7 +284,7 @@ void FXAudio2Device::TeardownHardware()
 	}
 
 #if WITH_XMA2
-	FXMAAudioInfo::Shutdown();
+	XMA2_INFO_CALL(FXMAAudioInfo::Shutdown());
 #endif
 
 #if PLATFORM_WINDOWS
@@ -299,7 +299,7 @@ void FXAudio2Device::TeardownHardware()
 void FXAudio2Device::UpdateHardware()
 {
 #if WITH_XMA2
-	FXMAAudioInfo::Tick();
+		XMA2_INFO_CALL(FXMAAudioInfo::Tick());
 #endif //WITH_XMA2
 
 	// If the audio device changed, we need to tear down and restart the audio engine state
@@ -370,7 +370,8 @@ class ICompressedAudioInfo* FXAudio2Device::CreateCompressedAudioInfo(USoundWave
 #if WITH_XMA2 && USE_XMA2_FOR_STREAMING
 		if (SoundWave->NumChannels <= 2 )
 		{
-			ICompressedAudioInfo* CompressedInfo = new FXMAAudioInfo();
+			ICompressedAudioInfo* CompressedInfo = XMA2_INFO_NEW();
+		
 			if (!CompressedInfo)
 			{
 				UE_LOG(LogAudio, Error, TEXT("Failed to create new FXMAAudioInfo for streaming SoundWave %s: out of memory."), *SoundWave->GetName());
@@ -405,7 +406,7 @@ class ICompressedAudioInfo* FXAudio2Device::CreateCompressedAudioInfo(USoundWave
 	static const FName NAME_XMA(TEXT("XMA"));
 	if (FPlatformProperties::RequiresCookedData() ? SoundWave->HasCompressedData(NAME_XMA) : (SoundWave->GetCompressedData(NAME_XMA) != nullptr))
 	{
-		ICompressedAudioInfo* CompressedInfo = new FXMAAudioInfo();
+		ICompressedAudioInfo* CompressedInfo = XMA2_INFO_NEW();			
 		if (!CompressedInfo)
 		{
 			UE_LOG(LogAudio, Error, TEXT("Failed to create new FXMAAudioInfo for SoundWave %s: out of memory."), *SoundWave->GetName());
