Added null check for SwapChain in FD3D11Viewport destructor. This handles the case where the engine is launched with -RenderOffscreen and the swap chain is never initialized

#rnx

diff --git a/Engine/Source/Runtime/Windows/D3D11RHI/Private/D3D11Viewport.cpp b/Engine/Source/Runtime/Windows/D3D11RHI/Private/D3D11Viewport.cpp
--- a/Engine/Source/Runtime/Windows/D3D11RHI/Private/D3D11Viewport.cpp
+++ b/Engine/Source/Runtime/Windows/D3D11RHI/Private/D3D11Viewport.cpp
@@ -196,7 +196,10 @@ FD3D11Viewport::~FD3D11Viewport()
 
 	// If the swap chain was in fullscreen mode, switch back to windowed before releasing the swap chain.
 	// DXGI throws an error otherwise.
-	VERIFYD3D11RESULT_EX(SwapChain->SetFullscreenState(false,NULL), D3DRHI->GetDevice());
+	if (SwapChain)
+	{
+		VERIFYD3D11RESULT_EX(SwapChain->SetFullscreenState(false, NULL), D3DRHI->GetDevice());
+	}
 
 	FrameSyncEvent.ReleaseResource();
 
