Debug tracking for an elusive bug.

diff --git a/Engine/Source/Runtime/Windows/D3D11RHI/Private/D3D11VertexBuffer.cpp b/Engine/Source/Runtime/Windows/D3D11RHI/Private/D3D11VertexBuffer.cpp
--- a/Engine/Source/Runtime/Windows/D3D11RHI/Private/D3D11VertexBuffer.cpp
+++ b/Engine/Source/Runtime/Windows/D3D11RHI/Private/D3D11VertexBuffer.cpp
@@ -81,7 +81,9 @@ FVertexBufferRHIRef FD3D11DynamicRHI::RHICreateVertexBuffer(uint32 Size,uint32 I
 	}
 
 	TRefCountPtr<ID3D11Buffer> VertexBufferResource;
+	check(Desc.ByteWidth != 0); //tracking down an elusive bug.
 	HRESULT Result = Direct3DDevice->CreateBuffer(&Desc, pInitData, VertexBufferResource.GetInitReference());
+	check(Desc.ByteWidth != 0); //tracking down an elusive bug.
 	if (FAILED(Result))
 	{
 		UE_LOG(LogD3D11RHI, Error, TEXT("D3DDevice failed CreateBuffer VB with ByteWidth=%d, BindFlags=0x%x Usage=%d, CPUAccess=0x%x, MiscFlags=0x%x"), Desc.ByteWidth, (uint32)Desc.BindFlags, (uint32)Desc.Usage, Desc.CPUAccessFlags, Desc.MiscFlags);
