Remove shadow variable that looks to be able to cause xaudio2 to go silent if in super perfy state

diff --git a/Engine/Source/Runtime/Windows/AudioMixerXAudio2/Private/AudioMixerPlatformXAudio2.h b/Engine/Source/Runtime/Windows/AudioMixerXAudio2/Private/AudioMixerPlatformXAudio2.h
--- a/Engine/Source/Runtime/Windows/AudioMixerXAudio2/Private/AudioMixerPlatformXAudio2.h
+++ b/Engine/Source/Runtime/Windows/AudioMixerXAudio2/Private/AudioMixerPlatformXAudio2.h
@@ -112,7 +112,6 @@ namespace Audio
 		FCriticalSection AudioDeviceSwapCriticalSection;
 		FString OriginalAudioDeviceId;
 		FString NewAudioDeviceId;
-		FThreadSafeBool bMoveAudioStreamToNewAudioDevice;
 		double LastDeviceSwapTime;
 
 		// When we are running the null device,
