Fallback fix for Xaudio2_9 failing in some rare circumstances on Windows 8 machines.

diff --git a/Engine/Source/Runtime/Windows/AudioMixerXAudio2/Private/AudioMixerPlatformXAudio2.cpp b/Engine/Source/Runtime/Windows/AudioMixerXAudio2/Private/AudioMixerPlatformXAudio2.cpp
--- a/Engine/Source/Runtime/Windows/AudioMixerXAudio2/Private/AudioMixerPlatformXAudio2.cpp
+++ b/Engine/Source/Runtime/Windows/AudioMixerXAudio2/Private/AudioMixerPlatformXAudio2.cpp
@@ -86,6 +86,23 @@ static FString GetErrorString(HRESULT Result)
 		}
 	}
 }
+#if PLATFORM_WINDOWS || PLATFORM_HOLOLENS
+FName GetNextDllToTry(FName Current = NAME_None) 
+{
+	static const TArray<FName> kDllsToTry = { TEXT("xaudio2_9redist.dll"), TEXT("XAudio2_7.dll") };
+	int32 Index = kDllsToTry.IndexOfByKey(Current);
+	
+	// This behaves like a link list of names. 
+	// Current=Name_NONE will have an index of -1, thus yield Index 0	
+	// until we reach end of the list and return Name_NONE.
+
+	if (kDllsToTry.IsValidIndex(Index + 1))			
+	{
+		return kDllsToTry[Index + 1];
+	}
+	return NAME_None;
+}
+#endif //#if PLATFORM_WINDOWS || PLATFORM_HOLOLENS
 
 
 /*
@@ -147,6 +164,7 @@ namespace Audio
 	{
 #if PLATFORM_WINDOWS || PLATFORM_HOLOLENS
 		FPlatformMisc::CoInitialize();
+		DllName = GetNextDllToTry();
 #endif // #if PLATFORM_WINDOWS || PLATFORM_HOLOLENS
 	}
 
@@ -214,8 +232,8 @@ namespace Audio
 		// Load the xaudio2 library and keep a handle so we can free it on teardown
 		// Note: windows internally ref-counts the library per call to load library so 
 		// when we call FreeLibrary, it will only free it once the refcount is zero
-		XAudio2Dll = (HMODULE)FPlatformProcess::GetDllHandle(TEXT("XAudio2_9redist.dll"));
-
+		XAudio2Dll = (HMODULE)FPlatformProcess::GetDllHandle(*DllName.GetPlainNameString());
+			
 		// returning null means we failed to load XAudio2, which means everything will fail
 		if (XAudio2Dll == nullptr)
 		{
@@ -895,8 +913,26 @@ namespace Audio
 
 	Cleanup:
 		if (FAILED(Result))
-		{
+		{		
+			bool bFailedToCreateMasterVoice = OutputAudioStreamMasteringVoice == nullptr;
 			CloseAudioStream();
+
+#if PLATFORM_WINDOWS
+			if (bFailedToCreateMasterVoice)
+			{
+				// Try another DLL.
+				DllName = GetNextDllToTry(DllName);
+				if (DllName != NAME_None)
+				{
+					TeardownHardware();
+					if (InitializeHardware())
+					{
+						return OpenAudioStream(Params);
+					}
+				}
+			}
+#endif //PLATFORM_WINDOWS
+
 		}
 		return SUCCEEDED(Result);
 	}
diff --git a/Engine/Source/Runtime/Windows/AudioMixerXAudio2/Private/AudioMixerPlatformXAudio2.h b/Engine/Source/Runtime/Windows/AudioMixerXAudio2/Private/AudioMixerPlatformXAudio2.h
--- a/Engine/Source/Runtime/Windows/AudioMixerXAudio2/Private/AudioMixerPlatformXAudio2.h
+++ b/Engine/Source/Runtime/Windows/AudioMixerXAudio2/Private/AudioMixerPlatformXAudio2.h
@@ -99,6 +99,7 @@ namespace Audio
 		bool ResetXAudio2System();
 
 		// Handle to XAudio2DLL
+		FName DllName;
 		HMODULE XAudio2Dll;
 
 		// Bool indicating that the default audio device changed
