Fix deprecation warnings with new texture aliasing APIs.

diff --git a/Engine/Source/Runtime/D3D12RHI/Private/D3D12RHIPrivate.h b/Engine/Source/Runtime/D3D12RHI/Private/D3D12RHIPrivate.h
--- a/Engine/Source/Runtime/D3D12RHI/Private/D3D12RHIPrivate.h
+++ b/Engine/Source/Runtime/D3D12RHI/Private/D3D12RHIPrivate.h
@@ -186,6 +186,8 @@ struct FD3D12UpdateTexture3DData
 /** Forward declare the context for the AMD AGS utility library. */
 struct AGSContext;
 
+PRAGMA_DISABLE_DEPRECATION_WARNINGS
+
 /** The interface which is implemented by the dynamically bound RHI. */
 class FD3D12DynamicRHI : public FDynamicRHI
 {
@@ -341,7 +343,9 @@ public:
 	virtual bool RHIGetRenderQueryResult(FRHIRenderQuery* RenderQuery, uint64& OutResult, bool bWait) final override;
 	virtual uint32 RHIGetViewportNextPresentGPUIndex(FRHIViewport* Viewport) final override;
 	virtual FTexture2DRHIRef RHIGetViewportBackBuffer(FRHIViewport* Viewport) final override;
+	UE_DEPRECATED(4.25, "RHIAliasTextureResources now takes references to FTextureRHIRef objects as parameters")
 	virtual void RHIAliasTextureResources(FRHITexture* DestTexture, FRHITexture* SrcTexture) final override;
+	UE_DEPRECATED(4.25, "RHICreateAliasedTexture now takes a reference to an FTextureRHIRef object")
 	virtual FTextureRHIRef RHICreateAliasedTexture(FRHITexture* SourceTexture) final override;
 	virtual void RHIAliasTextureResources(FTextureRHIRef& DestTexture, FTextureRHIRef& SrcTexture) final override;
 	virtual FTextureRHIRef RHICreateAliasedTexture(FTextureRHIRef& SourceTexture) final override;
@@ -967,6 +971,8 @@ protected:
 	HANDLE FlipEvent;
 };
 
+PRAGMA_ENABLE_DEPRECATION_WARNINGS
+
 /** Implements the D3D12RHI module as a dynamic RHI providing module. */
 class FD3D12DynamicRHIModule : public IDynamicRHIModule
 {
diff --git a/Engine/Source/Runtime/OpenGLDrv/Private/OpenGLTexture.cpp b/Engine/Source/Runtime/OpenGLDrv/Private/OpenGLTexture.cpp
--- a/Engine/Source/Runtime/OpenGLDrv/Private/OpenGLTexture.cpp
+++ b/Engine/Source/Runtime/OpenGLDrv/Private/OpenGLTexture.cpp
@@ -2800,6 +2800,8 @@ FTextureCubeRHIRef FOpenGLDynamicRHI::RHICreateTextureCubeFromResource(EPixelFor
 	return TextureCube;
 }
 
+PRAGMA_DISABLE_DEPRECATION_WARNINGS
+
 void FOpenGLDynamicRHI::RHIAliasTextureResources(FRHITexture* DestRHITexture, FRHITexture* SrcRHITexture)
 {
 	FOpenGLTextureBase* DestTexture = GetOpenGLTextureFromRHITexture(DestRHITexture);
@@ -2880,6 +2882,8 @@ FTextureRHIRef FOpenGLDynamicRHI::RHICreateAliasedTexture(FRHITexture* SourceTex
 	return AliasedTexture;
 }
 
+PRAGMA_ENABLE_DEPRECATION_WARNINGS
+
 void* FOpenGLDynamicRHI::LockTexture2D_RenderThread(class FRHICommandListImmediate& RHICmdList, FRHITexture2D* Texture, uint32 MipIndex, EResourceLockMode LockMode, uint32& DestStride, bool bLockWithinMiptail, bool bNeedsDefaultRHIFlush)
 {
 	check(IsInRenderingThread());
diff --git a/Engine/Source/Runtime/OpenGLDrv/Public/OpenGLDrv.h b/Engine/Source/Runtime/OpenGLDrv/Public/OpenGLDrv.h
--- a/Engine/Source/Runtime/OpenGLDrv/Public/OpenGLDrv.h
+++ b/Engine/Source/Runtime/OpenGLDrv/Public/OpenGLDrv.h
@@ -291,6 +291,7 @@ struct FOpenGLGPUProfiler : public FGPUProfiler
 	void EndFrame();
 };
 
+PRAGMA_DISABLE_DEPRECATION_WARNINGS
 
 /** The interface which is implemented by the dynamically bound RHI. */
 class OPENGLDRV_API FOpenGLDynamicRHI  final : public FDynamicRHI, public IRHICommandContextPSOFallback
@@ -416,7 +417,9 @@ public:
 	virtual FRenderQueryRHIRef RHICreateRenderQuery(ERenderQueryType QueryType) final override;
 	virtual bool RHIGetRenderQueryResult(FRHIRenderQuery* RenderQuery, uint64& OutResult, bool bWait) final override;
 	virtual FTexture2DRHIRef RHIGetViewportBackBuffer(FRHIViewport* Viewport) final override;
+	UE_DEPRECATED(4.25, "RHIAliasTextureResources now takes references to FTextureRHIRef objects as parameters")
 	virtual void RHIAliasTextureResources(FRHITexture* DestTexture, FRHITexture* SrcTexture) final override;
+	UE_DEPRECATED(4.25, "RHICreateAliasedTexture now takes a reference to an FTextureRHIRef object")
 	virtual FTextureRHIRef RHICreateAliasedTexture(FRHITexture* SourceTexture) final override;
 	virtual void RHIAdvanceFrameForGetViewportBackBuffer(FRHIViewport* Viewport) final override;
 	virtual void RHIAcquireThreadOwnership() final override;
@@ -1201,6 +1204,8 @@ private:
 
 };
 
+PRAGMA_ENABLE_DEPRECATION_WARNINGS
+
 /** Implements the OpenGLDrv module as a dynamic RHI providing module. */
 class FOpenGLDynamicRHIModule : public IDynamicRHIModule
 {
diff --git a/Engine/Source/Runtime/RHI/Public/RHIValidation.h b/Engine/Source/Runtime/RHI/Public/RHIValidation.h
--- a/Engine/Source/Runtime/RHI/Public/RHIValidation.h
+++ b/Engine/Source/Runtime/RHI/Public/RHIValidation.h
@@ -984,21 +984,6 @@ public:
 		return RHI->RHICreateUnorderedAccessViewStencil(DepthTarget, MipLevel);
 	}
 
-	virtual void RHIAliasTextureResources(FRHITexture* DestTexture, FRHITexture* SourceTexture) override final
-	{
-		// Source and target need to be valid objects.
-		check(DestTexture && SourceTexture);
-		// Source texture must have been created (i.e. have a native resource backing).
-		check(SourceTexture->GetNativeResource() != nullptr);
-		RHI->RHIAliasTextureResources(DestTexture, SourceTexture);
-	}
-
-	virtual FTextureRHIRef RHICreateAliasedTexture(FRHITexture* SourceTexture) override final
-	{
-		check(SourceTexture);
-		return RHI->RHICreateAliasedTexture(SourceTexture);
-	}
-
 	virtual void RHIAliasTextureResources(FTextureRHIRef& DestTexture, FTextureRHIRef& SourceTexture) override final
 	{
 		// Source and target need to be valid objects.
diff --git a/Engine/Source/Runtime/VulkanRHI/Public/VulkanDynamicRHI.h b/Engine/Source/Runtime/VulkanRHI/Public/VulkanDynamicRHI.h
--- a/Engine/Source/Runtime/VulkanRHI/Public/VulkanDynamicRHI.h
+++ b/Engine/Source/Runtime/VulkanRHI/Public/VulkanDynamicRHI.h
@@ -36,7 +36,7 @@ struct FOptionalVulkanInstanceExtensions
 	void Setup(const TArray<const ANSICHAR*>& InInstanceExtensions);
 };
 
-
+PRAGMA_DISABLE_DEPRECATION_WARNINGS
 
 /** The interface which is implemented by the dynamically bound RHI. */
 class FVulkanDynamicRHI : public FDynamicRHI
@@ -156,7 +156,9 @@ public:
 	virtual FRenderQueryRHIRef RHICreateRenderQuery(ERenderQueryType QueryType) final override;
 	virtual bool RHIGetRenderQueryResult(FRHIRenderQuery* RenderQuery, uint64& OutResult, bool bWait) final override;
 	virtual FTexture2DRHIRef RHIGetViewportBackBuffer(FRHIViewport* Viewport) final override;
+	UE_DEPRECATED(4.25, "RHIAliasTextureResources now takes references to FTextureRHIRef objects as parameters")
 	virtual void RHIAliasTextureResources(FRHITexture* DestTexture, FRHITexture* SrcTexture) final override;
+	UE_DEPRECATED(4.25, "RHICreateAliasedTexture now takes a reference to an FTextureRHIRef object")
 	virtual FTextureRHIRef RHICreateAliasedTexture(FRHITexture* SourceTexture) final override;
 	virtual void RHIAliasTextureResources(FTextureRHIRef& DestTexture, FTextureRHIRef& SrcTexture) final override;
 	virtual FTextureRHIRef RHICreateAliasedTexture(FTextureRHIRef& SourceTexture) final override;
@@ -487,6 +489,8 @@ public:
 	static TSharedPtr< IHeadMountedDisplayVulkanExtensions, ESPMode::ThreadSafe > HMDVulkanExtensions;
 };
 
+PRAGMA_ENABLE_DEPRECATION_WARNINGS
+
 /** Implements the Vulkan module as a dynamic RHI providing module. */
 class FVulkanDynamicRHIModule : public IDynamicRHIModule
 {
diff --git a/Engine/Source/Runtime/Windows/D3D11RHI/Private/D3D11RHIPrivate.h b/Engine/Source/Runtime/Windows/D3D11RHI/Private/D3D11RHIPrivate.h
--- a/Engine/Source/Runtime/Windows/D3D11RHI/Private/D3D11RHIPrivate.h
+++ b/Engine/Source/Runtime/Windows/D3D11RHI/Private/D3D11RHIPrivate.h
@@ -341,6 +341,8 @@ private:
 /** Forward declare the context for the AMD AGS utility library. */
 struct AGSContext;
 
+PRAGMA_DISABLE_DEPRECATION_WARNINGS
+
 /** The interface which is implemented by the dynamically bound RHI. */
 class D3D11RHI_API FD3D11DynamicRHI : public FDynamicRHI, public IRHICommandContextPSOFallback
 {
@@ -476,7 +478,9 @@ public:
 	virtual FRenderQueryRHIRef RHICreateRenderQuery(ERenderQueryType QueryType) final override;
 	virtual bool RHIGetRenderQueryResult(FRHIRenderQuery* RenderQuery, uint64& OutResult, bool bWait) final override;
 	virtual FTexture2DRHIRef RHIGetViewportBackBuffer(FRHIViewport* Viewport) final override;
+	UE_DEPRECATED(4.25, "RHIAliasTextureResources now takes references to FTextureRHIRef objects as parameters")
 	virtual void RHIAliasTextureResources(FRHITexture* DestTexture, FRHITexture* SrcTexture) final override;
+	UE_DEPRECATED(4.25, "RHICreateAliasedTexture now takes a reference to an FTextureRHIRef object")
 	virtual FTextureRHIRef RHICreateAliasedTexture(FRHITexture* SourceTexture) final override;
 	virtual void RHIAliasTextureResources(FTextureRHIRef& DestTexture, FTextureRHIRef& SrcTexture) final override;
 	virtual FTextureRHIRef RHICreateAliasedTexture(FTextureRHIRef& SourceTexture) final override;
@@ -1047,6 +1051,8 @@ protected:
 
 };
 
+PRAGMA_ENABLE_DEPRECATION_WARNINGS
+
 struct FD3D11Adapter
 {
 	/** -1 if not supported or FindAdpater() wasn't called. Ideally we would store a pointer to IDXGIAdapter but it's unlikely the adpaters change during engine init. */
diff --git a/Engine/Source/Runtime/Windows/D3D11RHI/Private/D3D11Texture.cpp b/Engine/Source/Runtime/Windows/D3D11RHI/Private/D3D11Texture.cpp
--- a/Engine/Source/Runtime/Windows/D3D11RHI/Private/D3D11Texture.cpp
+++ b/Engine/Source/Runtime/Windows/D3D11RHI/Private/D3D11Texture.cpp
@@ -2437,7 +2437,12 @@ TD3D11Texture2D<BaseResourceType>* FD3D11DynamicRHI::CreateAliasedD3D11Texture2D
 
 	// We'll be the same size, since we're the same thing. Avoid the check in D3D11Resources.h (AliasResources).
 	Texture2D->SetMemorySize(SourceTexture->GetMemorySize());
+
+	// Disable deprecation warning; when the DynamicRHI raw-pointer method is fully deprecated, the D3D11 class will still provide a raw pointer version
+	// since this is required in this path.
+	PRAGMA_DISABLE_DEPRECATION_WARNINGS
 	RHIAliasTextureResources(Texture2D, SourceTexture);
+	PRAGMA_ENABLE_DEPRECATION_WARNINGS
 
 	return Texture2D;
 }
@@ -2483,6 +2488,8 @@ FTextureRHIRef FD3D11DynamicRHI::RHICreateAliasedTexture(FRHITexture* SourceText
 	return nullptr;
 }
 
+PRAGMA_DISABLE_DEPRECATION_WARNINGS
+
 void FD3D11DynamicRHI::RHIAliasTextureResources(FTextureRHIRef& DestTextureRHI, FTextureRHIRef& SrcTextureRHI)
 {
 	// @todo: Move the raw-pointer implementation down here when it's deprecation is completed.
@@ -2495,6 +2502,8 @@ FTextureRHIRef FD3D11DynamicRHI::RHICreateAliasedTexture(FTextureRHIRef& SourceT
 	return RHICreateAliasedTexture((FRHITexture*)SourceTexture);
 }
 
+PRAGMA_ENABLE_DEPRECATION_WARNINGS
+
 void FD3D11DynamicRHI::RHICopyTexture(FRHITexture* SourceTextureRHI, FRHITexture* DestTextureRHI, const FRHICopyTextureInfo& CopyInfo)
 {
 	if (!SourceTextureRHI || !DestTextureRHI || SourceTextureRHI == DestTextureRHI)
