Fixed: RHIFlushComputeShaderCache() and the equivalent RHIEndUAVOverlap(); RHIBeginUAVOverlap(); did not work.

The driver extension functions just set an internal flag which is processed on the next dispatch, so doing an end/begin is just a no-op. In order to correctly break up an overlap group, RHIBeginUAVOverlap() must remember that overlap was requested, and actually call the driver extension function after the next draw or dispatch.

Also removed the attempted UAV "barrier" from FD3D11DynamicRHI::RHITransitionResources(). This didn't do anything because of the above bug, but if it did, it would have prevented overlap from occurring in many valid cases, because we don't track which UAVs are actually written by the previous dispatches in a group. The D3D11 RHI must not try to interpret transitions for breaking up overlap groups; it's up to the caller to ensure that the begin/end calls wrap dispatches which are safe to overlap, and must also use transitions for the RHIs which support barriers, where begin/end overlap has no effect.

#rnx

diff --git a/Engine/Source/Runtime/Windows/D3D11RHI/Private/D3D11Commands.cpp b/Engine/Source/Runtime/Windows/D3D11RHI/Private/D3D11Commands.cpp
--- a/Engine/Source/Runtime/Windows/D3D11RHI/Private/D3D11Commands.cpp
+++ b/Engine/Source/Runtime/Windows/D3D11RHI/Private/D3D11Commands.cpp
@@ -221,6 +221,7 @@ void FD3D11DynamicRHI::RHIDispatchComputeShader(uint32 ThreadGroupCountX, uint32
 	
 	Direct3DDeviceIMContext->Dispatch(ThreadGroupCountX, ThreadGroupCountY, ThreadGroupCountZ);	
 	StateCache.SetComputeShader(nullptr);
+	ApplyUAVOverlapState();
 }
 
 void FD3D11DynamicRHI::RHIDispatchIndirectComputeShader(FRHIVertexBuffer* ArgumentBufferRHI, uint32 ArgumentOffset)
@@ -241,6 +242,7 @@ void FD3D11DynamicRHI::RHIDispatchIndirectComputeShader(FRHIVertexBuffer* Argume
 
 	Direct3DDeviceIMContext->DispatchIndirect(ArgumentBuffer->Resource,ArgumentOffset);
 	StateCache.SetComputeShader(nullptr);
+	ApplyUAVOverlapState();
 }
 
 void FD3D11DynamicRHI::RHISetViewport(float MinX, float MinY, float MinZ, float MaxX, float MaxY, float MaxZ)
@@ -1678,6 +1680,8 @@ void FD3D11DynamicRHI::RHIDrawPrimitive(uint32 BaseVertexIndex,uint32 NumPrimiti
 	{
 		Direct3DDeviceIMContext->Draw(VertexCount,BaseVertexIndex);
 	}
+
+	ApplyUAVOverlapState();
 }
 
 void FD3D11DynamicRHI::RHIDrawPrimitiveIndirect(FRHIVertexBuffer* ArgumentBufferRHI,uint32 ArgumentOffset)
@@ -1693,6 +1697,8 @@ void FD3D11DynamicRHI::RHIDrawPrimitiveIndirect(FRHIVertexBuffer* ArgumentBuffer
 
 	StateCache.SetPrimitiveTopology(GetD3D11PrimitiveType(PrimitiveType,bUsingTessellation));
 	Direct3DDeviceIMContext->DrawInstancedIndirect(ArgumentBuffer->Resource,ArgumentOffset);
+
+	ApplyUAVOverlapState();
 }
 
 void FD3D11DynamicRHI::RHIDrawIndexedIndirect(FRHIIndexBuffer* IndexBufferRHI, FRHIStructuredBuffer* ArgumentsBufferRHI, int32 DrawArgumentsIndex, uint32 NumInstances)
@@ -1723,6 +1729,8 @@ void FD3D11DynamicRHI::RHIDrawIndexedIndirect(FRHIIndexBuffer* IndexBufferRHI, F
 	{
 		check(0);
 	}
+
+	ApplyUAVOverlapState();
 }
 
 void FD3D11DynamicRHI::RHIDrawIndexedPrimitive(FRHIIndexBuffer* IndexBufferRHI,int32 BaseVertexIndex,uint32 FirstInstance,uint32 NumVertices,uint32 StartIndex,uint32 NumPrimitives,uint32 NumInstances)
@@ -1764,6 +1772,8 @@ void FD3D11DynamicRHI::RHIDrawIndexedPrimitive(FRHIIndexBuffer* IndexBufferRHI,i
 	{
 		Direct3DDeviceIMContext->DrawIndexed(IndexCount,StartIndex,BaseVertexIndex);
 	}
+
+	ApplyUAVOverlapState();
 }
 
 void FD3D11DynamicRHI::RHIDrawIndexedPrimitiveIndirect(FRHIIndexBuffer* IndexBufferRHI, FRHIVertexBuffer* ArgumentBufferRHI,uint32 ArgumentOffset)
@@ -1785,6 +1795,8 @@ void FD3D11DynamicRHI::RHIDrawIndexedPrimitiveIndirect(FRHIIndexBuffer* IndexBuf
 	StateCache.SetIndexBuffer(IndexBuffer->Resource, Format, 0);
 	StateCache.SetPrimitiveTopology(GetD3D11PrimitiveType(PrimitiveType,bUsingTessellation));
 	Direct3DDeviceIMContext->DrawIndexedInstancedIndirect(ArgumentBuffer->Resource,ArgumentOffset);
+
+	ApplyUAVOverlapState();
 }
 
 // Raster operations.
@@ -2018,18 +2030,6 @@ void FD3D11DynamicRHI::RHITransitionResources(EResourceTransitionAccess Transiti
 			Resource->SetCurrentGPUAccess(TransitionType);
 		}
 	}
-
-	// If we get any write transitions and UAV overlaps are enabled,
-	// we have to toggle the overlap off/back on again to get the driver
-	// to insert an appropriate UAV barrier.
-	switch (TransitionType)
-	{
-	case EResourceTransitionAccess::EWritable:
-	case EResourceTransitionAccess::ERWBarrier:
-	case EResourceTransitionAccess::ERWSubResBarrier:
-		UAVBarrier();
-		break;
-	}
 }
 
 void FD3D11DynamicRHI::RHITransitionResources(EResourceTransitionAccess TransitionType, EResourceTransitionPipeline TransitionPipeline, FRHIUnorderedAccessView** InUAVs, int32 InNumUAVs, FRHIComputeFence* WriteFence)
@@ -2050,18 +2050,6 @@ void FD3D11DynamicRHI::RHITransitionResources(EResourceTransitionAccess Transiti
 		}
 	}
 
-	// If we get any write transitions and UAV overlaps are enabled,
-	// we have to toggle the overlap off/back on again to get the driver
-	// to insert an appropriate UAV barrier.
-	switch (TransitionType)
-	{
-	case EResourceTransitionAccess::EWritable:
-	case EResourceTransitionAccess::ERWBarrier:
-	case EResourceTransitionAccess::ERWSubResBarrier:
-		UAVBarrier();
-		break;
-	}
-
 	if (WriteFence)
 	{
 		WriteFence->WriteFence();
@@ -2094,17 +2082,14 @@ bool FD3D11DynamicRHI::IsUAVOverlapSupported()
 	return IsRHIDeviceNVIDIA() || IsRHIDeviceAMD() || IsRHIDeviceIntel();
 }
 
-void FD3D11DynamicRHI::RHIBeginUAVOverlap()
+void FD3D11DynamicRHI::ApplyUAVOverlapState()
 {
-	checkf(!bUAVOverlapEnabled, TEXT("Mismatched call to BeginUAVOverlap. Ensure all calls to RHICmdList.BeginUAVOverlap() are paired with a call to RHICmdList.EndUAVOverlap()."));
-
-	bUAVOverlapAllowed = (CVarAllowUAVFlushExt.GetValueOnRenderThread() != 0) && IsUAVOverlapSupported();
-	if (!bUAVOverlapAllowed)
+	if (UAVOverlapState != EUAVOverlapState::EPending)
 	{
 		return;
 	}
 
-	bUAVOverlapEnabled = true;
+	UAVOverlapState = EUAVOverlapState::EOn;
 
 #if !PLATFORM_HOLOLENS
 	if (IsRHIDeviceNVIDIA())
@@ -2131,63 +2116,103 @@ void FD3D11DynamicRHI::RHIBeginUAVOverlap()
 #endif
 }
 
-void FD3D11DynamicRHI::RHIEndUAVOverlap()
+void FD3D11DynamicRHI::RHIBeginUAVOverlap()
 {
+	checkf(UAVOverlapState == EUAVOverlapState::EOff, TEXT("Mismatched call to BeginUAVOverlap. Ensure all calls to RHICmdList.BeginUAVOverlap() are paired with a call to RHICmdList.EndUAVOverlap()."));
+
+	bUAVOverlapAllowed = (CVarAllowUAVFlushExt.GetValueOnRenderThread() != 0) && IsUAVOverlapSupported();
 	if (!bUAVOverlapAllowed)
 	{
 		return;
 	}
 
-	checkf(bUAVOverlapEnabled, TEXT("Mismatched call to EndUAVOverlap. Ensure all calls to RHICmdList.BeginUAVOverlap() are paired with a call to RHICmdList.EndUAVOverlap()."));
-	bUAVOverlapEnabled = false;
+	// The driver APIs just set an internal flag which is used by the next dispatch in order to determine if there needs to be a barrier before
+	// running the CS. This means that we need to call the API *after* the next dispatch, because we always want a barrier before the first
+	// dispatch in an overlap group. Consider this example:
+	//
+	//		// Dispatches 1 and 2 are independent, we want them to overlap.
+	//		RHICmdList.BeginUAVOverlap();
+	//		Dispatch1();
+	//		Dispatch2();
+	//		RHICmdList.EndUAVOverlap();
+	//
+	//		// Dispatches 3 and 4 have no inter-dependencies, so they should overlap, but read UAV locations written by 1 and/or 2, so we want a barrier here.
+	//		RHICmdList.BeginUAVOverlap();
+	//		Dispatch3();
+	//		Dispatch4();
+	//		RHICmdList.EndUAVOverlap();
+	//
+	// If we just call the driver extension API immediately here, it will simply overwrite the flag set by the previous end call, and all 4 dispatches will
+	// (potentially) overlap, as if the inner end/begin pair didn't exist, producing incorrect results. Instead, we must set the state to pending here,
+	// and the next RHI draw or dispatch function will call the API after running the draw/dispatch, so the example above results in this actual set of driver calls:
+	//
+	//		Dispatch1();				// Internal barrier before, makes sure we don't overlap any earlier dispatches.
+	//		Vendor_BeginUAVOverlap();	// Set overlap flag to true.
+	//		Dispatch2();				// No barrier, overlaps 1.
+	//		Vendor_EndUAVOverlap();		// Set overlap flag to false.
+	//
+	//		Dispatch3();				// Barrier before, makes sure we don't overlap 2.
+	//		Vendor_BeginUAVOverlap();	// Set overlap flag to true.
+	//		Dispatch4();				// No barrier, overlaps 3.
+	//		Vendor_EndUAVOverlap();		// Set overlap flag to false, any dispatches after this will not overlap.
+	//
+	// This correctly serializes dispatches 2 and 3.
+	UAVOverlapState = EUAVOverlapState::EPending;
+}
 
-#if !PLATFORM_HOLOLENS
-	if (IsRHIDeviceNVIDIA())
-	{
-		CHECK_NVAPI(NvAPI_D3D11_EndUAVOverlap(Direct3DDevice));
-	}
-	else if (IsRHIDeviceAMD())
+void FD3D11DynamicRHI::RHIEndUAVOverlap()
+{
+	if (!bUAVOverlapAllowed)
 	{
-		CHECK_AGS(agsDriverExtensionsDX11_EndUAVOverlap(AmdAgsContext, Direct3DDeviceIMContext));
+		return;
 	}
-	else if (IsRHIDeviceIntel())
+
+	checkf(UAVOverlapState != EUAVOverlapState::EOff, TEXT("Mismatched call to EndUAVOverlap. Ensure all calls to RHICmdList.BeginUAVOverlap() are paired with a call to RHICmdList.EndUAVOverlap()."));
+
+	// Only call the driver API if we got a dispatch in between the call to RHIBeginUAVOverlap() and this call to RHIEndUAVOverlap(). Otherwise it's an
+	// empty overlap group and we can simply cancel the request.
+	if (UAVOverlapState == EUAVOverlapState::EOn)
 	{
-#if INTEL_EXTENSIONS
-		if (IntelD3D11ExtensionFuncs && IntelD3D11ExtensionFuncs->D3D11EndUAVOverlap)
+#if !PLATFORM_HOLOLENS
+		if (IsRHIDeviceNVIDIA())
+		{
+			CHECK_NVAPI(NvAPI_D3D11_EndUAVOverlap(Direct3DDevice));
+		}
+		else if (IsRHIDeviceAMD())
 		{
-			CHECK_INTEL(IntelD3D11ExtensionFuncs->D3D11EndUAVOverlap(IntelExtensionContext));
+			CHECK_AGS(agsDriverExtensionsDX11_EndUAVOverlap(AmdAgsContext, Direct3DDeviceIMContext));
 		}
+		else if (IsRHIDeviceIntel())
+		{
+#if INTEL_EXTENSIONS
+			if (IntelD3D11ExtensionFuncs && IntelD3D11ExtensionFuncs->D3D11EndUAVOverlap)
+			{
+				CHECK_INTEL(IntelD3D11ExtensionFuncs->D3D11EndUAVOverlap(IntelExtensionContext));
+			}
 #endif
-	}
-	else
-	{
-		ensureMsgf(false, TEXT("EndUAVOverlap not implemented for this GPU IHV."));
-	}
+		}
+		else
+		{
+			ensureMsgf(false, TEXT("EndUAVOverlap not implemented for this GPU IHV."));
+		}
 #endif	
-}
-
-void FD3D11DynamicRHI::UAVBarrier()
-{
-	if (bUAVOverlapEnabled)
-	{
-		RHIEndUAVOverlap();
-		RHIBeginUAVOverlap();
 	}
-}
 
+	UAVOverlapState = EUAVOverlapState::EOff;
+}
 
 void FD3D11DynamicRHI::RHIAutomaticCacheFlushAfterComputeShader(bool bEnable)
 {
 	if (bEnable)
 	{
-		if (bUAVOverlapEnabled)
+		if (UAVOverlapState != EUAVOverlapState::EOff)
 		{
 			RHIEndUAVOverlap();
 		}
 	}
 	else
 	{
-		if (!bUAVOverlapEnabled)
+		if (UAVOverlapState == EUAVOverlapState::EOff)
 		{
 			RHIBeginUAVOverlap();
 		}
@@ -2196,7 +2221,11 @@ void FD3D11DynamicRHI::RHIAutomaticCacheFlushAfterComputeShader(bool bEnable)
 
 void FD3D11DynamicRHI::RHIFlushComputeShaderCache()
 {
-	UAVBarrier();
+	if (UAVOverlapState != EUAVOverlapState::EOff)
+	{
+		RHIEndUAVOverlap();
+		RHIBeginUAVOverlap();
+	}
 }
 
 //*********************** StagingBuffer Implementation ***********************//
diff --git a/Engine/Source/Runtime/Windows/D3D11RHI/Private/D3D11RHIPrivate.h b/Engine/Source/Runtime/Windows/D3D11RHI/Private/D3D11RHIPrivate.h
--- a/Engine/Source/Runtime/Windows/D3D11RHI/Private/D3D11RHIPrivate.h
+++ b/Engine/Source/Runtime/Windows/D3D11RHI/Private/D3D11RHIPrivate.h
@@ -1072,11 +1072,18 @@ protected:
 	void StopIntelExtensions();
 #endif // INTEL_EXTENSIONS
 
+	enum class EUAVOverlapState
+	{
+		EOff,
+		EPending,
+		EOn
+	};
+
 	// UAV overlap handling
 	bool bUAVOverlapAllowed = false;
-	bool bUAVOverlapEnabled = false;
+	EUAVOverlapState UAVOverlapState = EUAVOverlapState::EOff;
 	bool IsUAVOverlapSupported();
-	void UAVBarrier();
+	void ApplyUAVOverlapState();
 
 #if INTEL_METRICSDISCOVERY
 	void CreateIntelMetricsDiscovery();
