UE4.25 - Change warnings to display
#rnx

diff --git a/Engine/Source/Runtime/D3D12RHI/Private/D3D12Texture.cpp b/Engine/Source/Runtime/D3D12RHI/Private/D3D12Texture.cpp
--- a/Engine/Source/Runtime/D3D12RHI/Private/D3D12Texture.cpp
+++ b/Engine/Source/Runtime/D3D12RHI/Private/D3D12Texture.cpp
@@ -1175,7 +1175,7 @@ FTexture2DRHIRef FD3D12DynamicRHI::RHIAsyncCreateTexture2D(uint32 SizeX, uint32
 
 		if (MipSize > TempBufferSize)
 		{
-			UE_LOG(LogD3D12RHI, Warning, TEXT("Temp texture streaming buffer not large enough, needed %d bytes"), MipSize);
+			UE_LOG(LogD3D12RHI, Display, TEXT("Temp texture streaming buffer not large enough, needed %d bytes"), MipSize);
 			check(TempBufferSize == ZeroBufferSize);
 			TempBufferSize = MipSize;
 			TempBuffer = FMemory::Malloc(TempBufferSize);
diff --git a/Engine/Source/Runtime/Windows/D3D11RHI/Private/D3D11Texture.cpp b/Engine/Source/Runtime/Windows/D3D11RHI/Private/D3D11Texture.cpp
--- a/Engine/Source/Runtime/Windows/D3D11RHI/Private/D3D11Texture.cpp
+++ b/Engine/Source/Runtime/Windows/D3D11RHI/Private/D3D11Texture.cpp
@@ -1174,7 +1174,7 @@ FTexture2DRHIRef FD3D11DynamicRHI::RHIAsyncCreateTexture2D(uint32 SizeX,uint32 S
 
 		if (MipSize > TempBufferSize)
 		{
-			UE_LOG(LogD3D11RHI,Warning,TEXT("Temp texture streaming buffer not large enough, needed %d bytes"),MipSize);
+			UE_LOG(LogD3D11RHI, Display,TEXT("Temp texture streaming buffer not large enough, needed %d bytes"),MipSize);
 			check(TempBufferSize == ZeroBufferSize);
 			TempBufferSize = MipSize;
 			TempBuffer = FMemory::Malloc(TempBufferSize);
