UE4.25 - Potential compile fix
#rnx

diff --git a/Engine/Source/Runtime/Windows/D3D11RHI/Private/Windows/WindowsD3D11Device.cpp b/Engine/Source/Runtime/Windows/D3D11RHI/Private/Windows/WindowsD3D11Device.cpp
--- a/Engine/Source/Runtime/Windows/D3D11RHI/Private/Windows/WindowsD3D11Device.cpp
+++ b/Engine/Source/Runtime/Windows/D3D11RHI/Private/Windows/WindowsD3D11Device.cpp
@@ -20,7 +20,9 @@
 #include "Runtime/HeadMountedDisplay/Public/IHeadMountedDisplayModule.h"
 #include "GenericPlatform/GenericPlatformDriver.h"			// FGPUDriverInfo
 #include "GenericPlatform/GenericPlatformCrashContext.h"
+THIRD_PARTY_INCLUDES_START
 #include "dxgi1_6.h"
+THIRD_PARTY_INCLUDES_END
 #include "RHIValidation.h"
 
 #if NV_AFTERMATH
