UE4.25 - Compile fix
#rnx

diff --git a/Engine/Source/Runtime/Windows/D3D11RHI/Private/Windows/WindowsD3D11Viewport.cpp b/Engine/Source/Runtime/Windows/D3D11RHI/Private/Windows/WindowsD3D11Viewport.cpp
--- a/Engine/Source/Runtime/Windows/D3D11RHI/Private/Windows/WindowsD3D11Viewport.cpp
+++ b/Engine/Source/Runtime/Windows/D3D11RHI/Private/Windows/WindowsD3D11Viewport.cpp
@@ -11,7 +11,9 @@
 #include "Windows/AllowWindowsPlatformTypes.h"
 #include <dwmapi.h>
 
+THIRD_PARTY_INCLUDES_START
 #include "dxgi1_6.h"
+THIRD_PARTY_INCLUDES_END
 
 #ifndef DXGI_SWAP_CHAIN_FLAG_ALLOW_TEARING
 #define DXGI_SWAP_CHAIN_FLAG_ALLOW_TEARING  2048
