proposed fix for a rare crash.
theory is this happens for users with multiple adapters, where not all of them support ALLOW_TEARING

diff --git a/Engine/Source/Runtime/Windows/D3D11RHI/Private/Windows/WindowsD3D11Viewport.cpp b/Engine/Source/Runtime/Windows/D3D11RHI/Private/Windows/WindowsD3D11Viewport.cpp
--- a/Engine/Source/Runtime/Windows/D3D11RHI/Private/Windows/WindowsD3D11Viewport.cpp
+++ b/Engine/Source/Runtime/Windows/D3D11RHI/Private/Windows/WindowsD3D11Viewport.cpp
@@ -19,7 +19,6 @@ THIRD_PARTY_INCLUDES_END
 #define DXGI_SWAP_CHAIN_FLAG_ALLOW_TEARING  2048
 #endif
 
-static bool GSwapFlagsInitialized = false;
 static DXGI_SWAP_EFFECT GSwapEffect = DXGI_SWAP_EFFECT_DISCARD;
 static uint32 GSwapChainFlags = DXGI_SWAP_CHAIN_FLAG_ALLOW_MODE_SWITCH;
 static uint32 GSwapChainBufferCount = 1;
@@ -66,8 +65,9 @@ FD3D11Viewport::FD3D11Viewport(FD3D11DynamicRHI* InD3DRHI,HWND InWindowHandle,ui
 	TRefCountPtr<IDXGIDevice> DXGIDevice;
 	VERIFYD3D11RESULT_EX(D3DRHI->GetDevice()->QueryInterface(IID_IDXGIDevice, (void**)DXGIDevice.GetInitReference()), D3DRHI->GetDevice());
 
-	if(!GSwapFlagsInitialized)
 	{
+		GSwapEffect = DXGI_SWAP_EFFECT_DISCARD;
+		GSwapChainFlags = DXGI_SWAP_CHAIN_FLAG_ALLOW_MODE_SWITCH;
 		IDXGIFactory1* Factory1 = D3DRHI->GetFactory();
 		TRefCountPtr<IDXGIFactory5> Factory5;
 
@@ -84,7 +84,6 @@ FD3D11Viewport::FD3D11Viewport(FD3D11DynamicRHI* InD3DRHI,HWND InWindowHandle,ui
 				}
 			}
 		}
-		GSwapFlagsInitialized = true;
 	}
 	uint32 BufferCount = GSwapChainBufferCount;
 	// If requested, keep a handle to a DXGIOutput so we can force that display on fullscreen swap
