Protection against nullptr in GetMMDeviceInfo

diff --git a/Engine/Source/Runtime/Windows/AudioMixerXAudio2/Private/AudioMixerPlatformXAudio2.cpp b/Engine/Source/Runtime/Windows/AudioMixerXAudio2/Private/AudioMixerPlatformXAudio2.cpp
--- a/Engine/Source/Runtime/Windows/AudioMixerXAudio2/Private/AudioMixerPlatformXAudio2.cpp
+++ b/Engine/Source/Runtime/Windows/AudioMixerXAudio2/Private/AudioMixerPlatformXAudio2.cpp
@@ -415,7 +415,13 @@ namespace Audio
 
 		// Get the format of the property
 		WaveFormatEx = (WAVEFORMATEX *)DeviceFormat.blob.pBlobData;
-
+		if (!ensure(WaveFormatEx))
+		{
+			// Force an error if we failed to get a WaveFormat back from the data blob.
+			Result = E_FAIL;
+			XAUDIO2_GOTO_CLEANUP_ON_FAIL(Result);
+		}
+		
 		OutInfo.DeviceId = FString(DeviceId);
 		OutInfo.NumChannels = FMath::Clamp((int32)WaveFormatEx->nChannels, 2, 8);
 		OutInfo.SampleRate = WaveFormatEx->nSamplesPerSec;
@@ -581,11 +587,13 @@ namespace Audio
 			XAUDIO2_GOTO_CLEANUP_ON_FAIL(Result);
 		}
 
-		check(Device);
-		bSucceeded = GetMMDeviceInfo(Device, OutInfo);
+		if (ensure(Device))
+		{
+			bSucceeded = GetMMDeviceInfo(Device, OutInfo);
 
-		// Fix up if this was a default device
-		OutInfo.bIsSystemDefault = bIsDefault;
+			// Fix up if this was a default device
+			OutInfo.bIsSystemDefault = bIsDefault;
+		}
 
 	Cleanup:
 		SAFE_RELEASE(Device);
