fix a crash that would occur if game starts up in a specific mode, and clicks in another window.
This would cause DXGI_ERROR_INVALID_CALL, which would cause the engine to halt.

diff --git a/Engine/Plugins/Runtime/WindowsMixedReality/Source/WindowsMixedRealityRHI/Private/WindowsMixedRealityRHI.cpp b/Engine/Plugins/Runtime/WindowsMixedReality/Source/WindowsMixedRealityRHI/Private/WindowsMixedRealityRHI.cpp
--- a/Engine/Plugins/Runtime/WindowsMixedReality/Source/WindowsMixedRealityRHI/Private/WindowsMixedRealityRHI.cpp
+++ b/Engine/Plugins/Runtime/WindowsMixedReality/Source/WindowsMixedRealityRHI/Private/WindowsMixedRealityRHI.cpp
@@ -212,7 +212,7 @@ FWindowsMixedRealityViewport::FWindowsMixedRealityViewport(FD3D11DynamicRHI* InD
 	SizeY = InSizeY;
 	bIsFullscreen = bInIsFullscreen;
 	PixelFormat  = InPreferredPixelFormat;
-	bIsValid = true;
+	ValidState = 0;
 
 	D3DRHI->Viewports.Add(this);
 
@@ -320,7 +320,7 @@ void FWindowsMixedRealityViewport::Resize(uint32 InSizeX, uint32 InSizeY, bool b
 	if (bIsFullscreen != bInIsFullscreen)
 	{
 		bIsFullscreen = bInIsFullscreen;
-		bIsValid = false;
+		ValidState = VIEWPORT_INVALID;
 	}
 
 	// Float RGBA backbuffers are requested whenever HDR mode is desired
diff --git a/Engine/Source/Runtime/Windows/D3D11RHI/Private/D3D11Viewport.cpp b/Engine/Source/Runtime/Windows/D3D11RHI/Private/D3D11Viewport.cpp
--- a/Engine/Source/Runtime/Windows/D3D11RHI/Private/D3D11Viewport.cpp
+++ b/Engine/Source/Runtime/Windows/D3D11RHI/Private/D3D11Viewport.cpp
@@ -288,7 +288,7 @@ void FD3D11Viewport::Resize(uint32 InSizeX, uint32 InSizeY, bool bInIsFullscreen
 	if(bIsFullscreen != bInIsFullscreen)
 	{
 		bIsFullscreen = bInIsFullscreen;
-		bIsValid = false;
+		ValidState = VIEWPORT_INVALID;
 
 		if (bNeedSwapChain)
 		{
@@ -349,7 +349,6 @@ bool FD3D11Viewport::PresentChecked(int32 SyncInterval)
 			{
 				Flags |= DXGI_PRESENT_ALLOW_TEARING;
 			}
-
 			Result = SwapChain->Present(SyncInterval, Flags);
 		}
 
@@ -360,9 +359,10 @@ bool FD3D11Viewport::PresentChecked(int32 SyncInterval)
 	}
 
 	FThreadHeartBeat::Get().PresentFrame();
-
 	if (FAILED(Result))
 	{
+		PresentFailCount++;
+		UE_LOG(LogRHI, Error, TEXT("Present Fail Count %i"), PresentFailCount);
 		DXGI_SWAP_CHAIN_DESC Desc;
 		UE_LOG(LogRHI, Error, TEXT("SyncInterval %i"), SyncInterval);
 		if (!FAILED(SwapChain->GetDesc(&Desc)))
@@ -383,10 +383,22 @@ bool FD3D11Viewport::PresentChecked(int32 SyncInterval)
 			UE_LOG(LogRHI, Error, TEXT("SwapChainDesc.SwapEffect %u"), Desc.SwapEffect);
 			UE_LOG(LogRHI, Error, TEXT("SwapChainDesc.Flags %u"), Desc.Flags);
 		}
+		if (PresentFailCount > 10)
+		{
+			VERIFYD3D11RESULT_EX(Result, D3DRHI->GetDevice());
+		}
+		else
+		{
+			ValidState = (VIEWPORT_INVALID | VIEWPORT_FULLSCREEN_LOST);
+		}
 	}
+	else
+	{
+		PresentFailCount = 0;
+	}
+
 	FThreadHeartBeat::Get().PresentFrame();
 
-	VERIFYD3D11RESULT_EX(Result, D3DRHI->GetDevice());
 	D3DRHI->GetDeviceContext()->OMSetRenderTargets(0,0,0);
 
 	return bNeedNativePresent;
@@ -530,7 +542,7 @@ bool FD3D11Viewport::Present(bool bLockToVsync)
 #if	D3D11_WITH_DWMAPI
 #if !PLATFORM_HOLOLENS
 	// We can't call Present if !bIsValid, as it waits a window message to be processed, but the main thread may not be pumping the message handler.
-	if(bIsValid && SwapChain.IsValid())
+	if(ValidState != 0 && SwapChain.IsValid())
 	{
 		// Check if the viewport's swap chain has been invalidated by DXGI.
 		BOOL bSwapChainFullscreenState;
@@ -539,7 +551,7 @@ bool FD3D11Viewport::Present(bool bLockToVsync)
 		// Can't compare BOOL with bool...
 		if ( (!!bSwapChainFullscreenState)  != bIsFullscreen )
 		{
-			bIsValid = false;
+			ValidState = VIEWPORT_INVALID;
 		}
 	}
 #endif
@@ -551,7 +563,7 @@ bool FD3D11Viewport::Present(bool bLockToVsync)
 		DXGIDevice->SetMaximumFrameLatency(MaximumFrameLatency);
 	}
 
-	if (!bIsValid)
+	if(0 != (ValidState & VIEWPORT_INVALID))
 	{
 		return false;
 	}
diff --git a/Engine/Source/Runtime/Windows/D3D11RHI/Private/HoloLens/HoloLensD3D11Viewport.cpp b/Engine/Source/Runtime/Windows/D3D11RHI/Private/HoloLens/HoloLensD3D11Viewport.cpp
--- a/Engine/Source/Runtime/Windows/D3D11RHI/Private/HoloLens/HoloLensD3D11Viewport.cpp
+++ b/Engine/Source/Runtime/Windows/D3D11RHI/Private/HoloLens/HoloLensD3D11Viewport.cpp
@@ -54,9 +54,10 @@ FD3D11Viewport::FD3D11Viewport(FD3D11DynamicRHI* InD3DRHI,HWND InWindowHandle,ui
 	WindowHandle(InWindowHandle),
 	SizeX(InSizeX),
 	SizeY(InSizeY),
-	bIsFullscreen(bInIsFullscreen),
+	PresentFailCount(0),
+	ValidState(0),
 	PixelFormat(InPreferredPixelFormat),
-	bIsValid(true),
+	bIsFullscreen(bInIsFullscreen),
 	FrameSyncEvent(InD3DRHI)
 {
 	D3DRHI->Viewports.Add(this);
@@ -202,12 +203,13 @@ void FD3D11Viewport::CheckHDRMonitorStatus()
 
 void FD3D11Viewport::ConditionalResetSwapChain(bool bIgnoreFocus)
 {
-	if (!bIsValid)
+	uint32 Valid = ValidState;
+	if (0 != (Valid & VIEWPORT_INVALID))
 	{
-		if (bFullscreenLost)
+		if (0 != (Valid & VIEWPORT_FULLSCREEN_LOST))
 		{
 			FlushRenderingCommands();
-			bFullscreenLost = false;
+			ValidState &= ~(VIEWPORT_FULLSCREEN_LOST);
 			Resize(SizeX, SizeY, false, PixelFormat);
 		}
 		else
@@ -219,7 +221,7 @@ void FD3D11Viewport::ConditionalResetSwapChain(bool bIgnoreFocus)
 
 void FD3D11Viewport::ResetSwapChainInternal(bool bIgnoreFocus)
 {
-	if(!bIsValid)
+	if (0 != (ValidState & VIEWPORT_INVALID))
 	{
 		const bool bIsFocused = true;
 		const bool bIsIconic = false;
@@ -230,7 +232,7 @@ void FD3D11Viewport::ResetSwapChainInternal(bool bIgnoreFocus)
 			HRESULT Result = SwapChain->SetFullscreenState(bIsFullscreen,NULL);
 			if(SUCCEEDED(Result))
 			{
-				bIsValid = true;
+				ValidState &= ~(VIEWPORT_INVALID);
 			}
 			else
 			{
diff --git a/Engine/Source/Runtime/Windows/D3D11RHI/Private/Windows/WindowsD3D11Viewport.cpp b/Engine/Source/Runtime/Windows/D3D11RHI/Private/Windows/WindowsD3D11Viewport.cpp
--- a/Engine/Source/Runtime/Windows/D3D11RHI/Private/Windows/WindowsD3D11Viewport.cpp
+++ b/Engine/Source/Runtime/Windows/D3D11RHI/Private/Windows/WindowsD3D11Viewport.cpp
@@ -48,11 +48,11 @@ FD3D11Viewport::FD3D11Viewport(FD3D11DynamicRHI* InD3DRHI,HWND InWindowHandle,ui
 	MaximumFrameLatency(3),
 	SizeX(InSizeX),
 	SizeY(InSizeY),
-	bIsFullscreen(bInIsFullscreen),
-	bFullscreenLost(false),
+	PresentFailCount(0),
+	ValidState(0),
 	PixelFormat(InPreferredPixelFormat),
 	PixelColorSpace(EColorSpaceAndEOTF::ERec709_sRGB),
-	bIsValid(true),
+	bIsFullscreen(bInIsFullscreen),
 	FrameSyncEvent(InD3DRHI)
 {
 	check(IsInGameThread());
@@ -285,12 +285,13 @@ void FD3D11Viewport::CheckHDRMonitorStatus()
 
 void FD3D11Viewport::ConditionalResetSwapChain(bool bIgnoreFocus)
 {
-	if (!bIsValid)
+	uint32 Valid = ValidState;
+	if(0 != (Valid & VIEWPORT_INVALID))
 	{
-		if (bFullscreenLost)
+		if(0 != (Valid & VIEWPORT_FULLSCREEN_LOST))
 		{
 			FlushRenderingCommands();
-			bFullscreenLost = false;
+			ValidState &= ~(VIEWPORT_FULLSCREEN_LOST);
 			Resize(SizeX, SizeY, false, PixelFormat);
 		}
 		else
@@ -302,7 +303,7 @@ void FD3D11Viewport::ConditionalResetSwapChain(bool bIgnoreFocus)
 
 void FD3D11Viewport::ResetSwapChainInternal(bool bIgnoreFocus)
 {
-	if (!bIsValid)
+	if (0 != (ValidState & VIEWPORT_INVALID))
 	{
 		// Check if the viewport's window is focused before resetting the swap chain's fullscreen state.
 		HWND FocusWindow = ::GetFocus();
@@ -318,7 +319,7 @@ void FD3D11Viewport::ResetSwapChainInternal(bool bIgnoreFocus)
 
 			if (SUCCEEDED(Result))
 			{
-				bIsValid = true;
+				ValidState &= ~(VIEWPORT_INVALID);
 			}
 			else if (Result != DXGI_ERROR_NOT_CURRENTLY_AVAILABLE && Result != DXGI_STATUS_MODE_CHANGE_IN_PROGRESS)
 			{
diff --git a/Engine/Source/Runtime/Windows/D3D11RHI/Public/D3D11Viewport.h b/Engine/Source/Runtime/Windows/D3D11RHI/Public/D3D11Viewport.h
--- a/Engine/Source/Runtime/Windows/D3D11RHI/Public/D3D11Viewport.h
+++ b/Engine/Source/Runtime/Windows/D3D11RHI/Public/D3D11Viewport.h
@@ -60,8 +60,14 @@ static DXGI_FORMAT GetRenderTargetFormat(EPixelFormat PixelFormat)
 class D3D11RHI_API FD3D11Viewport : public FRHIViewport
 {
 public:
+	enum ED3DViewportValidFlags : uint32
+	{
+		VIEWPORT_INVALID = 0x1,
+		VIEWPORT_FULLSCREEN_LOST = 0x2,
+
+	};
 
-	FD3D11Viewport(class FD3D11DynamicRHI* InD3DRHI) : D3DRHI(InD3DRHI), bFullscreenLost(false), FrameSyncEvent(InD3DRHI) {}
+	FD3D11Viewport(class FD3D11DynamicRHI* InD3DRHI) : D3DRHI(InD3DRHI), PresentFailCount(0), ValidState (0), FrameSyncEvent(InD3DRHI) {}
 	FD3D11Viewport(class FD3D11DynamicRHI* InD3DRHI, HWND InWindowHandle, uint32 InSizeX, uint32 InSizeY, bool bInIsFullscreen, EPixelFormat InPreferredPixelFormat);
 	~FD3D11Viewport();
 
@@ -119,6 +125,7 @@ public:
 	virtual void* GetNativeWindow(void** AddParam = nullptr) const override { return (void*)WindowHandle; }
 	static FD3D11Texture2D* GetSwapChainSurface(FD3D11DynamicRHI* D3DRHI, EPixelFormat PixelFormat, uint32 SizeX, uint32 SizeY, IDXGISwapChain* SwapChain);
 
+
 protected:
 
 	void ResetSwapChainInternal(bool bIgnoreFocus);
@@ -143,11 +150,12 @@ protected:
 	uint32 SizeX;
 	uint32 SizeY;
 	uint32 BackBufferCount;
-	bool bIsFullscreen : 1;
-	bool bFullscreenLost : 1;
+	uint32 PresentFailCount;
+	TAtomic<uint32> ValidState;
 	EPixelFormat PixelFormat;
 	EColorSpaceAndEOTF PixelColorSpace;
-	bool bIsValid;
+	bool bIsFullscreen;
+
 #if PLATFORM_HOLOLENS
 	TRefCountPtr<IDXGISwapChain1> SwapChain;
 #else
