disabling video/image sharing/capture on Windows during protected video playback

diff --git a/Engine/Source/Runtime/Windows/WindowsPlatformFeatures/Private/WindowsVideoRecordingSystem.cpp b/Engine/Source/Runtime/Windows/WindowsPlatformFeatures/Private/WindowsVideoRecordingSystem.cpp
--- a/Engine/Source/Runtime/Windows/WindowsPlatformFeatures/Private/WindowsVideoRecordingSystem.cpp
+++ b/Engine/Source/Runtime/Windows/WindowsPlatformFeatures/Private/WindowsVideoRecordingSystem.cpp
@@ -13,12 +13,149 @@
 
 #include "HighlightRecorder.h"
 
+#include <dwmapi.h>
+
 DEFINE_VIDEOSYSTEMRECORDING_STATS
 DEFINE_LOG_CATEGORY(WindowsVideoRecordingSystem);
 CSV_DEFINE_CATEGORY(WindowsVideoRecordingSystem, true);
 
 WINDOWSPLATFORMFEATURES_START
 
+/**
+ * This internal helper class handles disabling of the Windows built-in screenshot (PrtScr key) and clip
+ * recorder (via Win-G key menu on Windows 10)
+ */
+class FWindowsVideoRecordingSystem::FWindowsScreenRecording
+{
+public:
+	bool Disable()
+	{
+		return DisableInternal();
+	}
+
+	bool Enable()
+	{
+		return EnableInternal();
+	}
+
+private:
+	struct FWindowHandle
+	{
+		FWindowHandle() : WindowsHandle(0), Affinity(0), bIsSupported(false), DisableCount(0) { }
+		HWND			WindowsHandle;
+		DWORD			Affinity;
+		BOOL			bIsSupported;
+		volatile int32	DisableCount;
+	};
+
+	bool DisableInternal()
+	{
+		if (GetCurrentlyActiveHandle())
+		{
+			check(CurrentHandle.IsValid());
+			if (FPlatformAtomics::InterlockedIncrement(&CurrentHandle->DisableCount) == 1)
+			{
+				if (SetWindowDisplayAffinity(CurrentHandle->WindowsHandle, WDA_MONITOR))
+				{
+					return true;
+				}
+				else
+				{
+					//DWORD lerr = GetLastError();
+					return false;
+				}
+			}
+			return true;
+		}
+		return false;
+	}
+
+	bool EnableInternal()
+	{
+		if (GetCurrentlyActiveHandle())
+		{
+			check(CurrentHandle.IsValid());
+			// Check that the disable count is valid. Should we have either lost the original handle or the user
+			// is calling enable without having called disable first (which is an error but we let slide) we
+			// should not drop the disable count.
+			if (FPlatformAtomics::AtomicRead_Relaxed(&CurrentHandle->DisableCount) > 0)
+			{
+				if (FPlatformAtomics::InterlockedDecrement(&CurrentHandle->DisableCount) == 0)
+				{
+					if (SetWindowDisplayAffinity(CurrentHandle->WindowsHandle, CurrentHandle->Affinity))
+					{
+						return true;
+					}
+					else
+					{
+						//DWORD lerr = GetLastError();
+						return false;
+					}
+				}
+			}
+			// Let's return true even if we weren't disabled on a new handle.
+			return true;
+		}
+		return false;
+	}
+
+	TUniquePtr<FWindowHandle> CreateCurrentHandle()
+	{
+		TUniquePtr<FWindowHandle>	NewHandle(new FWindowHandle);
+		HRESULT						res;
+		res = DwmIsCompositionEnabled(&NewHandle->bIsSupported);
+		if (res == S_OK)
+		{
+			if (NewHandle->bIsSupported)
+			{
+				if ((NewHandle->WindowsHandle = GetWindowHandleInternal()) != 0)
+				{
+					if ((NewHandle->bIsSupported = GetWindowDisplayAffinity(NewHandle->WindowsHandle, &NewHandle->Affinity)) != false)
+					{
+						return MoveTemp(NewHandle);
+					}
+				}
+			}
+		}
+		return nullptr;
+	}
+	HWND GetWindowHandleInternal()
+	{
+		//return (HWND)GEngine->GameViewport->GetWindow()->GetNativeWindow()->GetOSWindowHandle();
+		return GetActiveWindow();
+	}
+
+	bool GetCurrentlyActiveHandle()
+	{
+		TUniquePtr<FWindowHandle> Now = CreateCurrentHandle();
+		if (Now.IsValid())
+		{
+			if (CurrentHandle.IsValid() && CurrentHandle->WindowsHandle != Now->WindowsHandle)
+			{
+				// In case the window handle has changed for any reason we do not try to restore the
+				// state we got when we created our internal handle. In all likelihood we do not own
+				// the window any more and should not be messing with its handle.
+				CurrentHandle.Reset();
+			}
+			if (!CurrentHandle.IsValid())
+			{
+				CurrentHandle = MoveTemp(Now);
+			}
+			return true;
+		}
+		else
+		{
+			// Likewise if we can't get a valid handle right now we drop whatever we were using before.
+			CurrentHandle.Reset();
+		}
+		return false;
+	}
+
+	TUniquePtr<FWindowHandle>	CurrentHandle;
+};
+
+
+
 FWindowsVideoRecordingSystem::FWindowsVideoRecordingSystem()
 {
 	UE_LOG(WindowsVideoRecordingSystem, Verbose, TEXT("%s"), __FUNCTIONW__);
@@ -53,6 +190,21 @@ void FWindowsVideoRecordingSystem::EnableRecording(bool bEnableRecording)
 	{
 		Recorder.Reset();
 	}
+
+
+	// Also disable/enable the Windows screenshot and clip recording ability.
+	if (!ScreenshotAndRecorderHandler.IsValid())
+	{
+		ScreenshotAndRecorderHandler.Reset(new FWindowsScreenRecording);
+	}
+	if (bEnableRecording)
+	{
+		ScreenshotAndRecorderHandler->Enable();
+	}
+	else
+	{
+		ScreenshotAndRecorderHandler->Disable();
+	}
 }
 
 bool FWindowsVideoRecordingSystem::IsEnabled() const
diff --git a/Engine/Source/Runtime/Windows/WindowsPlatformFeatures/Private/WindowsVideoRecordingSystem.h b/Engine/Source/Runtime/Windows/WindowsPlatformFeatures/Private/WindowsVideoRecordingSystem.h
--- a/Engine/Source/Runtime/Windows/WindowsPlatformFeatures/Private/WindowsVideoRecordingSystem.h
+++ b/Engine/Source/Runtime/Windows/WindowsPlatformFeatures/Private/WindowsVideoRecordingSystem.h
@@ -55,5 +55,8 @@ private:
 	TUniquePtr<FHighlightRecorder> Recorder;
 
 	bool bAudioFormatChecked = false;
+
+	class FWindowsScreenRecording;
+	TUniquePtr<FWindowsScreenRecording>	ScreenshotAndRecorderHandler;
 };
 
