Deprecate XAudio2

diff --git a/Engine/Source/Runtime/Windows/XAudio2/Private/XAudio2Device.cpp b/Engine/Source/Runtime/Windows/XAudio2/Private/XAudio2Device.cpp
--- a/Engine/Source/Runtime/Windows/XAudio2/Private/XAudio2Device.cpp
+++ b/Engine/Source/Runtime/Windows/XAudio2/Private/XAudio2Device.cpp
@@ -42,7 +42,7 @@ THIRD_PARTY_INCLUDES_END
 
 DEFINE_LOG_CATEGORY(LogXAudio2);
 
-class FXAudio2DeviceModule : public IAudioDeviceModule
+class UE_DEPRECATED(4.26, "XAudio2 Module has been depricated.") FXAudio2DeviceModule : public IAudioDeviceModule
 {
 public:
 
diff --git a/Engine/Source/Runtime/Windows/XAudio2/Private/XAudio2Source.cpp b/Engine/Source/Runtime/Windows/XAudio2/Private/XAudio2Source.cpp
--- a/Engine/Source/Runtime/Windows/XAudio2/Private/XAudio2Source.cpp
+++ b/Engine/Source/Runtime/Windows/XAudio2/Private/XAudio2Source.cpp
@@ -1672,6 +1672,8 @@ FSpatializationHelper::FSpatializationHelper( void )
 
 void FSpatializationHelper::Init()
 {
+#if WITH_XAUDIO3D
+
 	// Initialize X3DAudio
 	//
 	//  Speaker geometry configuration on the final mix, specifies assignment of channels
@@ -1752,10 +1754,13 @@ void FSpatializationHelper::Init()
 	DSPSettings.DstChannelCount = SPEAKER_COUNT;
 	DSPSettings.pMatrixCoefficients = MatrixCoefficients;
 	DSPSettings.pDelayTimes = nullptr;
+
+#endif //WITH_XAUDIO3D
 }
 
 void FSpatializationHelper::DumpSpatializationState() const
 {
+#if WITH_XAUDIO3D
 	struct FLocal
 	{
 		static void DumpChannelArray(const FString& Indent, const FString& ArrayName, int32 NumChannels, const float* pChannelArray)
@@ -1905,10 +1910,14 @@ void FSpatializationHelper::DumpSpatializationState() const
 
 	// MatrixCoefficients
 	FLocal::DumpChannelArray(TEXT("  "), TEXT("MatrixCoefficients"), UE_ARRAY_COUNT(MatrixCoefficients), MatrixCoefficients);
+
+#endif //WITH_XAUDIO3D
 }
 
 void FSpatializationHelper::CalculateDolbySurroundRate( const FVector& OrientFront, const FVector& ListenerPosition, const FVector& EmitterPosition, float OmniRadius, float* OutVolumes  )
 {
+#if WITH_XAUDIO3D
+
 #if ENABLE_NAN_DIAGNOSTIC
 	OrientFront.DiagnosticCheckNaN(TEXT("FSpatializationHelper: OrientFront"));
 	ListenerPosition.DiagnosticCheckNaN(TEXT("FSpatializationHelper: ListenerPosition"));
@@ -1965,6 +1974,8 @@ void FSpatializationHelper::CalculateDolbySurroundRate( const FVector& OrientFro
 	}
 
 	OutVolumes[CHANNELOUT_REVERB] *= DSPSettings.ReverbLevel;
+
+#endif //WITH_XAUDIO3D
 }
 
 /*------------------------------------------------------------------------------------
diff --git a/Engine/Source/Runtime/Windows/XAudio2/Private/XAudio2Support.h b/Engine/Source/Runtime/Windows/XAudio2/Private/XAudio2Support.h
--- a/Engine/Source/Runtime/Windows/XAudio2/Private/XAudio2Support.h
+++ b/Engine/Source/Runtime/Windows/XAudio2/Private/XAudio2Support.h
@@ -169,7 +169,7 @@ private:
 
 #define UE4_XAUDIO3D_INPUTCHANNELS 1
 
-struct FPCMBufferInfo
+struct UE_DEPRECATED(4.26, "XAudio2 Module has been deprecated.") FPCMBufferInfo
 {
 	/** Format of the source PCM data */
 	WAVEFORMATEX				PCMFormat;
@@ -180,7 +180,7 @@ struct FPCMBufferInfo
 };
 
 #if XAUDIO_SUPPORTS_XMA2WAVEFORMATEX
-struct FXMA2BufferInfo
+struct UE_DEPRECATED(4.26, "XAudio2 Module has been deprecated.") FXMA2BufferInfo
 {
 	/** Format of the source XMA2 data */
 	XMA2WAVEFORMATEX			XMA2Format;
@@ -191,7 +191,7 @@ struct FXMA2BufferInfo
 };
 #endif	//XAUDIO_SUPPORTS_XMA2WAVEFORMATEX
 
-struct FXWMABufferInfo
+struct UE_DEPRECATED(4.26, "XAudio2 Module has been deprecated.")  FXWMABufferInfo
 {
 	/** Format of the source XWMA data */
 	WAVEFORMATEXTENSIBLE		XWMAFormat;
@@ -229,7 +229,7 @@ struct FPendingAsyncTaskInfo
 /**
  * XAudio2 implementation of FSoundBuffer, containing the wave data and format information.
  */
-class FXAudio2SoundBuffer : public FSoundBuffer
+class UE_DEPRECATED(4.26, "XAudio2 Module has been deprecated.") FXAudio2SoundBuffer : public FSoundBuffer
 {
 public:
 	/** 
@@ -374,7 +374,7 @@ public:
 /**
  * Source callback class for handling loops
  */
-class FXAudio2SoundSourceCallback : public IXAudio2VoiceCallback
+class UE_DEPRECATED(4.26, "XAudio2 Module has been deprecated.") FXAudio2SoundSourceCallback : public IXAudio2VoiceCallback
 {
 public:
 	FXAudio2SoundSourceCallback() {}
@@ -402,7 +402,7 @@ public:
 /**
  * XAudio2 implementation of FSoundSource, the interface used to play, stop and update sources
  */
-class FXAudio2SoundSource : public FSoundSource
+class  UE_DEPRECATED(4.26, "XAudio2 Module has been deprecated.") FXAudio2SoundSource : public FSoundSource
 {
 public:
 	/**
@@ -672,8 +672,11 @@ protected:
 /**
  * Helper class for 5.1 spatialization.
  */
-class FSpatializationHelper
+
+class UE_DEPRECATED(4.26, "XAudio2 Module has been deprecated.") FSpatializationHelper
 {
+#if WITH_XAUDIO3D
+
 	/** Instance of X3D used to calculate volume multipliers.	*/
 	X3DAUDIO_HANDLE		          X3DInstance;
 	
@@ -688,6 +691,8 @@ class FSpatializationHelper
 	X3DAUDIO_DISTANCE_CURVE_POINT ReverbVolumeCurvePoint[2];
 	X3DAUDIO_DISTANCE_CURVE       ReverbVolumeCurve;
 
+#endif //WITH_XAUDIO3D
+
 	float                         EmitterAzimuths[UE4_XAUDIO3D_INPUTCHANNELS];
 
 	// TODO: Hardcoding this to 8 because X3DAudioCalculate is ignoring the destination speaker count we put in and
@@ -722,7 +727,7 @@ public:
 };
 
 /** A pool entry for related IXAudio2SourceVoices */
-struct FSourceVoicePoolEntry
+struct UE_DEPRECATED(4.26, "XAudio2 Module has been deprecated.") FSourceVoicePoolEntry
 {
 	/** The format for all voices in this entry */
 	WAVEFORMATEX Format;
@@ -752,7 +757,7 @@ FORCEINLINE bool operator==(const WAVEFORMATEX& FormatA, const WAVEFORMATEX& For
 
 
 /** This structure holds any singleton XAudio2 resources which need to be used, not just "properties" of the device. */
-struct FXAudioDeviceProperties final : public IDeviceChangedListener
+struct UE_DEPRECATED(4.26, "XAudio2 Module has been deprecated.") FXAudioDeviceProperties final : public IDeviceChangedListener
 {
 	// These variables are non-static to support multiple audio device instances
 	struct IXAudio2*					XAudio2;
diff --git a/Engine/Source/Runtime/Windows/XAudio2/Public/XAudio2Device.h b/Engine/Source/Runtime/Windows/XAudio2/Public/XAudio2Device.h
--- a/Engine/Source/Runtime/Windows/XAudio2/Public/XAudio2Device.h
+++ b/Engine/Source/Runtime/Windows/XAudio2/Public/XAudio2Device.h
@@ -86,7 +86,7 @@ class FSpatializationHelper;
 /**
  * XAudio2 implementation of an Unreal audio device. Neither use XACT nor X3DAudio.
  */
-class FXAudio2Device : public FAudioDevice
+class UE_DEPRECATED(4.26, "XAudio2 Module has been deprecated.") FXAudio2Device : public FAudioDevice
 {
 public:
 	FXAudio2Device() : FAudioDevice()
@@ -221,7 +221,7 @@ private:
 #endif // PLATFORM_WINDOWS || PLATFORM_HOLOLENS
 };
 
-class FXMPHelper
+class UE_DEPRECATED(4.26, "XAudio2 Module has been deprecated.") FXMPHelper
 {
 private:
 	/** Count of current cinematic audio clips playing (used to turn on/off XMP background music, allowing for overlap) */
diff --git a/Engine/Source/Runtime/Windows/XAudio2/Public/XAudio2Effects.h b/Engine/Source/Runtime/Windows/XAudio2/Public/XAudio2Effects.h
--- a/Engine/Source/Runtime/Windows/XAudio2/Public/XAudio2Effects.h
+++ b/Engine/Source/Runtime/Windows/XAudio2/Public/XAudio2Effects.h
@@ -12,7 +12,7 @@
 /** 
  * XAudio2 effects manager
  */
-class FXAudio2EffectsManager : public FAudioEffectsManager
+class UE_DEPRECATED(4.26, "XAudio2 Module has been deprecated.") FXAudio2EffectsManager : public FAudioEffectsManager
 {
 public:
 	FXAudio2EffectsManager( FXAudio2Device* InDevice );
diff --git a/Engine/Source/Runtime/Windows/XAudio2/XAudio2.Build.cs b/Engine/Source/Runtime/Windows/XAudio2/XAudio2.Build.cs
--- a/Engine/Source/Runtime/Windows/XAudio2/XAudio2.Build.cs
+++ b/Engine/Source/Runtime/Windows/XAudio2/XAudio2.Build.cs
@@ -42,6 +42,7 @@ public class XAudio2 : ModuleRules
 			"VorbisFile"
 			);
 
-		PublicDefinitions.Add("WITH_OGGVORBIS=1");
+		PrivateDefinitions.Add("WITH_XAUDIO3D=0");
+        PublicDefinitions.Add("WITH_OGGVORBIS=1");
 	}
 }
