Force gpucrashdebugging when in editor.
use -nocrashdebugging to disable.

diff --git a/Engine/Source/Runtime/Windows/D3D11RHI/Private/Windows/WindowsD3D11Device.cpp b/Engine/Source/Runtime/Windows/D3D11RHI/Private/Windows/WindowsD3D11Device.cpp
--- a/Engine/Source/Runtime/Windows/D3D11RHI/Private/Windows/WindowsD3D11Device.cpp
+++ b/Engine/Source/Runtime/Windows/D3D11RHI/Private/Windows/WindowsD3D11Device.cpp
@@ -1131,9 +1131,10 @@ static void CacheNVAftermathEnabled()
 {
 	if (GNVAftermathModuleLoaded && IsRHIDeviceNVIDIA())
 	{
+		const bool bEnableInEditor = GIsEditor && !FParse::Param(FCommandLine::Get(), TEXT("nogpucrashdebugging"));
 		// Two ways to enable aftermath, command line or the r.GPUCrashDebugging variable
 		// Note: If intending to change this please alert game teams who use this for user support.
-		if (FParse::Param(FCommandLine::Get(), TEXT("gpucrashdebugging")))
+		if (bEnableInEditor || FParse::Param(FCommandLine::Get(), TEXT("gpucrashdebugging")))
 		{
 			GDX11NVAfterMathEnabled = true;
 		}
@@ -1166,8 +1167,9 @@ void FD3D11DynamicRHI::StartNVAftermath()
 		static IConsoleVariable* CallstackCVar = IConsoleManager::Get().FindConsoleVariable(TEXT("r.GPUCrashDebugging.Aftermath.Callstack"));
 		static IConsoleVariable* ResourcesCVar = IConsoleManager::Get().FindConsoleVariable(TEXT("r.GPUCrashDebugging.Aftermath.ResourceTracking"));
 		static IConsoleVariable* TrackAllCVar = IConsoleManager::Get().FindConsoleVariable(TEXT("r.GPUCrashDebugging.Aftermath.TrackAll"));
-
-		const bool bEnableMarkers = FParse::Param(FCommandLine::Get(), TEXT("aftermathmarkers")) || (MarkersCVar && MarkersCVar->GetInt());
+		
+		const bool bEnableInEditor = GIsEditor && !FParse::Param(FCommandLine::Get(), TEXT("nogpucrashdebugging"));
+		const bool bEnableMarkers = FParse::Param(FCommandLine::Get(), TEXT("aftermathmarkers")) || (MarkersCVar && MarkersCVar->GetInt()) || bEnableInEditor;
 		const bool bEnableCallstack = FParse::Param(FCommandLine::Get(), TEXT("aftermathcallstack")) || (CallstackCVar && CallstackCVar->GetInt());
 		const bool bEnableResources = FParse::Param(FCommandLine::Get(), TEXT("aftermathresources")) || (ResourcesCVar && ResourcesCVar->GetInt());
 		const bool bEnableAll = FParse::Param(FCommandLine::Get(), TEXT("aftermathall")) || (TrackAllCVar && TrackAllCVar->GetInt());
