F - Copy from 12877954
D3D11 - Fix for AMD-based CPUs with built-in graphics. IsRHIDeviceAMD might be integrated grapihcs or a discrete GPU. Just get the budget from the OS via IDXGIAdapter3. Fall back to the original calculation method if the interface for IDXGIAdapter3 does not exist (e.g. Win7).

#rnx

diff --git a/Engine/Source/Runtime/Windows/D3D11RHI/Private/Windows/WindowsD3D11Device.cpp b/Engine/Source/Runtime/Windows/D3D11RHI/Private/Windows/WindowsD3D11Device.cpp
--- a/Engine/Source/Runtime/Windows/D3D11RHI/Private/Windows/WindowsD3D11Device.cpp
+++ b/Engine/Source/Runtime/Windows/D3D11RHI/Private/Windows/WindowsD3D11Device.cpp
@@ -1672,6 +1672,7 @@ void FD3D11DynamicRHI::InitD3DDevice()
 				// Consider 50% of the shared memory but max 25% of total system memory.
 				int64 ConsideredSharedSystemMemory = FMath::Min( FD3D11GlobalStats::GSharedSystemMemory / 2ll, TotalPhysicalMemory / 4ll );
 
+				TRefCountPtr<IDXGIAdapter3> DxgiAdapter3;
 				FD3D11GlobalStats::GTotalGraphicsMemory = 0;
 				if ( IsRHIDeviceIntel() )
 				{
@@ -1680,6 +1681,14 @@ void FD3D11DynamicRHI::InitD3DDevice()
 					FD3D11GlobalStats::GTotalGraphicsMemory += FD3D11GlobalStats::GDedicatedSystemMemory;
 					FD3D11GlobalStats::GTotalGraphicsMemory += ConsideredSharedSystemMemory;
 				}
+				else if (IsRHIDeviceAMD() && SUCCEEDED(Adapter->QueryInterface(_uuidof(IDXGIAdapter3), (void**)DxgiAdapter3.GetInitReference())) && DxgiAdapter3.IsValid())
+				{
+					DXGI_QUERY_VIDEO_MEMORY_INFO LocalVideoMemoryInfo;
+					VERIFYD3D11RESULT(DxgiAdapter3->QueryVideoMemoryInfo(0, DXGI_MEMORY_SEGMENT_GROUP_LOCAL, &LocalVideoMemoryInfo));
+					// use the entire budget for D3D11, in keeping with setting GTotalGraphicsMemory to all of AdapterDesc.DedicatedVideoMemory
+					// in the other method directly below
+					FD3D11GlobalStats::GTotalGraphicsMemory = LocalVideoMemoryInfo.Budget;
+				}
 				else if ( FD3D11GlobalStats::GDedicatedVideoMemory >= 200*1024*1024 )
 				{
 					// Use dedicated video memory, if it's more than 200 MB
