Fixed crash on audio device hot swap when channels have changed.

diff --git a/Engine/Source/Runtime/Windows/AudioMixerXAudio2/Private/AudioMixerPlatformXAudio2.cpp b/Engine/Source/Runtime/Windows/AudioMixerXAudio2/Private/AudioMixerPlatformXAudio2.cpp
--- a/Engine/Source/Runtime/Windows/AudioMixerXAudio2/Private/AudioMixerPlatformXAudio2.cpp
+++ b/Engine/Source/Runtime/Windows/AudioMixerXAudio2/Private/AudioMixerPlatformXAudio2.cpp
@@ -1199,6 +1199,10 @@ namespace Audio
 			// Create the output source voice
 			XAUDIO2_RETURN_ON_FAIL(XAudio2System->CreateSourceVoice(&OutputAudioStreamSourceVoice, &Format, XAUDIO2_VOICE_NOPITCH, 2.0f, &OutputVoiceCallback));
 
+			// Reinitialize the output circular buffer to match the buffer math of the new audio device.
+			const int32 NumOutputSamples = AudioStreamInfo.NumOutputFrames * AudioStreamInfo.DeviceInfo.NumChannels;
+			OutputBuffer.Init(AudioStreamInfo.AudioMixer, NumOutputSamples, NumOutputBuffers, AudioStreamInfo.DeviceInfo.Format);
+
 			const int32 NewNumSamples = OpenStreamParams.NumFrames * AudioStreamInfo.DeviceInfo.NumChannels;
 		}
 		else
