UE4.25 - Back out 13044830
#rnx

diff --git a/Engine/Source/Runtime/Windows/D3D11RHI/Private/Windows/WindowsD3D11Device.cpp b/Engine/Source/Runtime/Windows/D3D11RHI/Private/Windows/WindowsD3D11Device.cpp
--- a/Engine/Source/Runtime/Windows/D3D11RHI/Private/Windows/WindowsD3D11Device.cpp
+++ b/Engine/Source/Runtime/Windows/D3D11RHI/Private/Windows/WindowsD3D11Device.cpp
@@ -22,8 +22,6 @@
 #include "GenericPlatform/GenericPlatformDriver.h"			// FGPUDriverInfo
 #include "GenericPlatform/GenericPlatformCrashContext.h"
 THIRD_PARTY_INCLUDES_START
-#include "dxgi1_3.h"
-#include "dxgi1_4.h"
 #include "dxgi1_6.h"
 THIRD_PARTY_INCLUDES_END
 #include "RHIValidation.h"
@@ -1674,7 +1672,6 @@ void FD3D11DynamicRHI::InitD3DDevice()
 				// Consider 50% of the shared memory but max 25% of total system memory.
 				int64 ConsideredSharedSystemMemory = FMath::Min( FD3D11GlobalStats::GSharedSystemMemory / 2ll, TotalPhysicalMemory / 4ll );
 
-				TRefCountPtr<IDXGIAdapter3> DxgiAdapter3;
 				FD3D11GlobalStats::GTotalGraphicsMemory = 0;
 				if ( IsRHIDeviceIntel() )
 				{
@@ -1683,14 +1680,6 @@ void FD3D11DynamicRHI::InitD3DDevice()
 					FD3D11GlobalStats::GTotalGraphicsMemory += FD3D11GlobalStats::GDedicatedSystemMemory;
 					FD3D11GlobalStats::GTotalGraphicsMemory += ConsideredSharedSystemMemory;
 				}
-				else if (IsRHIDeviceAMD() && SUCCEEDED(Adapter->QueryInterface(_uuidof(IDXGIAdapter3), (void**)DxgiAdapter3.GetInitReference())) && DxgiAdapter3.IsValid())
-				{
-					DXGI_QUERY_VIDEO_MEMORY_INFO LocalVideoMemoryInfo;
-					VERIFYD3D11RESULT(DxgiAdapter3->QueryVideoMemoryInfo(0, DXGI_MEMORY_SEGMENT_GROUP_LOCAL, &LocalVideoMemoryInfo));
-					// use the entire budget for D3D11, in keeping with setting GTotalGraphicsMemory to all of AdapterDesc.DedicatedVideoMemory
-					// in the other method directly below
-					FD3D11GlobalStats::GTotalGraphicsMemory = LocalVideoMemoryInfo.Budget;
-				}
 				else if ( FD3D11GlobalStats::GDedicatedVideoMemory >= 200*1024*1024 )
 				{
 					// Use dedicated video memory, if it's more than 200 MB
