Fixes PF_R8 on d3d11

diff --git a/Engine/Source/Runtime/Windows/D3D11RHI/Private/D3D11Device.cpp b/Engine/Source/Runtime/Windows/D3D11RHI/Private/D3D11Device.cpp
--- a/Engine/Source/Runtime/Windows/D3D11RHI/Private/D3D11Device.cpp
+++ b/Engine/Source/Runtime/Windows/D3D11RHI/Private/D3D11Device.cpp
@@ -199,8 +199,8 @@ FD3D11DynamicRHI::FD3D11DynamicRHI(IDXGIFactory1* InDXGIFactory1,D3D_FEATURE_LEV
 
 	GPixelFormats[ PF_BC6H			].PlatformFormat	= DXGI_FORMAT_BC6H_UF16;
 	GPixelFormats[ PF_BC7			].PlatformFormat	= DXGI_FORMAT_BC7_TYPELESS;
-	GPixelFormats[ PF_R8_UINT		].PlatformFormat = DXGI_FORMAT_R8_UINT;
-	GPixelFormats[ PF_R8			].PlatformFormat = DXGI_FORMAT_A8_UNORM;
+	GPixelFormats[ PF_R8_UINT		].PlatformFormat	= DXGI_FORMAT_R8_UINT;
+	GPixelFormats[ PF_R8			].PlatformFormat	= DXGI_FORMAT_R8_UNORM;
 
 	GPixelFormats[PF_R16G16B16A16_UNORM].PlatformFormat = DXGI_FORMAT_R16G16B16A16_UNORM;
 	GPixelFormats[PF_R16G16B16A16_SNORM].PlatformFormat = DXGI_FORMAT_R16G16B16A16_SNORM;
