Fix how we handled default affinity mask with xaudio2 which caused too many time critical threads to be created and causing lots of context switch
  - This eliminates 31 time critical threads on AMD TR 3970X.

#rnx

diff --git a/Engine/Source/Runtime/Windows/AudioMixerXAudio2/Private/AudioMixerPlatformXAudio2.cpp b/Engine/Source/Runtime/Windows/AudioMixerXAudio2/Private/AudioMixerPlatformXAudio2.cpp
--- a/Engine/Source/Runtime/Windows/AudioMixerXAudio2/Private/AudioMixerPlatformXAudio2.cpp
+++ b/Engine/Source/Runtime/Windows/AudioMixerXAudio2/Private/AudioMixerPlatformXAudio2.cpp
@@ -86,6 +86,26 @@ static FString GetErrorString(HRESULT Result)
 		}
 	}
 }
+
+static XAUDIO2_PROCESSOR GetXAudio2ProcessorsToUse()
+{
+	XAUDIO2_PROCESSOR ProcessorsToUse = (XAUDIO2_PROCESSOR)FPlatformAffinity::GetAudioThreadMask();
+	// https://docs.microsoft.com/en-us/windows/win32/api/xaudio2/nf-xaudio2-xaudio2create
+	// Warning If you specify XAUDIO2_ANY_PROCESSOR, the system will use all of the device's processors and, as noted above, create a worker thread for each processor.
+	// We certainly don't want to use all available CPU. XAudio threads are time critical priority and wake up every 10 ms, they may cause lots of unwarranted context switches.
+	// In case no specific affinity is specified, let XAudio choose the default processor. It should allocate a single thread and should be enough.
+	if (ProcessorsToUse == XAUDIO2_ANY_PROCESSOR)
+	{
+	#ifdef XAUDIO2_USE_DEFAULT_PROCESSOR
+		ProcessorsToUse = XAUDIO2_USE_DEFAULT_PROCESSOR;
+	#else
+		ProcessorsToUse = XAUDIO2_DEFAULT_PROCESSOR;
+	#endif
+	}
+
+	return ProcessorsToUse;
+}
+
 #if PLATFORM_WINDOWS || PLATFORM_HOLOLENS
 FName GetNextDllToTry(FName Current = NAME_None) 
 {
@@ -212,7 +232,7 @@ namespace Audio
 		Flags |= XAUDIO2_DO_NOT_USE_SHAPE;
 #endif
 
-		if (FAILED(XAudio2Create(&XAudio2System, Flags, (XAUDIO2_PROCESSOR)FPlatformAffinity::GetAudioThreadMask())))
+		if (FAILED(XAudio2Create(&XAudio2System, Flags, GetXAudio2ProcessorsToUse())))
 		{
 			XAudio2System = nullptr;
 			return false;
@@ -255,7 +275,7 @@ namespace Audio
 		Flags |= XAUDIO2_DO_NOT_USE_SHAPE;
 #endif
 
-		if (!XAudio2System && FAILED(XAudio2Create(&XAudio2System, Flags, (XAUDIO2_PROCESSOR)FPlatformAffinity::GetAudioThreadMask())))
+		if (!XAudio2System && FAILED(XAudio2Create(&XAudio2System, Flags, GetXAudio2ProcessorsToUse())))
 		{
 			FMessageDialog::Open(EAppMsgType::Ok, NSLOCTEXT("Audio", "XAudio2Error", "Failed to initialize audio. This may be an issue with your installation of XAudio 2.7. XAudio2 is available in the DirectX End-User Runtime (June 2010)."));
 			return false;
diff --git a/Engine/Source/Runtime/Windows/XAudio2/Private/XAudio2Device.cpp b/Engine/Source/Runtime/Windows/XAudio2/Private/XAudio2Device.cpp
--- a/Engine/Source/Runtime/Windows/XAudio2/Private/XAudio2Device.cpp
+++ b/Engine/Source/Runtime/Windows/XAudio2/Private/XAudio2Device.cpp
@@ -155,8 +155,22 @@ bool FXAudio2Device::InitializeHardware()
 	Flags |= XAUDIO2_DO_NOT_USE_SHAPE;
 #endif
 
+	XAUDIO2_PROCESSOR ProcessorsToUse = (XAUDIO2_PROCESSOR)FPlatformAffinity::GetAudioThreadMask();
+	// https://docs.microsoft.com/en-us/windows/win32/api/xaudio2/nf-xaudio2-xaudio2create
+	// Warning If you specify XAUDIO2_ANY_PROCESSOR, the system will use all of the device's processors and, as noted above, create a worker thread for each processor.
+	// We certainly don't want to use all available CPU. XAudio threads are time critical priority and wake up every 10 ms, they may cause lots of unwarranted context switches.
+	// In case no specific affinity is specified, let XAudio choose the default processor. It should allocate a single thread and should be enough.
+	if (ProcessorsToUse == XAUDIO2_ANY_PROCESSOR)
+	{
+#ifdef XAUDIO2_USE_DEFAULT_PROCESSOR
+		ProcessorsToUse = XAUDIO2_USE_DEFAULT_PROCESSOR;
+#else
+		ProcessorsToUse = XAUDIO2_DEFAULT_PROCESSOR;
+#endif
+	}
+
 	// Create a new XAudio2 device object instance
-	if (XAudio2Create(&DeviceProperties->XAudio2, Flags, (XAUDIO2_PROCESSOR)FPlatformAffinity::GetAudioThreadMask()) != S_OK)
+	if (XAudio2Create(&DeviceProperties->XAudio2, Flags, ProcessorsToUse) != S_OK)
 	{
 		UE_LOG(LogInit, Log, TEXT("Failed to create XAudio2 interface"));
 		return(false);
