M2 - DX11 shader creation errors due to Extension and DXIL shaders passed into DX11 when r.CreateShadersOnLoad=1 is used

diff --git a/Engine/Shaders/Public/ShaderVersion.ush b/Engine/Shaders/Public/ShaderVersion.ush
--- a/Engine/Shaders/Public/ShaderVersion.ush
+++ b/Engine/Shaders/Public/ShaderVersion.ush
@@ -3,4 +3,4 @@
 // in Platform.ush (which should be included in any shader) it allows to invalidate the shader DDC.
 // 
 // If you are merging streams and there is a conflict with this GUID you should make a new GUID rather than taking one or the other.
-// GUID = 7CE7A4C3-6DB7-4C94-A510-BEC64474ECDA
+// GUID = 7CE7A4C3-E63C-4B2B-96BF-4A424261A5E3
diff --git a/Engine/Source/Developer/Windows/ShaderFormatD3D/Private/D3DShaderCompilerDXC.cpp b/Engine/Source/Developer/Windows/ShaderFormatD3D/Private/D3DShaderCompilerDXC.cpp
--- a/Engine/Source/Developer/Windows/ShaderFormatD3D/Private/D3DShaderCompilerDXC.cpp
+++ b/Engine/Source/Developer/Windows/ShaderFormatD3D/Private/D3DShaderCompilerDXC.cpp
@@ -885,6 +885,8 @@ bool CompileAndProcessD3DShaderDXC(FString& PreprocessedShaderSource,
 
 				// We only need this to appear when using a DXC shader
 				ShaderCode.AddOptionalData<FShaderCodeFeatures>(CodeFeatures);
+				uint8 IsSM6 = 1;
+				ShaderCode.AddOptionalData('6', &IsSM6, 1);
 			};
 
 			//#todo-rco: Should compress ShaderCode?
diff --git a/Engine/Source/Runtime/Windows/D3D11RHI/Private/D3D11Shaders.cpp b/Engine/Source/Runtime/Windows/D3D11RHI/Private/D3D11Shaders.cpp
--- a/Engine/Source/Runtime/Windows/D3D11RHI/Private/D3D11Shaders.cpp
+++ b/Engine/Source/Runtime/Windows/D3D11RHI/Private/D3D11Shaders.cpp
@@ -53,16 +53,27 @@ static inline void ReadShaderOptionalData(FShaderCodeReader& InShaderCode, TShad
 		Ar << OutShader.VendorExtensions;
 	}
 	OutShader.bShaderNeedsGlobalConstantBuffer = PackedResourceCounts->bGlobalUniformBufferUsed;
+
+	int32 IsSm6ShaderSize = 1;
+	const uint8* IsSm6Shader = InShaderCode.FindOptionalData('6', IsSm6ShaderSize);
+	OutShader.bIsSm6Shader = IsSm6Shader && IsSm6ShaderSize && *IsSm6Shader;
 }
 
-static void ApplyVendorExtensions(ID3D11Device* Direct3DDevice, EShaderFrequency Frequency, const TArray<FShaderCodeVendorExtension>& VendorExtensions)
+static bool ApplyVendorExtensions(ID3D11Device* Direct3DDevice, EShaderFrequency Frequency, const TArray<FShaderCodeVendorExtension>& VendorExtensions)
 {
+	bool IsValidHardwareExtension = true;
 #if !PLATFORM_HOLOLENS
 	for (int32 ExtensionIndex = 0; ExtensionIndex < VendorExtensions.Num(); ++ExtensionIndex)
 	{
 		const FShaderCodeVendorExtension& Extension = VendorExtensions[ExtensionIndex];
 		if (Extension.VendorId == 0x10DE) // NVIDIA
 		{
+			if (!IsRHIDeviceNVIDIA())
+			{
+				IsValidHardwareExtension = false;
+				break;
+			}
+
 			// https://developer.nvidia.com/unlocking-gpu-intrinsics-hlsl
 			if (Extension.Parameter.Type == EShaderParameterType::UAV)
 			{
@@ -71,14 +82,25 @@ static void ApplyVendorExtensions(ID3D11Device* Direct3DDevice, EShaderFrequency
 		}
 		else if (Extension.VendorId == 0x1002) // AMD
 		{
+			if (!IsRHIDeviceAMD())
+			{
+				IsValidHardwareExtension = false;
+				break;
+			}
 			// TODO: [REDACTED_PROFILE_URL]/AGS_SDK/blob/master/ags_lib/hlsl/ags_shader_intrinsics_dx11.hlsl
 		}
 		else if (Extension.VendorId == 0x8086) // Intel
 		{
+			if (!IsRHIDeviceIntel())
+			{
+				IsValidHardwareExtension = false;
+				break;
+			}
 			// TODO: [REDACTED_PROFILE_URL]/intel-graphics-compiler/blob/master/inc/IntelExtensions.hlsl
 		}
 	}
 #endif
+	return IsValidHardwareExtension;
 }
 
 static void ResetVendorExtensions(ID3D11Device* Direct3DDevice, EShaderFrequency Frequency, const TArray<FShaderCodeVendorExtension>& VendorExtensions)
@@ -138,10 +160,12 @@ FVertexShaderRHIRef FD3D11DynamicRHI::RHICreateVertexShader(TArrayView<const uin
 	const size_t CodeSize = ShaderCode.GetActualShaderCodeSize() - Offset;
 
 	ReadShaderOptionalData(ShaderCode, *Shader);
-	ApplyVendorExtensions(Direct3DDevice, SF_Vertex, Shader->VendorExtensions);
-	VERIFYD3D11SHADERRESULT(Direct3DDevice->CreateVertexShader((void*)CodePtr, CodeSize, nullptr, Shader->Resource.GetInitReference()), Shader, Direct3DDevice);
-	ResetVendorExtensions(Direct3DDevice, SF_Vertex, Shader->VendorExtensions);
-	InitUniformBufferStaticSlots(Shader);
+	if (!Shader->bIsSm6Shader && ApplyVendorExtensions(Direct3DDevice, SF_Vertex, Shader->VendorExtensions))
+	{
+		VERIFYD3D11SHADERRESULT(Direct3DDevice->CreateVertexShader((void*)CodePtr, CodeSize, nullptr, Shader->Resource.GetInitReference()), Shader, Direct3DDevice);
+		ResetVendorExtensions(Direct3DDevice, SF_Vertex, Shader->VendorExtensions);
+		InitUniformBufferStaticSlots(Shader);
+	}
 	
 	// TEMP
 	Shader->Code = Code;
@@ -170,10 +194,12 @@ FGeometryShaderRHIRef FD3D11DynamicRHI::RHICreateGeometryShader(TArrayView<const
 	const size_t CodeSize = ShaderCode.GetActualShaderCodeSize() - Offset;
 
 	ReadShaderOptionalData(ShaderCode, *Shader);
-	ApplyVendorExtensions(Direct3DDevice, SF_Geometry, Shader->VendorExtensions);
-	VERIFYD3D11SHADERRESULT(Direct3DDevice->CreateGeometryShader((void*)CodePtr, CodeSize, nullptr, Shader->Resource.GetInitReference()), Shader, Direct3DDevice);
-	ResetVendorExtensions(Direct3DDevice, SF_Geometry, Shader->VendorExtensions);
-	InitUniformBufferStaticSlots(Shader);
+	if (!Shader->bIsSm6Shader && ApplyVendorExtensions(Direct3DDevice, SF_Geometry, Shader->VendorExtensions))
+	{
+		VERIFYD3D11SHADERRESULT(Direct3DDevice->CreateGeometryShader((void*)CodePtr, CodeSize, nullptr, Shader->Resource.GetInitReference()), Shader, Direct3DDevice);
+		ResetVendorExtensions(Direct3DDevice, SF_Geometry, Shader->VendorExtensions);
+		InitUniformBufferStaticSlots(Shader);
+	}
 
 	return Shader;
 }
@@ -198,10 +224,12 @@ FHullShaderRHIRef FD3D11DynamicRHI::RHICreateHullShader(TArrayView<const uint8>
 	const size_t CodeSize = ShaderCode.GetActualShaderCodeSize() - Offset;
 
 	ReadShaderOptionalData(ShaderCode, *Shader);
-	ApplyVendorExtensions(Direct3DDevice, SF_Hull, Shader->VendorExtensions);
-	VERIFYD3D11SHADERRESULT(Direct3DDevice->CreateHullShader((void*)CodePtr, CodeSize, nullptr, Shader->Resource.GetInitReference()), Shader, Direct3DDevice);
-	ResetVendorExtensions(Direct3DDevice, SF_Hull, Shader->VendorExtensions);
-	InitUniformBufferStaticSlots(Shader);
+	if (!Shader->bIsSm6Shader && ApplyVendorExtensions(Direct3DDevice, SF_Hull, Shader->VendorExtensions))
+	{
+		VERIFYD3D11SHADERRESULT(Direct3DDevice->CreateHullShader((void*)CodePtr, CodeSize, nullptr, Shader->Resource.GetInitReference()), Shader, Direct3DDevice);
+		ResetVendorExtensions(Direct3DDevice, SF_Hull, Shader->VendorExtensions);
+		InitUniformBufferStaticSlots(Shader);
+	}
 
 	return Shader;
 }
@@ -226,10 +254,12 @@ FDomainShaderRHIRef FD3D11DynamicRHI::RHICreateDomainShader(TArrayView<const uin
 	const size_t CodeSize = ShaderCode.GetActualShaderCodeSize() - Offset;
 
 	ReadShaderOptionalData(ShaderCode, *Shader);
-	ApplyVendorExtensions(Direct3DDevice, SF_Domain, Shader->VendorExtensions);
-	VERIFYD3D11SHADERRESULT(Direct3DDevice->CreateDomainShader((void*)CodePtr, CodeSize, nullptr, Shader->Resource.GetInitReference()), Shader, Direct3DDevice);
-	ResetVendorExtensions(Direct3DDevice, SF_Domain, Shader->VendorExtensions);
-	InitUniformBufferStaticSlots(Shader);
+	if (!Shader->bIsSm6Shader && ApplyVendorExtensions(Direct3DDevice, SF_Domain, Shader->VendorExtensions))
+	{
+		VERIFYD3D11SHADERRESULT(Direct3DDevice->CreateDomainShader((void*)CodePtr, CodeSize, nullptr, Shader->Resource.GetInitReference()), Shader, Direct3DDevice);
+		ResetVendorExtensions(Direct3DDevice, SF_Domain, Shader->VendorExtensions);
+		InitUniformBufferStaticSlots(Shader);
+	}
 
 	return Shader;
 }
@@ -254,10 +284,12 @@ FPixelShaderRHIRef FD3D11DynamicRHI::RHICreatePixelShader(TArrayView<const uint8
 	const size_t CodeSize = ShaderCode.GetActualShaderCodeSize() - Offset;
 
 	ReadShaderOptionalData(ShaderCode, *Shader);
-	ApplyVendorExtensions(Direct3DDevice, SF_Pixel, Shader->VendorExtensions);
-	VERIFYD3D11SHADERRESULT(Direct3DDevice->CreatePixelShader((void*)CodePtr, CodeSize, nullptr, Shader->Resource.GetInitReference()), Shader, Direct3DDevice);
-	ResetVendorExtensions(Direct3DDevice, SF_Pixel, Shader->VendorExtensions);
-	InitUniformBufferStaticSlots(Shader);
+	if (!Shader->bIsSm6Shader && ApplyVendorExtensions(Direct3DDevice, SF_Pixel, Shader->VendorExtensions))
+	{
+		VERIFYD3D11SHADERRESULT(Direct3DDevice->CreatePixelShader((void*)CodePtr, CodeSize, nullptr, Shader->Resource.GetInitReference()), Shader, Direct3DDevice);
+		ResetVendorExtensions(Direct3DDevice, SF_Pixel, Shader->VendorExtensions);
+		InitUniformBufferStaticSlots(Shader);
+	}
 
 	return Shader;
 }
@@ -282,10 +314,12 @@ FComputeShaderRHIRef FD3D11DynamicRHI::RHICreateComputeShader(TArrayView<const u
 	const size_t CodeSize = ShaderCode.GetActualShaderCodeSize() - Offset;
 
 	ReadShaderOptionalData(ShaderCode, *Shader);
-	ApplyVendorExtensions(Direct3DDevice, SF_Compute, Shader->VendorExtensions);
-	VERIFYD3D11SHADERRESULT(Direct3DDevice->CreateComputeShader((void*)CodePtr, CodeSize, nullptr, Shader->Resource.GetInitReference()), Shader, Direct3DDevice);
-	ResetVendorExtensions(Direct3DDevice, SF_Compute, Shader->VendorExtensions);
-	InitUniformBufferStaticSlots(Shader);
+	if (!Shader->bIsSm6Shader && ApplyVendorExtensions(Direct3DDevice, SF_Compute, Shader->VendorExtensions))
+	{
+		VERIFYD3D11SHADERRESULT(Direct3DDevice->CreateComputeShader((void*)CodePtr, CodeSize, nullptr, Shader->Resource.GetInitReference()), Shader, Direct3DDevice);
+		ResetVendorExtensions(Direct3DDevice, SF_Compute, Shader->VendorExtensions);
+		InitUniformBufferStaticSlots(Shader);
+	}
 
 	return Shader;
 }
diff --git a/Engine/Source/Runtime/Windows/D3D11RHI/Public/D3D11Resources.h b/Engine/Source/Runtime/Windows/D3D11RHI/Public/D3D11Resources.h
--- a/Engine/Source/Runtime/Windows/D3D11RHI/Public/D3D11Resources.h
+++ b/Engine/Source/Runtime/Windows/D3D11RHI/Public/D3D11Resources.h
@@ -45,6 +45,7 @@ struct FD3D11ShaderData
 	TArray<FUniformBufferStaticSlot>	StaticSlots;
 	TArray<FShaderCodeVendorExtension>	VendorExtensions;
 	bool								bShaderNeedsGlobalConstantBuffer;
+	bool								bIsSm6Shader;
 	uint16								OutputMask;
 	uint16								UAVMask;
 };
