Fixing DXGI failure when using the old swap effect, scaling needs to be set to 'stretch' and not 'none' in these cases.

diff --git a/Engine/Source/Runtime/Windows/D3D11RHI/Private/Windows/WindowsD3D11Viewport.cpp b/Engine/Source/Runtime/Windows/D3D11RHI/Private/Windows/WindowsD3D11Viewport.cpp
--- a/Engine/Source/Runtime/Windows/D3D11RHI/Private/Windows/WindowsD3D11Viewport.cpp
+++ b/Engine/Source/Runtime/Windows/D3D11RHI/Private/Windows/WindowsD3D11Viewport.cpp
@@ -20,6 +20,7 @@ THIRD_PARTY_INCLUDES_END
 #endif
 
 static DXGI_SWAP_EFFECT GSwapEffect = DXGI_SWAP_EFFECT_DISCARD;
+static DXGI_SCALING GSwapScaling = DXGI_SCALING_STRETCH;
 static uint32 GSwapChainFlags = DXGI_SWAP_CHAIN_FLAG_ALLOW_MODE_SWITCH;
 static uint32 GSwapChainBufferCount = 1;
 
@@ -67,6 +68,7 @@ FD3D11Viewport::FD3D11Viewport(FD3D11DynamicRHI* InD3DRHI,HWND InWindowHandle,ui
 
 	{
 		GSwapEffect = DXGI_SWAP_EFFECT_DISCARD;
+		GSwapScaling = DXGI_SCALING_STRETCH;
 		GSwapChainFlags = DXGI_SWAP_CHAIN_FLAG_ALLOW_MODE_SWITCH;
 		IDXGIFactory1* Factory1 = D3DRHI->GetFactory();
 		TRefCountPtr<IDXGIFactory5> Factory5;
@@ -79,6 +81,7 @@ FD3D11Viewport::FD3D11Viewport(FD3D11DynamicRHI* InD3DRHI,HWND InWindowHandle,ui
 				if (S_OK == Factory5->CheckFeatureSupport(DXGI_FEATURE_PRESENT_ALLOW_TEARING, &AllowTearing, sizeof(UINT)) && AllowTearing != 0)
 				{
 					GSwapEffect = DXGI_SWAP_EFFECT_FLIP_DISCARD;
+					GSwapScaling = DXGI_SCALING_NONE;
 					GSwapChainFlags |= DXGI_SWAP_CHAIN_FLAG_ALLOW_TEARING;
 					GSwapChainBufferCount = 2;
 				}
@@ -174,7 +177,7 @@ FD3D11Viewport::FD3D11Viewport(FD3D11DynamicRHI* InD3DRHI,HWND InWindowHandle,ui
 			SwapChainDesc.SwapEffect = GSwapEffect;
 			SwapChainDesc.BufferCount = BackBufferCount;
 			SwapChainDesc.Flags = GSwapChainFlags;
-			SwapChainDesc.Scaling = DXGI_SCALING_NONE;
+			SwapChainDesc.Scaling = GSwapScaling;
 
 			IDXGISwapChain1* SwapChain1 = nullptr;
 			IDXGIFactory2* Factory2 = (IDXGIFactory2*)D3DRHI->GetFactory();
