Moving D3D11 Error/FatalError text around such that the error code is on the first line.
This is primarily for crash reporters that can't consume multi-line errors. Changes errors caught as "CreateSwapChainResult failed" to something like "CreateSwapChainResult failed with error E_INVALIDARG".
Added DXGI_ERROR_NOT_CURRENTLY_AVAILABLE to the D3D11 error strings since it's something worth calling out.

diff --git a/Engine/Source/Runtime/Windows/D3D11RHI/Private/D3D11Util.cpp b/Engine/Source/Runtime/Windows/D3D11RHI/Private/D3D11Util.cpp
--- a/Engine/Source/Runtime/Windows/D3D11RHI/Private/D3D11Util.cpp
+++ b/Engine/Source/Runtime/Windows/D3D11RHI/Private/D3D11Util.cpp
@@ -64,6 +64,7 @@ FString GetD3D11ErrorString(HRESULT ErrorCode, ID3D11Device* Device)
 		D3DERR(DXGI_ERROR_INVALID_CALL)
 		D3DERR(E_NOINTERFACE)
 		D3DERR(DXGI_ERROR_DEVICE_REMOVED)
+		D3DERR(DXGI_ERROR_NOT_CURRENTLY_AVAILABLE)
 		default: ErrorCodeText = FString::Printf(TEXT("%08X"),(int32)ErrorCode);
 	}
 
@@ -241,7 +242,7 @@ void VerifyD3D11ResultNoExit(HRESULT D3DResult, const ANSICHAR* Code, const ANSI
 
 	const FString& ErrorString = GetD3D11ErrorString(D3DResult, Device);
 
-	UE_LOG(LogD3D11RHI, Error, TEXT("%s failed \n at %s:%u \n with error %s Error Code List: https://docs.microsoft.com/en-us/windows/desktop/direct3ddxgi/dxgi-error"), ANSI_TO_TCHAR(Code), ANSI_TO_TCHAR(Filename), Line, *ErrorString);
+	UE_LOG(LogD3D11RHI, Error, TEXT("%s failed with error %s\n at %s:%u\n Error Code List: https://docs.microsoft.com/en-us/windows/desktop/direct3ddxgi/dxgi-error"), ANSI_TO_TCHAR(Code), *ErrorString, ANSI_TO_TCHAR(Filename), Line);
 }
 
 void VerifyD3D11Result(HRESULT D3DResult,const ANSICHAR* Code,const ANSICHAR* Filename,uint32 Line, ID3D11Device* Device)
@@ -250,12 +251,12 @@ void VerifyD3D11Result(HRESULT D3DResult,const ANSICHAR* Code,const ANSICHAR* Fi
 
 	const FString& ErrorString = GetD3D11ErrorString(D3DResult, Device);
 
-	UE_LOG(LogD3D11RHI, Error, TEXT("%s failed \n at %s:%u \n with error %s"), ANSI_TO_TCHAR(Code), ANSI_TO_TCHAR(Filename), Line, *ErrorString);
+	UE_LOG(LogD3D11RHI, Error, TEXT("%s failed with error %s\n at %s:%u"), ANSI_TO_TCHAR(Code), *ErrorString, ANSI_TO_TCHAR(Filename), Line);
 
 	TerminateOnDeviceRemoved(D3DResult, Device);
 	TerminateOnOutOfMemory(D3DResult, false);
 
-	UE_LOG(LogD3D11RHI, Fatal,TEXT("%s failed \n at %s:%u \n with error %s"),ANSI_TO_TCHAR(Code),ANSI_TO_TCHAR(Filename),Line,*ErrorString);
+	UE_LOG(LogD3D11RHI, Fatal,TEXT("%s failed with error %s\n at %s:%u"),ANSI_TO_TCHAR(Code), *ErrorString, ANSI_TO_TCHAR(Filename), Line);
 }
 
 void VerifyD3D11ShaderResult(FRHIShader* Shader, HRESULT D3DResult, const ANSICHAR* Code, const ANSICHAR* Filename, uint32 Line, ID3D11Device* Device)
@@ -267,11 +268,11 @@ void VerifyD3D11ShaderResult(FRHIShader* Shader, HRESULT D3DResult, const ANSICH
 #if !(UE_BUILD_SHIPPING || UE_BUILD_TEST)
 	if (Shader->ShaderName.Len())
 	{
-		UE_LOG(LogD3D11RHI, Error, TEXT("%s failed trying to create shader %s\n at %s:%u \n with error %s"), ANSI_TO_TCHAR(Code), *Shader->ShaderName, ANSI_TO_TCHAR(Filename), Line, *ErrorString);
+		UE_LOG(LogD3D11RHI, Error, TEXT("%s failed trying to create shader '%s' with error %s\n at %s:%u"), ANSI_TO_TCHAR(Code), *Shader->ShaderName, *ErrorString, ANSI_TO_TCHAR(Filename), Line);
 		TerminateOnDeviceRemoved(D3DResult, Device);
 		TerminateOnOutOfMemory(D3DResult, false);
 
-		UE_LOG(LogD3D11RHI, Fatal, TEXT("%s failed trying to create shader %s \n at %s:%u \n with error %s"), ANSI_TO_TCHAR(Code), *Shader->ShaderName, ANSI_TO_TCHAR(Filename), Line, *ErrorString);
+		UE_LOG(LogD3D11RHI, Fatal, TEXT("%s failed trying to create shader '%s' with error %s\n at %s:%u"), ANSI_TO_TCHAR(Code), *Shader->ShaderName, *ErrorString, ANSI_TO_TCHAR(Filename), Line);
 	}
 	else
 #endif
@@ -290,11 +291,11 @@ void VerifyD3D11CreateTextureResult(HRESULT D3DResult, int32 UEFormat,const ANSI
 	const TCHAR* D3DFormatString = GetD3D11TextureFormatString((DXGI_FORMAT)D3DFormat);
 
 	UE_LOG(LogD3D11RHI, Error,
-		TEXT("%s failed \n at %s:%u \n with error %s, \n Size=%ix%ix%i PF=%d D3DFormat=%s(0x%08X), NumMips=%i, Flags=%s, Usage:0x%x, CPUFlags:0x%x, MiscFlags:0x%x, SampleCount:0x%x, SampleQuality:0x%x, SubresPtr:0x%p, SubresPitch:%i, SubresSlicePitch:%i"),
+		TEXT("%s failed with error %s\n at %s:%u\n Size=%ix%ix%i PF=%d D3DFormat=%s(0x%08X), NumMips=%i, Flags=%s, Usage:0x%x, CPUFlags:0x%x, MiscFlags:0x%x, SampleCount:0x%x, SampleQuality:0x%x, SubresPtr:0x%p, SubresPitch:%i, SubresSlicePitch:%i"),
 		ANSI_TO_TCHAR(Code),
+		*ErrorString,
 		ANSI_TO_TCHAR(Filename),
 		Line,
-		*ErrorString,
 		SizeX,
 		SizeY,
 		SizeZ,
@@ -316,11 +317,11 @@ void VerifyD3D11CreateTextureResult(HRESULT D3DResult, int32 UEFormat,const ANSI
 	TerminateOnOutOfMemory(D3DResult, true);
 
 	UE_LOG(LogD3D11RHI, Fatal,
-		TEXT("%s failed \n at %s:%u \n with error %s, \n Size=%ix%ix%i PF=%d Format=%s(0x%08X), NumMips=%i, Flags=%s, Usage:0x%x, CPUFlags:0x%x, MiscFlags:0x%x, SampleCount:0x%x, SampleQuality:0x%x, SubresPtr:0x%p, SubresPitch:%i, SubresSlicePitch:%i"),
+		TEXT("%s failed with error %s\n at %s:%u\n Size=%ix%ix%i PF=%d Format=%s(0x%08X), NumMips=%i, Flags=%s, Usage:0x%x, CPUFlags:0x%x, MiscFlags:0x%x, SampleCount:0x%x, SampleQuality:0x%x, SubresPtr:0x%p, SubresPitch:%i, SubresSlicePitch:%i"),
 		ANSI_TO_TCHAR(Code),
+		*ErrorString,
 		ANSI_TO_TCHAR(Filename),
 		Line,
-		*ErrorString,
 		SizeX,
 		SizeY,
 		SizeZ,
@@ -355,11 +356,11 @@ void VerifyD3D11ResizeViewportResult(
 	const TCHAR* NewStateFormat = GetD3D11TextureFormatString(NewState.Format);
 
 	UE_LOG(LogD3D11RHI, Error,
-		TEXT("%s failed \n at %s:%u \n with error %s, \n (Size=%ix%i Fullscreen=%d Format=%s(0x%08X)) -> (Size=%ix%i Fullscreen=%d Format=%s(0x%08X))"),
+		TEXT("%s failed with error %s\n at %s:%u\n (Size=%ix%i Fullscreen=%d Format=%s(0x%08X)) -> (Size=%ix%i Fullscreen=%d Format=%s(0x%08X))"),
 		ANSI_TO_TCHAR(Code),
+		*ErrorString,
 		ANSI_TO_TCHAR(Filename),
 		Line,
-		*ErrorString,
 		OldState.SizeX,
 		OldState.SizeY,
 		OldState.bIsFullscreen ? 1 : 0,
@@ -375,11 +376,11 @@ void VerifyD3D11ResizeViewportResult(
 	TerminateOnOutOfMemory(D3DResult, true);
 
 	UE_LOG(LogD3D11RHI, Fatal,
-		TEXT("%s failed \n at %s:%u \n with error %s, \n (Size=%ix%i Fullscreen=%d Format=%s(0x%08X)) -> (Size=%ix%i Fullscreen=%d Format=%s(0x%08X))"),
+		TEXT("%s failed with error %s\n at %s:%u\n (Size=%ix%i Fullscreen=%d Format=%s(0x%08X)) -> (Size=%ix%i Fullscreen=%d Format=%s(0x%08X))"),
 		ANSI_TO_TCHAR(Code),
+		*ErrorString,
 		ANSI_TO_TCHAR(Filename),
 		Line,
-		*ErrorString,
 		OldState.SizeX,
 		OldState.SizeY,
 		OldState.bIsFullscreen ? 1 : 0,
