Buffer refactor merging Vertex, Index and Structured Buffer types into a single FRHIBuffer. No changes to APIs in this changelist other than those that naturally folded together when types merged. A f

diff --git a/Engine/Source/Runtime/VulkanRHI/Private/VulkanIndexBuffer.cpp b/Engine/Source/Runtime/VulkanRHI/Private/VulkanIndexBuffer.cpp
--- a/Engine/Source/Runtime/VulkanRHI/Private/VulkanIndexBuffer.cpp
+++ b/Engine/Source/Runtime/VulkanRHI/Private/VulkanIndexBuffer.cpp
@@ -60,8 +60,9 @@ static FORCEINLINE void UpdateVulkanBufferStats(uint64_t Size, VkBufferUsageFlag
 	}
 }
 
-FVulkanResourceMultiBuffer::FVulkanResourceMultiBuffer(FVulkanDevice* InDevice, VkBufferUsageFlags InBufferUsageFlags, uint32 InSize, uint32 InUEUsage, FRHIResourceCreateInfo& CreateInfo, class FRHICommandListImmediate* InRHICmdList)
-	: VulkanRHI::FDeviceChild(InDevice)
+FVulkanResourceMultiBuffer::FVulkanResourceMultiBuffer(FVulkanDevice* InDevice, VkBufferUsageFlags InBufferUsageFlags, uint32 InSize, uint32 InUEUsage, uint32 InStride, FRHIResourceCreateInfo& CreateInfo, class FRHICommandListImmediate* InRHICmdList)
+	: FRHIBuffer(InSize, InUEUsage, InStride)
+	, VulkanRHI::FDeviceChild(InDevice)
 	, UEUsage(InUEUsage)
 	, BufferUsageFlags(InBufferUsageFlags)
 	, NumBuffers(0)
@@ -69,13 +70,26 @@ FVulkanResourceMultiBuffer::FVulkanResourceMultiBuffer(FVulkanDevice* InDevice,
 {
 	VULKAN_TRACK_OBJECT_CREATE(FVulkanResourceMultiBuffer, this);
 
+	if ((InUEUsage & BUF_VertexBuffer) != 0)
+	{
+		BufferUsageFlags |= VK_BUFFER_USAGE_VERTEX_BUFFER_BIT;
+	}
+	if ((InUEUsage & BUF_IndexBuffer) != 0)
+	{
+		BufferUsageFlags |= VK_BUFFER_USAGE_INDEX_BUFFER_BIT;
+	}
+	if ((InUEUsage & BUF_StructuredBuffer) != 0)
+	{
+		BufferUsageFlags |= VK_BUFFER_USAGE_STORAGE_BUFFER_BIT;
+	}
+
 	if (InSize > 0)
 	{
 		const bool bStatic = (InUEUsage & BUF_Static) != 0;
 		const bool bDynamic = (InUEUsage & BUF_Dynamic) != 0;
 		const bool bVolatile = (InUEUsage & BUF_Volatile) != 0;
 		const bool bShaderResource = (InUEUsage & BUF_ShaderResource) != 0;
-		const bool bIsUniformBuffer = (InBufferUsageFlags & VK_BUFFER_USAGE_UNIFORM_BUFFER_BIT) != 0;
+		const bool bIsUniformBuffer = (BufferUsageFlags & VK_BUFFER_USAGE_UNIFORM_BUFFER_BIT) != 0;
 		const bool bUAV = (InUEUsage & BUF_UnorderedAccess) != 0;
 		const bool bIndirect = (InUEUsage & BUF_DrawIndirect) == BUF_DrawIndirect;
 		const bool bCPUReadable = (UEUsage & BUF_KeepCPUAccessible) != 0;
@@ -139,7 +153,7 @@ FVulkanResourceMultiBuffer::FVulkanResourceMultiBuffer(FVulkanDevice* InDevice,
 				CreateInfo.ResourceArray->Discard();
 			}
 
-			UpdateVulkanBufferStats(InSize * NumBuffers, InBufferUsageFlags, true);
+			UpdateVulkanBufferStats(InSize * NumBuffers, BufferUsageFlags, true);
 		}
 	}
 }
@@ -147,16 +161,16 @@ FVulkanResourceMultiBuffer::FVulkanResourceMultiBuffer(FVulkanDevice* InDevice,
 FVulkanResourceMultiBuffer::~FVulkanResourceMultiBuffer()
 {
 	VULKAN_TRACK_OBJECT_DELETE(FVulkanResourceMultiBuffer, this);
-	uint64_t Size = 0;
+	uint64_t TotalSize = 0;
 	for (uint32 Index = 0; Index < NumBuffers; ++Index)
 	{
-		Size += Buffers[Index].Size;
+		TotalSize += Buffers[Index].Size;
 		Device->GetMemoryManager().FreeVulkanAllocation(Buffers[Index]);
 	}
-	UpdateVulkanBufferStats(Size, BufferUsageFlags, false);
+	UpdateVulkanBufferStats(TotalSize, BufferUsageFlags, false);
 }
 
-void* FVulkanResourceMultiBuffer::Lock(bool bFromRenderingThread, EResourceLockMode LockMode, uint32 Size, uint32 Offset)
+void* FVulkanResourceMultiBuffer::Lock(bool bFromRenderingThread, EResourceLockMode LockMode, uint32 LockSize, uint32 Offset)
 {
 	void* Data = nullptr;
 
@@ -176,14 +190,14 @@ void* FVulkanResourceMultiBuffer::Lock(bool bFromRenderingThread, EResourceLockM
 		}
 		else
 		{
-			Device->GetImmediateContext().GetTempFrameAllocationBuffer().Alloc(Size + Offset, 256, VolatileLockInfo);
+			Device->GetImmediateContext().GetTempFrameAllocationBuffer().Alloc(LockSize + Offset, 256, VolatileLockInfo);
 			Data = VolatileLockInfo.Data;
 			++VolatileLockInfo.LockCounter;
 			check(!VolatileLockInfo.Allocation.HasAllocation());
 			Current.Alloc = VolatileLockInfo.Allocation;
 			Current.Handle = Current.Alloc.GetBufferHandle();
 			Current.Offset = VolatileLockInfo.CurrentOffset + Current.Alloc.Offset;
-			Current.Size = Size;
+			Current.Size = LockSize;
 		}
 	}
 	else
@@ -208,11 +222,11 @@ void* FVulkanResourceMultiBuffer::Lock(bool bFromRenderingThread, EResourceLockM
 				VulkanRHI::vkCmdPipelineBarrier(CmdBuffer->GetHandle(), VK_PIPELINE_STAGE_TRANSFER_BIT, VK_PIPELINE_STAGE_TRANSFER_BIT, 0, 1, &BarrierBefore, 0, nullptr, 0, nullptr);
 
 				// Create a staging buffer we can use to copy data from device to cpu.
-				VulkanRHI::FStagingBuffer* StagingBuffer = Device->GetStagingManager().AcquireBuffer(Size, VK_BUFFER_USAGE_TRANSFER_DST_BIT, VK_MEMORY_PROPERTY_HOST_CACHED_BIT);
+				VulkanRHI::FStagingBuffer* StagingBuffer = Device->GetStagingManager().AcquireBuffer(LockSize, VK_BUFFER_USAGE_TRANSFER_DST_BIT, VK_MEMORY_PROPERTY_HOST_CACHED_BIT);
 
 				// Fill the staging buffer with the data on the device.
 				VkBufferCopy Regions;
-				Regions.size = Size;
+				Regions.size = LockSize;
 				Regions.srcOffset = Offset;
 				Regions.dstOffset = 0;
 
@@ -235,7 +249,7 @@ void* FVulkanResourceMultiBuffer::Lock(bool bFromRenderingThread, EResourceLockM
 				// Release temp staging buffer during unlock.
 				VulkanRHI::FPendingBufferLock PendingLock;
 				PendingLock.Offset = 0;
-				PendingLock.Size = Size;
+				PendingLock.Size = LockSize;
 				PendingLock.LockMode = LockMode;
 				PendingLock.StagingBuffer = StagingBuffer;
 
@@ -255,7 +269,7 @@ void* FVulkanResourceMultiBuffer::Lock(bool bFromRenderingThread, EResourceLockM
 			Current.Alloc.Reference(Buffers[DynamicBufferIndex]);		
 			Current.Handle = (VkBuffer)Current.Alloc.VulkanHandle;
 			Current.Offset = Current.Alloc.Offset;
-			Current.Size = Size;
+			Current.Size = LockSize;
 
 			const bool bUnifiedMem = Device->HasUnifiedMemory();
 			if (bUnifiedMem)
@@ -266,10 +280,10 @@ void* FVulkanResourceMultiBuffer::Lock(bool bFromRenderingThread, EResourceLockM
 			{
 				VulkanRHI::FPendingBufferLock PendingLock;
 				PendingLock.Offset = Offset;
-				PendingLock.Size = Size;
+				PendingLock.Size = LockSize;
 				PendingLock.LockMode = LockMode;
 
-				VulkanRHI::FStagingBuffer* StagingBuffer = Device->GetStagingManager().AcquireBuffer(Size, VK_BUFFER_USAGE_TRANSFER_SRC_BIT);
+				VulkanRHI::FStagingBuffer* StagingBuffer = Device->GetStagingManager().AcquireBuffer(LockSize, VK_BUFFER_USAGE_TRANSFER_SRC_BIT);
 				PendingLock.StagingBuffer = StagingBuffer;
 				Data = StagingBuffer->GetMappedPointer();
 
@@ -395,38 +409,23 @@ void FVulkanResourceMultiBuffer::Unlock(bool bFromRenderingThread)
 
 void FVulkanResourceMultiBuffer::Swap(FVulkanResourceMultiBuffer& Other)
 {
+	FRHIBuffer::Swap(Other);
 	::Swap(*this, Other);
 }
 
-
-FVulkanIndexBuffer::FVulkanIndexBuffer(FVulkanDevice* InDevice, uint32 InStride, uint32 InSize, uint32 InUsage, FRHIResourceCreateInfo& CreateInfo, class FRHICommandListImmediate* InRHICmdList)
-	: FRHIIndexBuffer(InStride, InSize, InUsage)
-	, FVulkanResourceMultiBuffer(InDevice, VK_BUFFER_USAGE_INDEX_BUFFER_BIT, InSize, InUsage, CreateInfo, InRHICmdList)
-	, IndexType(InStride == 4 ? VK_INDEX_TYPE_UINT32 : VK_INDEX_TYPE_UINT16)
-{
-}
-
-void FVulkanIndexBuffer::Swap(FVulkanIndexBuffer& Other)
-{
-	FRHIIndexBuffer::Swap(Other);
-	FVulkanResourceMultiBuffer::Swap(Other);
-	::Swap(IndexType, Other.IndexType);
-}
-
-
 FIndexBufferRHIRef FVulkanDynamicRHI::RHICreateIndexBuffer(uint32 Stride, uint32 Size, uint32 InUsage, ERHIAccess InResourceState, FRHIResourceCreateInfo& CreateInfo)
 {
 	LLM_SCOPE_VULKAN(ELLMTagVulkan::VulkanIndexBuffers);
 	if (CreateInfo.bWithoutNativeResource)
 	{
-		return new FVulkanIndexBuffer(nullptr, 0, 0, 0, CreateInfo, nullptr);
+		return new FVulkanResourceMultiBuffer(nullptr, 0, 0, 0, 0, CreateInfo);
 	}
-	return new FVulkanIndexBuffer(Device, Stride, Size, InUsage, CreateInfo, nullptr);
+	return new FVulkanResourceMultiBuffer(Device, 0, Size, InUsage | BUF_IndexBuffer, Stride, CreateInfo);
 }
 
 void* FVulkanDynamicRHI::LockIndexBuffer_BottomOfPipe(FRHICommandListImmediate& RHICmdList, FRHIIndexBuffer* IndexBufferRHI, uint32 Offset, uint32 Size, EResourceLockMode LockMode)
 {
-	FVulkanIndexBuffer* IndexBuffer = ResourceCast(IndexBufferRHI);
+	FVulkanResourceMultiBuffer* IndexBuffer = ResourceCast(IndexBufferRHI);
 	return IndexBuffer->Lock(false, LockMode, Size, Offset);
 }
 
@@ -439,7 +438,7 @@ void* FVulkanDynamicRHI::LockIndexBuffer_RenderThread(class FRHICommandListImmed
 
 void FVulkanDynamicRHI::UnlockIndexBuffer_BottomOfPipe(FRHICommandListImmediate& RHICmdList, FRHIIndexBuffer* IndexBufferRHI)
 {
-	FVulkanIndexBuffer* IndexBuffer = ResourceCast(IndexBufferRHI);
+	FVulkanResourceMultiBuffer* IndexBuffer = ResourceCast(IndexBufferRHI);
 	IndexBuffer->Unlock(false);
 }
 
@@ -450,19 +449,19 @@ void FVulkanDynamicRHI::UnlockIndexBuffer_RenderThread(FRHICommandListImmediate&
 }
 #endif
 
-void FVulkanDynamicRHI::RHITransferIndexBufferUnderlyingResource(FRHIIndexBuffer* DestIndexBuffer, FRHIIndexBuffer* SrcIndexBuffer)
+void FVulkanDynamicRHI::RHITransferBufferUnderlyingResource(FRHIBuffer* DestBuffer, FRHIBuffer* SrcBuffer)
 {
-	check(DestIndexBuffer);
-	FVulkanIndexBuffer* Dest = ResourceCast(DestIndexBuffer);
-	if (!SrcIndexBuffer)
+	check(DestBuffer);
+	FVulkanResourceMultiBuffer* Dest = ResourceCast(DestBuffer);
+	if (!SrcBuffer)
 	{
 		FRHIResourceCreateInfo CreateInfo;
-		TRefCountPtr<FVulkanIndexBuffer> DeletionProxy = new FVulkanIndexBuffer(Dest->GetParent(), 0, 0, 0, CreateInfo, nullptr);
+		TRefCountPtr<FVulkanResourceMultiBuffer> DeletionProxy = new FVulkanResourceMultiBuffer(Dest->GetParent(), 0, 0, 0, 0, CreateInfo);
 		Dest->Swap(*DeletionProxy);
 	}
 	else
 	{
-		FVulkanIndexBuffer* Src = ResourceCast(SrcIndexBuffer);
+		FVulkanResourceMultiBuffer* Src = ResourceCast(SrcBuffer);
 		Dest->Swap(*Src);
 	}
 }
diff --git a/Engine/Source/Runtime/VulkanRHI/Private/VulkanPipelineState.h b/Engine/Source/Runtime/VulkanRHI/Private/VulkanPipelineState.h
--- a/Engine/Source/Runtime/VulkanRHI/Private/VulkanPipelineState.h
+++ b/Engine/Source/Runtime/VulkanRHI/Private/VulkanPipelineState.h
@@ -47,7 +47,7 @@ public:
 		bIsDSetsKeyDirty |= bDirty;
 	}
 
-	inline void SetStorageBuffer(uint8 DescriptorSet, uint32 BindingIndex, const FVulkanStructuredBuffer* StructuredBuffer)
+	inline void SetStorageBuffer(uint8 DescriptorSet, uint32 BindingIndex, const FVulkanResourceMultiBuffer* StructuredBuffer)
 	{
 		check(StructuredBuffer && (StructuredBuffer->GetBufferUsageFlags() & VK_BUFFER_USAGE_STORAGE_BUFFER_BIT) == VK_BUFFER_USAGE_STORAGE_BUFFER_BIT);
 		MarkDirty(DSWriter[DescriptorSet].WriteStorageBuffer(BindingIndex, StructuredBuffer->GetCurrentAllocation(), StructuredBuffer->GetOffset(), StructuredBuffer->GetCurrentSize()));
diff --git a/Engine/Source/Runtime/VulkanRHI/Private/VulkanResources.h b/Engine/Source/Runtime/VulkanRHI/Private/VulkanResources.h
--- a/Engine/Source/Runtime/VulkanRHI/Private/VulkanResources.h
+++ b/Engine/Source/Runtime/VulkanRHI/Private/VulkanResources.h
@@ -1095,13 +1095,13 @@ protected:
 	friend class FVulkanCommandListContext;
 };
 
-class FVulkanResourceMultiBuffer : public FVulkanEvictable, public VulkanRHI::FDeviceChild
+class FVulkanResourceMultiBuffer : public FRHIBuffer, public FVulkanEvictable, public VulkanRHI::FDeviceChild
 {
 	virtual void Evict(FVulkanDevice& Device);
 	virtual void Move(FVulkanDevice& Device, VulkanRHI::FVulkanAllocation& NewAllocation);
 
 public:
-	FVulkanResourceMultiBuffer(FVulkanDevice* InDevice, VkBufferUsageFlags InBufferUsageFlags, uint32 InSize, uint32 InUEUsage, FRHIResourceCreateInfo& CreateInfo, class FRHICommandListImmediate* InRHICmdList = nullptr);
+	FVulkanResourceMultiBuffer(FVulkanDevice* InDevice, VkBufferUsageFlags InBufferUsageFlags, uint32 InSize, uint32 InUEUsage, uint32 InStride, FRHIResourceCreateInfo& CreateInfo, class FRHICommandListImmediate* InRHICmdList = nullptr);
 	virtual ~FVulkanResourceMultiBuffer();
 
 	inline const VulkanRHI::FVulkanAllocation& GetCurrentAllocation() const
@@ -1164,9 +1164,13 @@ public:
 
 	inline uint32 GetUEUsage() const
 	{
-		return BufferUsageFlags;
+		return UEUsage;
 	}
 
+	inline VkIndexType GetIndexType() const
+	{
+		return (GetStride() == 4)? VK_INDEX_TYPE_UINT32 : VK_INDEX_TYPE_UINT16;
+	}
 
 	void* Lock(bool bFromRenderingThread, EResourceLockMode LockMode, uint32 Size, uint32 Offset);
 	void Unlock(bool bFromRenderingThread);
@@ -1208,30 +1212,6 @@ protected:
 	friend struct FRHICommandMultiBufferUnlock;
 };
 
-class FVulkanIndexBuffer : public FRHIIndexBuffer, public FVulkanResourceMultiBuffer
-{
-public:
-	FVulkanIndexBuffer(FVulkanDevice* InDevice, uint32 InStride, uint32 InSize, uint32 InUsage, FRHIResourceCreateInfo& CreateInfo, class FRHICommandListImmediate* InRHICmdList);
-
-	inline VkIndexType GetIndexType() const
-	{
-		return IndexType;
-	}
-
-	void Swap(FVulkanIndexBuffer& Other);
-
-private:
-	VkIndexType IndexType;
-};
-
-class FVulkanVertexBuffer : public FRHIVertexBuffer, public FVulkanResourceMultiBuffer
-{
-public:
-	FVulkanVertexBuffer(FVulkanDevice* InDevice, uint32 InSize, uint32 InUsage, FRHIResourceCreateInfo& CreateInfo, class FRHICommandListImmediate* InRHICmdList);
-
-	void Swap(FVulkanVertexBuffer& Other);
-};
-
 class FVulkanUniformBuffer : public FRHIUniformBuffer
 {
 public:
@@ -1276,26 +1256,13 @@ public:
 	VulkanRHI::FVulkanAllocation Allocation;
 };
 
-class FVulkanStructuredBuffer : public FRHIStructuredBuffer, public FVulkanResourceMultiBuffer
-{
-public:
-	FVulkanStructuredBuffer(FVulkanDevice* InDevice, uint32 Stride, uint32 Size, FRHIResourceCreateInfo& CreateInfo, uint32 InUsage);
-
-	~FVulkanStructuredBuffer();
-
-};
-
-
-
 class FVulkanUnorderedAccessView : public FRHIUnorderedAccessView, public VulkanRHI::FVulkanViewBase
 {
 public:
 
-	FVulkanUnorderedAccessView(FVulkanDevice* Device, FVulkanStructuredBuffer* StructuredBuffer, bool bUseUAVCounter, bool bAppendBuffer);
+	FVulkanUnorderedAccessView(FVulkanDevice* Device, FVulkanResourceMultiBuffer* Buffer, bool bUseUAVCounter, bool bAppendBuffer);
 	FVulkanUnorderedAccessView(FVulkanDevice* Device, FRHITexture* TextureRHI, uint32 MipLevel);
-	FVulkanUnorderedAccessView(FVulkanDevice* Device, FVulkanVertexBuffer* VertexBuffer, EPixelFormat Format);
-	FVulkanUnorderedAccessView(FVulkanDevice* Device, FVulkanIndexBuffer* IndexBuffer, EPixelFormat Format);
-
+	FVulkanUnorderedAccessView(FVulkanDevice* Device, FVulkanResourceMultiBuffer* Buffer, EPixelFormat Format);
 
 	~FVulkanUnorderedAccessView();
 
@@ -1305,15 +1272,15 @@ public:
 
 protected:
 	// the potential resources to refer to with the UAV object
-	TRefCountPtr<FVulkanStructuredBuffer> SourceStructuredBuffer;
+	TRefCountPtr<FVulkanResourceMultiBuffer> SourceStructuredBuffer;
 	// The texture that this UAV come from
 	TRefCountPtr<FRHITexture> SourceTexture;
 	FVulkanTextureView TextureView;
 	uint32 MipLevel;
 
 	// The vertex buffer this UAV comes from (can be null)
-	TRefCountPtr<FVulkanVertexBuffer> SourceVertexBuffer;
-	TRefCountPtr<FVulkanIndexBuffer> SourceIndexBuffer;
+	TRefCountPtr<FVulkanResourceMultiBuffer> SourceVertexBuffer;
+	TRefCountPtr<FVulkanResourceMultiBuffer> SourceIndexBuffer;
 	TRefCountPtr<FVulkanBufferView> BufferView;
 	EPixelFormat BufferViewFormat;
 
@@ -1331,7 +1298,7 @@ class FVulkanShaderResourceView : public FRHIShaderResourceView, public VulkanRH
 public:
 	FVulkanShaderResourceView(FVulkanDevice* Device, FRHIResource* InRHIBuffer, FVulkanResourceMultiBuffer* InSourceBuffer, uint32 InSize, EPixelFormat InFormat, uint32 InOffset = 0);
 	FVulkanShaderResourceView(FVulkanDevice* Device, FRHITexture* InSourceTexture, const FRHITextureSRVCreateInfo& InCreateInfo);
-	FVulkanShaderResourceView(FVulkanDevice* Device, FVulkanStructuredBuffer* InStructuredBuffer, uint32 InOffset = 0);
+	FVulkanShaderResourceView(FVulkanDevice* Device, FVulkanResourceMultiBuffer* InStructuredBuffer, uint32 InOffset = 0);
 
 	void Clear();
 
@@ -1351,7 +1318,7 @@ public:
 	// The texture that this SRV come from
 	TRefCountPtr<FRHITexture> SourceTexture;
 	FVulkanTextureView TextureView;
-	FVulkanStructuredBuffer* SourceStructuredBuffer;
+	FVulkanResourceMultiBuffer* SourceStructuredBuffer;
 	uint32 MipLevel = 0;
 	uint32 NumMips = MAX_uint32;
 	uint32 FirstArraySlice = 0;
@@ -1616,19 +1583,9 @@ struct TVulkanResourceTraits<FRHIUniformBuffer>
 	typedef FVulkanUniformBuffer TConcreteType;
 };
 template<>
-struct TVulkanResourceTraits<FRHIIndexBuffer>
-{
-	typedef FVulkanIndexBuffer TConcreteType;
-};
-template<>
-struct TVulkanResourceTraits<FRHIStructuredBuffer>
-{
-	typedef FVulkanStructuredBuffer TConcreteType;
-};
-template<>
-struct TVulkanResourceTraits<FRHIVertexBuffer>
+struct TVulkanResourceTraits<FRHIBuffer>
 {
-	typedef FVulkanVertexBuffer TConcreteType;
+	typedef FVulkanResourceMultiBuffer TConcreteType;
 };
 template<>
 struct TVulkanResourceTraits<FRHIShaderResourceView>
diff --git a/Engine/Source/Runtime/VulkanRHI/Private/VulkanStructuredBuffer.cpp b/Engine/Source/Runtime/VulkanRHI/Private/VulkanStructuredBuffer.cpp
--- a/Engine/Source/Runtime/VulkanRHI/Private/VulkanStructuredBuffer.cpp
+++ b/Engine/Source/Runtime/VulkanRHI/Private/VulkanStructuredBuffer.cpp
@@ -4,32 +4,21 @@
 #include "VulkanRHIPrivate.h"
 #include "Containers/ResourceArray.h"
 
-FVulkanStructuredBuffer::FVulkanStructuredBuffer(FVulkanDevice* InDevice, uint32 InStride, uint32 InSize, FRHIResourceCreateInfo& CreateInfo, uint32 InUsage)
-	: FRHIStructuredBuffer(InStride, InSize, InUsage)
-	, FVulkanResourceMultiBuffer(InDevice, VK_BUFFER_USAGE_STORAGE_BUFFER_BIT, InSize, InUsage, CreateInfo)
-{
-}
-
-FVulkanStructuredBuffer::~FVulkanStructuredBuffer()
-{
-}
-
-
 FStructuredBufferRHIRef FVulkanDynamicRHI::RHICreateStructuredBuffer(uint32 InStride, uint32 InSize, uint32 InUsage, ERHIAccess InResourceState, FRHIResourceCreateInfo& CreateInfo)
 {
-	return new FVulkanStructuredBuffer(Device, InStride, InSize, CreateInfo, InUsage);
+	return new FVulkanResourceMultiBuffer(Device, VK_BUFFER_USAGE_STORAGE_BUFFER_BIT, InSize, InUsage | BUF_StructuredBuffer, InStride, CreateInfo);
 }
 
-void* FVulkanDynamicRHI::LockStructuredBuffer_BottomOfPipe(FRHICommandListImmediate& RHICmdList, FRHIStructuredBuffer* StructuredBufferRHI,uint32 Offset,uint32 Size,EResourceLockMode LockMode)
+void* FVulkanDynamicRHI::LockStructuredBuffer_BottomOfPipe(FRHICommandListImmediate& RHICmdList, FRHIStructuredBuffer* StructuredBufferRHI, uint32 Offset, uint32 Size, EResourceLockMode LockMode)
 {
-	FVulkanStructuredBuffer* StructuredBuffer = ResourceCast(StructuredBufferRHI);
+	FVulkanResourceMultiBuffer* StructuredBuffer = ResourceCast(StructuredBufferRHI);
 
 	return StructuredBuffer->Lock(false, LockMode, Size, Offset);
 }
 
 void FVulkanDynamicRHI::UnlockStructuredBuffer_BottomOfPipe(FRHICommandListImmediate& RHICmdList, FRHIStructuredBuffer* StructuredBufferRHI)
 {
-	FVulkanStructuredBuffer* StructuredBuffer = ResourceCast(StructuredBufferRHI);
+	FVulkanResourceMultiBuffer* StructuredBuffer = ResourceCast(StructuredBufferRHI);
 
 	StructuredBuffer->Unlock(false);
 }
diff --git a/Engine/Source/Runtime/VulkanRHI/Private/VulkanTexture.cpp b/Engine/Source/Runtime/VulkanRHI/Private/VulkanTexture.cpp
--- a/Engine/Source/Runtime/VulkanRHI/Private/VulkanTexture.cpp
+++ b/Engine/Source/Runtime/VulkanRHI/Private/VulkanTexture.cpp
@@ -2989,8 +2989,8 @@ void FVulkanCommandListContext::RHICopyBufferRegion(FRHIVertexBuffer* DstBuffer,
 		return;
 	}
 
-	FVulkanVertexBuffer* DstBufferVk = ResourceCast(DstBuffer);
-	FVulkanVertexBuffer* SrcBufferVk = ResourceCast(SrcBuffer);
+	FVulkanResourceMultiBuffer* DstBufferVk = ResourceCast(DstBuffer);
+	FVulkanResourceMultiBuffer* SrcBufferVk = ResourceCast(SrcBuffer);
 
 	check(DstBufferVk && SrcBufferVk);
 	check(DstOffset + NumBytes <= DstBuffer->GetSize() && SrcOffset + NumBytes <= SrcBuffer->GetSize());
diff --git a/Engine/Source/Runtime/VulkanRHI/Private/VulkanUAV.cpp b/Engine/Source/Runtime/VulkanRHI/Private/VulkanUAV.cpp
--- a/Engine/Source/Runtime/VulkanRHI/Private/VulkanUAV.cpp
+++ b/Engine/Source/Runtime/VulkanRHI/Private/VulkanUAV.cpp
@@ -42,7 +42,7 @@ FVulkanShaderResourceView::FVulkanShaderResourceView(FVulkanDevice* Device, FRHI
 
 }
 
-FVulkanShaderResourceView::FVulkanShaderResourceView(FVulkanDevice* Device, FVulkanStructuredBuffer* InStructuredBuffer, uint32 InOffset)
+FVulkanShaderResourceView::FVulkanShaderResourceView(FVulkanDevice* Device, FVulkanResourceMultiBuffer* InStructuredBuffer, uint32 InOffset)
 	: VulkanRHI::FVulkanViewBase(Device)
 	, BufferViewFormat(PF_Unknown)
 	, SourceTexture(nullptr)
@@ -201,7 +201,7 @@ void FVulkanShaderResourceView::UpdateView()
 		}
 	}
 }
-FVulkanUnorderedAccessView::FVulkanUnorderedAccessView(FVulkanDevice* Device, FVulkanStructuredBuffer* StructuredBuffer, bool bUseUAVCounter, bool bAppendBuffer)
+FVulkanUnorderedAccessView::FVulkanUnorderedAccessView(FVulkanDevice* Device, FVulkanResourceMultiBuffer* StructuredBuffer, bool bUseUAVCounter, bool bAppendBuffer)
 	: VulkanRHI::FVulkanViewBase(Device)
 	, SourceStructuredBuffer(StructuredBuffer)
 	, MipLevel(0)
@@ -222,22 +222,20 @@ FVulkanUnorderedAccessView::FVulkanUnorderedAccessView(FVulkanDevice* Device, FR
 }
 
 
-FVulkanUnorderedAccessView::FVulkanUnorderedAccessView(FVulkanDevice* Device, FVulkanVertexBuffer* VertexBuffer, EPixelFormat Format)
+FVulkanUnorderedAccessView::FVulkanUnorderedAccessView(FVulkanDevice* Device, FVulkanResourceMultiBuffer* Buffer, EPixelFormat Format)
 	: VulkanRHI::FVulkanViewBase(Device)
 	, MipLevel(0)
 	, BufferViewFormat(Format)
 	, VolatileLockCounter(MAX_uint32)
 {
-	SourceVertexBuffer = VertexBuffer;
-}
-
-FVulkanUnorderedAccessView::FVulkanUnorderedAccessView(FVulkanDevice* Device, FVulkanIndexBuffer* IndexBuffer, EPixelFormat Format)
-	: VulkanRHI::FVulkanViewBase(Device)
-	, MipLevel(0)
-	, BufferViewFormat(Format)
-	, VolatileLockCounter(MAX_uint32)
-{
-	SourceIndexBuffer = IndexBuffer;
+	if (Buffer->GetUEUsage() & BUF_IndexBuffer)
+	{
+		SourceIndexBuffer = Buffer;
+	}
+	else if (Buffer->GetUEUsage() & BUF_VertexBuffer)
+	{
+		SourceVertexBuffer = Buffer;
+	}
 }
 
 void FVulkanUnorderedAccessView::Invalidate()
@@ -337,11 +335,11 @@ void FVulkanUnorderedAccessView::UpdateView()
 	}
 }
 
-FUnorderedAccessViewRHIRef FVulkanDynamicRHI::RHICreateUnorderedAccessView(FRHIStructuredBuffer* StructuredBufferRHI, bool bUseUAVCounter, bool bAppendBuffer)
+FUnorderedAccessViewRHIRef FVulkanDynamicRHI::RHICreateUnorderedAccessView(FRHIBuffer* BufferRHI, bool bUseUAVCounter, bool bAppendBuffer)
 {
-	FVulkanStructuredBuffer* StructuredBuffer = ResourceCast(StructuredBufferRHI);
+	FVulkanResourceMultiBuffer* Buffer = ResourceCast(BufferRHI);
 
-	FVulkanUnorderedAccessView* UAV = new FVulkanUnorderedAccessView(Device, StructuredBuffer, bUseUAVCounter, bAppendBuffer);
+	FVulkanUnorderedAccessView* UAV = new FVulkanUnorderedAccessView(Device, Buffer, bUseUAVCounter, bAppendBuffer);
 	return UAV;
 }
 
@@ -351,37 +349,22 @@ FUnorderedAccessViewRHIRef FVulkanDynamicRHI::RHICreateUnorderedAccessView(FRHIT
 	return UAV;
 }
 
-FUnorderedAccessViewRHIRef FVulkanDynamicRHI::RHICreateUnorderedAccessView(FRHIVertexBuffer* VertexBufferRHI, uint8 Format)
-{
-	FVulkanVertexBuffer* VertexBuffer = ResourceCast(VertexBufferRHI);
-
-	FVulkanUnorderedAccessView* UAV = new FVulkanUnorderedAccessView(Device, VertexBuffer, (EPixelFormat)Format);
-	return UAV;
-}
-
-FUnorderedAccessViewRHIRef FVulkanDynamicRHI::RHICreateUnorderedAccessView(FRHIIndexBuffer* IndexBufferRHI, uint8 Format)
+FUnorderedAccessViewRHIRef FVulkanDynamicRHI::RHICreateUnorderedAccessView(FRHIBuffer* BufferRHI, uint8 Format)
 {
-	FVulkanIndexBuffer* IndexBuffer = ResourceCast(IndexBufferRHI);
+	FVulkanResourceMultiBuffer* Buffer = ResourceCast(BufferRHI);
 
-	FVulkanUnorderedAccessView* UAV = new FVulkanUnorderedAccessView(Device, IndexBuffer, (EPixelFormat)Format);
+	FVulkanUnorderedAccessView* UAV = new FVulkanUnorderedAccessView(Device, Buffer, (EPixelFormat)Format);
 	return UAV;
 }
 
-FShaderResourceViewRHIRef FVulkanDynamicRHI::RHICreateShaderResourceView(FRHIStructuredBuffer* StructuredBufferRHI)
-{
-	FVulkanStructuredBuffer* StructuredBuffer = ResourceCast(StructuredBufferRHI);
-	FVulkanShaderResourceView* SRV = new FVulkanShaderResourceView(Device, StructuredBuffer);
-	return SRV;
-}
-
-FShaderResourceViewRHIRef FVulkanDynamicRHI::RHICreateShaderResourceView(FRHIVertexBuffer* VertexBufferRHI, uint32 Stride, uint8 Format)
+FShaderResourceViewRHIRef FVulkanDynamicRHI::RHICreateShaderResourceView(FRHIBuffer* BufferRHI, uint32 Stride, uint8 Format)
 {	
-	if (!VertexBufferRHI)
+	if (!BufferRHI)
 	{
 		return new FVulkanShaderResourceView(Device, nullptr, nullptr, 0, (EPixelFormat)Format);
 	}
-	FVulkanVertexBuffer* VertexBuffer = ResourceCast(VertexBufferRHI);
-	return new FVulkanShaderResourceView(Device, VertexBufferRHI, VertexBuffer, VertexBuffer->GetSize(), (EPixelFormat)Format);
+	FVulkanResourceMultiBuffer* Buffer = ResourceCast(BufferRHI);
+	return new FVulkanShaderResourceView(Device, BufferRHI, Buffer, Buffer->GetSize(), (EPixelFormat)Format);
 }
 
 FShaderResourceViewRHIRef FVulkanDynamicRHI::RHICreateShaderResourceView(const FShaderResourceViewInitializer& Initializer)
@@ -391,12 +374,12 @@ FShaderResourceViewRHIRef FVulkanDynamicRHI::RHICreateShaderResourceView(const F
 		case FShaderResourceViewInitializer::EType::VertexBufferSRV:
 		{
 			const FShaderResourceViewInitializer::FVertexBufferShaderResourceViewInitializer Desc = Initializer.AsVertexBufferSRV();
-			if (Desc.VertexBuffer)
+			if (Desc.Buffer)
 			{
 				const uint32 Stride = GPixelFormats[Desc.Format].BlockBytes;
-				FVulkanVertexBuffer* VertexBuffer = ResourceCast(Desc.VertexBuffer);
+				FVulkanResourceMultiBuffer* VertexBuffer = ResourceCast(Desc.Buffer);
 				uint32 Size = FMath::Min(VertexBuffer->GetSize() - Desc.StartOffsetBytes, Desc.NumElements * Stride);
-				return new FVulkanShaderResourceView(Device, Desc.VertexBuffer, VertexBuffer, Size, (EPixelFormat)Desc.Format, Desc.StartOffsetBytes);
+				return new FVulkanShaderResourceView(Device, Desc.Buffer, VertexBuffer, Size, (EPixelFormat)Desc.Format, Desc.StartOffsetBytes);
 			}
 			else
 			{
@@ -406,20 +389,20 @@ FShaderResourceViewRHIRef FVulkanDynamicRHI::RHICreateShaderResourceView(const F
 		case FShaderResourceViewInitializer::EType::StructuredBufferSRV:
 		{
 			const FShaderResourceViewInitializer::FStructuredBufferShaderResourceViewInitializer Desc = Initializer.AsStructuredBufferSRV();
-			check(Desc.StructuredBuffer);
-			FVulkanStructuredBuffer* StructuredBuffer = ResourceCast(Desc.StructuredBuffer);
+			check(Desc.Buffer);
+			FVulkanResourceMultiBuffer* StructuredBuffer = ResourceCast(Desc.Buffer);
 			return new FVulkanShaderResourceView(Device, StructuredBuffer, Desc.StartOffsetBytes);
 		}			
 		case FShaderResourceViewInitializer::EType::IndexBufferSRV:
 		{
 			const FShaderResourceViewInitializer::FIndexBufferShaderResourceViewInitializer Desc = Initializer.AsIndexBufferSRV();
-			check(Desc.IndexBuffer);
-			FVulkanIndexBuffer* IndexBuffer = ResourceCast(Desc.IndexBuffer);
-			const uint32 Stride = Desc.IndexBuffer->GetStride();
+			check(Desc.Buffer);
+			FVulkanResourceMultiBuffer* IndexBuffer = ResourceCast(Desc.Buffer);
+			const uint32 Stride = Desc.Buffer->GetStride();
 			check(Stride == 2 || Stride == 4);
 			EPixelFormat Format = (Stride == 4) ? PF_R32_UINT : PF_R16_UINT;
 			uint32 Size = FMath::Min(IndexBuffer->GetSize() - Desc.StartOffsetBytes, Desc.NumElements * Stride);
-			return new FVulkanShaderResourceView(Device, Desc.IndexBuffer, IndexBuffer, Size, Format, Desc.StartOffsetBytes);
+			return new FVulkanShaderResourceView(Device, Desc.Buffer, IndexBuffer, Size, Format, Desc.StartOffsetBytes);
 		}
 	}
 	checkNoEntry();
@@ -432,46 +415,56 @@ FShaderResourceViewRHIRef FVulkanDynamicRHI::RHICreateShaderResourceView(FRHITex
 	return SRV;
 }
 
-FShaderResourceViewRHIRef FVulkanDynamicRHI::RHICreateShaderResourceView(FRHIIndexBuffer* IndexBufferRHI)
+FShaderResourceViewRHIRef FVulkanDynamicRHI::RHICreateShaderResourceView(FRHIBuffer* BufferRHI)
 {
-	if (!IndexBufferRHI)
+	if (BufferRHI && BufferRHI->GetUsage() & (BUF_VertexBuffer | BUF_StructuredBuffer))
 	{
-		return new FVulkanShaderResourceView(Device, nullptr, nullptr, 0, PF_R16_UINT);
+		FVulkanResourceMultiBuffer* Buffer = ResourceCast(BufferRHI);
+		FVulkanShaderResourceView* SRV = new FVulkanShaderResourceView(Device, Buffer);
+		return SRV;
+	}
+	else
+	{
+		if (!BufferRHI)
+		{
+			return new FVulkanShaderResourceView(Device, nullptr, nullptr, 0, PF_R16_UINT);
+		}
+		check(BufferRHI->GetUsage() & BUF_IndexBuffer);
+		check(BufferRHI->GetStride() == 2 || BufferRHI->GetStride() == 4);
+		FVulkanResourceMultiBuffer* Buffer = ResourceCast(BufferRHI);
+		EPixelFormat Format = (BufferRHI->GetStride() == 4) ? PF_R32_UINT : PF_R16_UINT;
+		FVulkanShaderResourceView* SRV = new FVulkanShaderResourceView(Device, BufferRHI, Buffer, Buffer->GetSize(), Format);
+		return SRV;
 	}
-	FVulkanIndexBuffer* IndexBuffer = ResourceCast(IndexBufferRHI);
-	check(IndexBufferRHI->GetStride() == 2 || IndexBufferRHI->GetStride() == 4);
-	EPixelFormat Format = (IndexBufferRHI->GetStride() == 4) ? PF_R32_UINT : PF_R16_UINT;
-	FVulkanShaderResourceView* SRV = new FVulkanShaderResourceView(Device, IndexBufferRHI, IndexBuffer, IndexBuffer->GetSize(), Format);
-	return SRV;
 }
 
-void FVulkanDynamicRHI::RHIUpdateShaderResourceView(FRHIShaderResourceView* SRV, FRHIVertexBuffer* VertexBuffer, uint32 Stride, uint8 Format)
+void FVulkanDynamicRHI::RHIUpdateShaderResourceView(FRHIShaderResourceView* SRV, FRHIBuffer* Buffer, uint32 Stride, uint8 Format)
 {
 	FVulkanShaderResourceView* SRVVk = ResourceCast(SRV);
 	check(SRVVk && SRVVk->GetParent() == Device);
-	if (!VertexBuffer)
+	if (!Buffer)
 	{
 		SRVVk->Clear();
 	}
-	else if (SRVVk->SourceRHIBuffer.GetReference() != VertexBuffer)
+	else if (SRVVk->SourceRHIBuffer.GetReference() != Buffer)
 	{
-		FVulkanVertexBuffer* VertexBufferVk = ResourceCast(VertexBuffer);
-		SRVVk->Rename(VertexBuffer, VertexBufferVk, VertexBufferVk->GetSize(), (EPixelFormat)Format);
+		FVulkanResourceMultiBuffer* BufferVk = ResourceCast(Buffer);
+		SRVVk->Rename(Buffer, BufferVk, BufferVk->GetSize(), (EPixelFormat)Format);
 	}
 }
 
-void FVulkanDynamicRHI::RHIUpdateShaderResourceView(FRHIShaderResourceView* SRV, FRHIIndexBuffer* IndexBuffer)
+void FVulkanDynamicRHI::RHIUpdateShaderResourceView(FRHIShaderResourceView* SRV, FRHIBuffer* Buffer)
 {
 	FVulkanShaderResourceView* SRVVk = ResourceCast(SRV);
 	check(SRVVk && SRVVk->GetParent() == Device);
-	if (!IndexBuffer)
+	if (!Buffer)
 	{
 		SRVVk->Clear();
 	}
-	else if (SRVVk->SourceRHIBuffer.GetReference() != IndexBuffer)
+	else if (SRVVk->SourceRHIBuffer.GetReference() != Buffer)
 	{
-		FVulkanIndexBuffer* IndexBufferVk = ResourceCast(IndexBuffer);
-		SRVVk->Rename(IndexBuffer, IndexBufferVk, IndexBufferVk->GetSize(), IndexBufferVk->GetStride() == 2u ? PF_R16_UINT : PF_R32_UINT);
+		FVulkanResourceMultiBuffer* BufferVk = ResourceCast(Buffer);
+		SRVVk->Rename(Buffer, BufferVk, BufferVk->GetSize(), BufferVk->GetStride() == 2u ? PF_R16_UINT : PF_R32_UINT);
 	}
 }
 
diff --git a/Engine/Source/Runtime/VulkanRHI/Private/VulkanVertexBuffer.cpp b/Engine/Source/Runtime/VulkanRHI/Private/VulkanVertexBuffer.cpp
--- a/Engine/Source/Runtime/VulkanRHI/Private/VulkanVertexBuffer.cpp
+++ b/Engine/Source/Runtime/VulkanRHI/Private/VulkanVertexBuffer.cpp
@@ -7,33 +7,21 @@
 #include "VulkanRHIPrivate.h"
 #include "VulkanLLM.h"
 
-FVulkanVertexBuffer::FVulkanVertexBuffer(FVulkanDevice* InDevice, uint32 InSize, uint32 InUsage, FRHIResourceCreateInfo& CreateInfo, class FRHICommandListImmediate* InRHICmdList)
-	: FRHIVertexBuffer(InSize, InUsage)
-	, FVulkanResourceMultiBuffer(InDevice, VK_BUFFER_USAGE_VERTEX_BUFFER_BIT, InSize, InUsage, CreateInfo, InRHICmdList)
-{
-}
-
-void FVulkanVertexBuffer::Swap(FVulkanVertexBuffer& Other)
-{
-	FRHIVertexBuffer::Swap(Other);
-	FVulkanResourceMultiBuffer::Swap(Other);
-}
-
 FVertexBufferRHIRef FVulkanDynamicRHI::RHICreateVertexBuffer(uint32 Size, uint32 InUsage, ERHIAccess InResourceState, FRHIResourceCreateInfo& CreateInfo)
 {
 	LLM_SCOPE_VULKAN(ELLMTagVulkan::VulkanVertexBuffers);
 	if (CreateInfo.bWithoutNativeResource)
 	{
-		return new FVulkanVertexBuffer(nullptr, 0, 0, CreateInfo, nullptr);
+		return new FVulkanResourceMultiBuffer(nullptr, VK_BUFFER_USAGE_VERTEX_BUFFER_BIT, 0, 0, 0, CreateInfo);
 	}
-	FVulkanVertexBuffer* VertexBuffer = new FVulkanVertexBuffer(Device, Size, InUsage, CreateInfo, nullptr);
+	FVulkanResourceMultiBuffer* VertexBuffer = new FVulkanResourceMultiBuffer(Device, VK_BUFFER_USAGE_VERTEX_BUFFER_BIT, Size, InUsage | BUF_VertexBuffer, 0, CreateInfo);
 	return VertexBuffer;
 }
 
 void* FVulkanDynamicRHI::LockVertexBuffer_BottomOfPipe(FRHICommandListImmediate& RHICmdList, FRHIVertexBuffer* VertexBufferRHI, uint32 Offset, uint32 Size, EResourceLockMode LockMode)
 {
 	LLM_SCOPE_VULKAN(ELLMTagVulkan::VulkanVertexBuffers);
-	FVulkanVertexBuffer* VertexBuffer = ResourceCast(VertexBufferRHI);
+	FVulkanResourceMultiBuffer* VertexBuffer = ResourceCast(VertexBufferRHI);
 	return VertexBuffer->Lock(false, LockMode, Size, Offset);
 }
 
@@ -47,7 +35,7 @@ void* FVulkanDynamicRHI::LockVertexBuffer_RenderThread(class FRHICommandListImme
 void FVulkanDynamicRHI::UnlockVertexBuffer_BottomOfPipe(FRHICommandListImmediate& RHICmdList, FRHIVertexBuffer* VertexBufferRHI)
 {
 	LLM_SCOPE_VULKAN(ELLMTagVulkan::VulkanVertexBuffers);
-	FVulkanVertexBuffer* VertexBuffer = ResourceCast(VertexBufferRHI);
+	FVulkanResourceMultiBuffer* VertexBuffer = ResourceCast(VertexBufferRHI);
 	VertexBuffer->Unlock(false);
 }
 
@@ -62,20 +50,3 @@ void FVulkanDynamicRHI::RHICopyVertexBuffer(FRHIVertexBuffer* SourceBufferRHI, F
 {
 	VULKAN_SIGNAL_UNIMPLEMENTED();
 }
-
-void FVulkanDynamicRHI::RHITransferVertexBufferUnderlyingResource(FRHIVertexBuffer* DestVertexBuffer, FRHIVertexBuffer* SrcVertexBuffer)
-{
-	check(DestVertexBuffer);
-	FVulkanVertexBuffer* Dest = ResourceCast(DestVertexBuffer);
-	if (!SrcVertexBuffer)
-	{
-		FRHIResourceCreateInfo CreateInfo;
-		TRefCountPtr<FVulkanVertexBuffer> DeletionProxy = new FVulkanVertexBuffer(Dest->GetParent(), 0, 0, CreateInfo, nullptr);
-		Dest->Swap(*DeletionProxy);
-	}
-	else
-	{
-		FVulkanVertexBuffer* Src = ResourceCast(SrcVertexBuffer);
-		Dest->Swap(*Src);
-	}
-}
diff --git a/Engine/Source/Runtime/VulkanRHI/Public/VulkanDynamicRHI.h b/Engine/Source/Runtime/VulkanRHI/Public/VulkanDynamicRHI.h
--- a/Engine/Source/Runtime/VulkanRHI/Public/VulkanDynamicRHI.h
+++ b/Engine/Source/Runtime/VulkanRHI/Public/VulkanDynamicRHI.h
@@ -83,25 +83,22 @@ public:
 	virtual FIndexBufferRHIRef RHICreateIndexBuffer(uint32 Stride, uint32 Size, uint32 InUsage, ERHIAccess InResourceState, FRHIResourceCreateInfo& CreateInfo) final override;
 	virtual void* LockIndexBuffer_BottomOfPipe(FRHICommandListImmediate& RHICmdList, FRHIIndexBuffer* IndexBuffer, uint32 Offset, uint32 Size, EResourceLockMode LockMode) final override;
 	virtual void UnlockIndexBuffer_BottomOfPipe(FRHICommandListImmediate& RHICmdList, FRHIIndexBuffer* IndexBuffer) final override;
-	virtual void RHITransferIndexBufferUnderlyingResource(FRHIIndexBuffer* DestIndexBuffer, FRHIIndexBuffer* SrcIndexBuffer) final override;
 	virtual FVertexBufferRHIRef RHICreateVertexBuffer(uint32 Size, uint32 InUsage, ERHIAccess InResourceState, FRHIResourceCreateInfo& CreateInfo) final override;
 	virtual void* LockVertexBuffer_BottomOfPipe(FRHICommandListImmediate& RHICmdList, FRHIVertexBuffer* VertexBuffer, uint32 Offset, uint32 SizeRHI, EResourceLockMode LockMode) final override;
 	virtual void UnlockVertexBuffer_BottomOfPipe(FRHICommandListImmediate& RHICmdList, FRHIVertexBuffer* VertexBuffer) final override;
 	virtual void RHICopyVertexBuffer(FRHIVertexBuffer* SourceBuffer, FRHIVertexBuffer* DestBuffer) final override;
-	virtual void RHITransferVertexBufferUnderlyingResource(FRHIVertexBuffer* DestVertexBuffer, FRHIVertexBuffer* SrcVertexBuffer) final override;
+	virtual void RHITransferBufferUnderlyingResource(FRHIBuffer* DestBuffer, FRHIBuffer* SrcBuffer) final override;
 	virtual FStructuredBufferRHIRef RHICreateStructuredBuffer(uint32 Stride, uint32 Size, uint32 InUsage, ERHIAccess InResourceState, FRHIResourceCreateInfo& CreateInfo) final override;
 	virtual void* LockStructuredBuffer_BottomOfPipe(FRHICommandListImmediate& RHICmdList, FRHIStructuredBuffer* StructuredBuffer, uint32 Offset, uint32 SizeRHI, EResourceLockMode LockMode) final override;
 	virtual void UnlockStructuredBuffer_BottomOfPipe(FRHICommandListImmediate& RHICmdList, FRHIStructuredBuffer* StructuredBuffer) final override;
-	virtual FUnorderedAccessViewRHIRef RHICreateUnorderedAccessView(FRHIStructuredBuffer* StructuredBuffer, bool bUseUAVCounter, bool bAppendBuffer) final override;
+	virtual FUnorderedAccessViewRHIRef RHICreateUnorderedAccessView(FRHIBuffer* Buffer, bool bUseUAVCounter, bool bAppendBuffer) final override;
 	virtual FUnorderedAccessViewRHIRef RHICreateUnorderedAccessView(FRHITexture* Texture, uint32 MipLevel) final override;
-	virtual FUnorderedAccessViewRHIRef RHICreateUnorderedAccessView(FRHIVertexBuffer* VertexBuffer, uint8 Format) final override;
-	virtual FUnorderedAccessViewRHIRef RHICreateUnorderedAccessView(FRHIIndexBuffer* IndexBuffer, uint8 Format) final override;
-	virtual FShaderResourceViewRHIRef RHICreateShaderResourceView(FRHIStructuredBuffer* StructuredBuffer) final override;
-	virtual FShaderResourceViewRHIRef RHICreateShaderResourceView(FRHIVertexBuffer* VertexBuffer, uint32 Stride, uint8 Format) final override;
+	virtual FUnorderedAccessViewRHIRef RHICreateUnorderedAccessView(FRHIBuffer* Buffer, uint8 Format) final override;
+	virtual FShaderResourceViewRHIRef RHICreateShaderResourceView(FRHIBuffer* Buffer) final override;
+	virtual FShaderResourceViewRHIRef RHICreateShaderResourceView(FRHIBuffer* Buffer, uint32 Stride, uint8 Format) final override;
 	virtual FShaderResourceViewRHIRef RHICreateShaderResourceView(const FShaderResourceViewInitializer& Initializer) final override;
-	virtual FShaderResourceViewRHIRef RHICreateShaderResourceView(FRHIIndexBuffer* Buffer) final override;
-	virtual void RHIUpdateShaderResourceView(FRHIShaderResourceView* SRV, FRHIVertexBuffer* VertexBuffer, uint32 Stride, uint8 Format) final override;
-	virtual void RHIUpdateShaderResourceView(FRHIShaderResourceView* SRV, FRHIIndexBuffer* IndexBuffer) final override;
+	virtual void RHIUpdateShaderResourceView(FRHIShaderResourceView* SRV, FRHIBuffer* Buffer, uint32 Stride, uint8 Format) final override;
+	virtual void RHIUpdateShaderResourceView(FRHIShaderResourceView* SRV, FRHIBuffer* Buffer) final override;
 	virtual uint64 RHICalcTexture2DPlatformSize(uint32 SizeX, uint32 SizeY, uint8 Format, uint32 NumMips, uint32 NumSamples, ETextureCreateFlags Flags, const FRHIResourceCreateInfo& CreateInfo, uint32& OutAlign) final override;
 	virtual uint64 RHICalcTexture3DPlatformSize(uint32 SizeX, uint32 SizeY, uint32 SizeZ, uint8 Format, uint32 NumMips, ETextureCreateFlags Flags, const FRHIResourceCreateInfo& CreateInfo, uint32& OutAlign) final override;
 	virtual uint64 RHICalcTextureCubePlatformSize(uint32 Size, uint8 Format, uint32 NumMips, ETextureCreateFlags Flags, const FRHIResourceCreateInfo& CreateInfo, uint32& OutAlign) final override;
diff --git a/Engine/Source/Runtime/Windows/D3D11RHI/Private/D3D11Commands.cpp b/Engine/Source/Runtime/Windows/D3D11RHI/Private/D3D11Commands.cpp
--- a/Engine/Source/Runtime/Windows/D3D11RHI/Private/D3D11Commands.cpp
+++ b/Engine/Source/Runtime/Windows/D3D11RHI/Private/D3D11Commands.cpp
@@ -173,7 +173,7 @@ void FD3D11DynamicRHI::RHIEndUpdateMultiFrameResource(FRHIUnorderedAccessView* U
 // Vertex state.
 void FD3D11DynamicRHI::RHISetStreamSource(uint32 StreamIndex, FRHIVertexBuffer* VertexBufferRHI, uint32 Offset)
 {
-	FD3D11VertexBuffer* VertexBuffer = ResourceCast(VertexBufferRHI);
+	FD3D11Buffer* VertexBuffer = ResourceCast(VertexBufferRHI);
 
 	ID3D11Buffer* D3DBuffer = VertexBuffer ? VertexBuffer->Resource : NULL;
 	TrackResourceBoundAsVB(VertexBuffer, StreamIndex);
@@ -239,11 +239,11 @@ void FD3D11DynamicRHI::RHIDispatchComputeShader(uint32 ThreadGroupCountX, uint32
 	EnableUAVOverlap();
 }
 
-void FD3D11DynamicRHI::RHIDispatchIndirectComputeShader(FRHIVertexBuffer* ArgumentBufferRHI, uint32 ArgumentOffset)
+void FD3D11DynamicRHI::RHIDispatchIndirectComputeShader(FRHIBuffer* ArgumentBufferRHI, uint32 ArgumentOffset)
 { 
 	FRHIComputeShader* ComputeShaderRHI = GetCurrentComputeShader();
 	FD3D11ComputeShader* ComputeShader = ResourceCast(ComputeShaderRHI);
-	FD3D11VertexBuffer* ArgumentBuffer = ResourceCast(ArgumentBufferRHI);
+	FD3D11Buffer* ArgumentBuffer = ResourceCast(ArgumentBufferRHI);
 
 	GPUProfilingData.RegisterGPUDispatch(FIntVector(1, 1, 1));
 
@@ -1675,9 +1675,9 @@ void FD3D11DynamicRHI::RHIDrawPrimitive(uint32 BaseVertexIndex,uint32 NumPrimiti
 	EnableUAVOverlap();
 }
 
-void FD3D11DynamicRHI::RHIDrawPrimitiveIndirect(FRHIVertexBuffer* ArgumentBufferRHI,uint32 ArgumentOffset)
+void FD3D11DynamicRHI::RHIDrawPrimitiveIndirect(FRHIBuffer* ArgumentBufferRHI, uint32 ArgumentOffset)
 {
-	FD3D11VertexBuffer* ArgumentBuffer = ResourceCast(ArgumentBufferRHI);
+	FD3D11Buffer* ArgumentBuffer = ResourceCast(ArgumentBufferRHI);
 
 	RHI_DRAW_CALL_INC();
 
@@ -1692,10 +1692,10 @@ void FD3D11DynamicRHI::RHIDrawPrimitiveIndirect(FRHIVertexBuffer* ArgumentBuffer
 	EnableUAVOverlap();
 }
 
-void FD3D11DynamicRHI::RHIDrawIndexedIndirect(FRHIIndexBuffer* IndexBufferRHI, FRHIStructuredBuffer* ArgumentsBufferRHI, int32 DrawArgumentsIndex, uint32 NumInstances)
+void FD3D11DynamicRHI::RHIDrawIndexedIndirect(FRHIBuffer* IndexBufferRHI, FRHIBuffer* ArgumentsBufferRHI, int32 DrawArgumentsIndex, uint32 NumInstances)
 {
-	FD3D11IndexBuffer* IndexBuffer = ResourceCast(IndexBufferRHI);
-	FD3D11StructuredBuffer* ArgumentsBuffer = ResourceCast(ArgumentsBufferRHI);
+	FD3D11Buffer* IndexBuffer = ResourceCast(IndexBufferRHI);
+	FD3D11Buffer* ArgumentsBuffer = ResourceCast(ArgumentsBufferRHI);
 
 	RHI_DRAW_CALL_INC();
 
@@ -1705,7 +1705,6 @@ void FD3D11DynamicRHI::RHIDrawIndexedIndirect(FRHIIndexBuffer* IndexBufferRHI, F
 	CommitNonComputeShaderConstants();
 
 	// determine 16bit vs 32bit indices
-	uint32 SizeFormat = sizeof(DXGI_FORMAT);
 	const DXGI_FORMAT Format = (IndexBuffer->GetStride() == sizeof(uint16) ? DXGI_FORMAT_R16_UINT : DXGI_FORMAT_R32_UINT);
 
 	TrackResourceBoundAsIB(IndexBuffer);
@@ -1724,11 +1723,11 @@ void FD3D11DynamicRHI::RHIDrawIndexedIndirect(FRHIIndexBuffer* IndexBufferRHI, F
 	EnableUAVOverlap();
 }
 
-void FD3D11DynamicRHI::RHIDrawIndexedPrimitive(FRHIIndexBuffer* IndexBufferRHI,int32 BaseVertexIndex,uint32 FirstInstance,uint32 NumVertices,uint32 StartIndex,uint32 NumPrimitives,uint32 NumInstances)
+void FD3D11DynamicRHI::RHIDrawIndexedPrimitive(FRHIBuffer* IndexBufferRHI, int32 BaseVertexIndex, uint32 FirstInstance, uint32 NumVertices, uint32 StartIndex, uint32 NumPrimitives, uint32 NumInstances)
 {
 	RHI_DRAW_CALL_STATS(PrimitiveType, FMath::Max(NumInstances, 1U) * NumPrimitives);
 
-	FD3D11IndexBuffer* IndexBuffer = ResourceCast(IndexBufferRHI);
+	FD3D11Buffer* IndexBuffer = ResourceCast(IndexBufferRHI);
 
 	// called should make sure the input is valid, this avoid hidden bugs
 	ensure(NumPrimitives > 0);
@@ -1739,7 +1738,6 @@ void FD3D11DynamicRHI::RHIDrawIndexedPrimitive(FRHIIndexBuffer* IndexBufferRHI,i
 	CommitNonComputeShaderConstants();
 
 	// determine 16bit vs 32bit indices
-	uint32 SizeFormat = sizeof(DXGI_FORMAT);
 	const DXGI_FORMAT Format = (IndexBuffer->GetStride() == sizeof(uint16) ? DXGI_FORMAT_R16_UINT : DXGI_FORMAT_R32_UINT);
 
 	uint32 IndexCount = GetVertexCountForPrimitiveCount(NumPrimitives,PrimitiveType);
@@ -1767,10 +1765,10 @@ void FD3D11DynamicRHI::RHIDrawIndexedPrimitive(FRHIIndexBuffer* IndexBufferRHI,i
 	EnableUAVOverlap();
 }
 
-void FD3D11DynamicRHI::RHIDrawIndexedPrimitiveIndirect(FRHIIndexBuffer* IndexBufferRHI, FRHIVertexBuffer* ArgumentBufferRHI,uint32 ArgumentOffset)
+void FD3D11DynamicRHI::RHIDrawIndexedPrimitiveIndirect(FRHIBuffer* IndexBufferRHI, FRHIBuffer* ArgumentBufferRHI, uint32 ArgumentOffset)
 {
-	FD3D11IndexBuffer* IndexBuffer = ResourceCast(IndexBufferRHI);
-	FD3D11VertexBuffer* ArgumentBuffer = ResourceCast(ArgumentBufferRHI);
+	FD3D11Buffer* IndexBuffer = ResourceCast(IndexBufferRHI);
+	FD3D11Buffer* ArgumentBuffer = ResourceCast(ArgumentBufferRHI);
 
 	RHI_DRAW_CALL_INC();
 
@@ -2124,7 +2122,7 @@ FStagingBufferRHIRef FD3D11DynamicRHI::RHICreateStagingBuffer()
 
 void FD3D11DynamicRHI::RHICopyToStagingBuffer(FRHIVertexBuffer* SourceBufferRHI, FRHIStagingBuffer* StagingBufferRHI, uint32 Offset, uint32 NumBytes)
 {
-	FD3D11VertexBuffer* SourceBuffer = ResourceCast(SourceBufferRHI);
+	FD3D11Buffer* SourceBuffer = ResourceCast(SourceBufferRHI);
 	FD3D11StagingBuffer* StagingBuffer = ResourceCast(StagingBufferRHI);
 	if (StagingBuffer)
 	{
diff --git a/Engine/Source/Runtime/Windows/D3D11RHI/Private/D3D11IndexBuffer.cpp b/Engine/Source/Runtime/Windows/D3D11RHI/Private/D3D11IndexBuffer.cpp
--- a/Engine/Source/Runtime/Windows/D3D11RHI/Private/D3D11IndexBuffer.cpp
+++ b/Engine/Source/Runtime/Windows/D3D11RHI/Private/D3D11IndexBuffer.cpp
@@ -6,11 +6,13 @@
 
 #include "D3D11RHIPrivate.h"
 
-FIndexBufferRHIRef FD3D11DynamicRHI::RHICreateIndexBuffer(uint32 Stride,uint32 Size,uint32 InUsage, ERHIAccess InResourceState, FRHIResourceCreateInfo& CreateInfo)
+FIndexBufferRHIRef FD3D11DynamicRHI::RHICreateIndexBuffer(uint32 Stride, uint32 Size, uint32 InUsage, ERHIAccess InResourceState, FRHIResourceCreateInfo& CreateInfo)
 {
+	InUsage |= BUF_IndexBuffer;
+
 	if (CreateInfo.bWithoutNativeResource)
 	{
-		return new FD3D11IndexBuffer();
+		return new FD3D11Buffer();
 	}
 
 	// Explicitly check that the size is nonzero before allowing CreateIndexBuffer to opaquely fail.
@@ -75,7 +77,7 @@ FIndexBufferRHIRef FD3D11DynamicRHI::RHICreateIndexBuffer(uint32 Stride,uint32 S
 		CreateInfo.ResourceArray->Discard();
 	}
 
-	return new FD3D11IndexBuffer(IndexBufferResource, Stride, Size, InUsage);
+	return new FD3D11Buffer(IndexBufferResource, Size, InUsage, Stride);
 }
 
 FIndexBufferRHIRef FD3D11DynamicRHI::CreateIndexBuffer_RenderThread(
@@ -91,7 +93,7 @@ FIndexBufferRHIRef FD3D11DynamicRHI::CreateIndexBuffer_RenderThread(
 
 void* FD3D11DynamicRHI::LockIndexBuffer_BottomOfPipe(FRHICommandListImmediate& RHICmdList, FRHIIndexBuffer* IndexBufferRHI, uint32 Offset, uint32 Size, EResourceLockMode LockMode)
 {
-	FD3D11IndexBuffer* IndexBuffer = ResourceCast(IndexBufferRHI);
+	FD3D11Buffer* IndexBuffer = ResourceCast(IndexBufferRHI);
 	// If this resource is bound to the device, unbind it
 	ConditionalClearShaderResource(IndexBuffer, true);
 
@@ -155,7 +157,7 @@ void* FD3D11DynamicRHI::LockIndexBuffer_BottomOfPipe(FRHICommandListImmediate& R
 
 void FD3D11DynamicRHI::UnlockIndexBuffer_BottomOfPipe(FRHICommandListImmediate& RHICmdList, FRHIIndexBuffer* IndexBufferRHI)
 {
-	FD3D11IndexBuffer* IndexBuffer = ResourceCast(IndexBufferRHI);
+	FD3D11Buffer* IndexBuffer = ResourceCast(IndexBufferRHI);
 
 	// Determine whether the index buffer is dynamic or not.
 	D3D11_BUFFER_DESC Desc;
@@ -191,17 +193,17 @@ void FD3D11DynamicRHI::UnlockIndexBuffer_BottomOfPipe(FRHICommandListImmediate&
 	}
 }
 
-void FD3D11DynamicRHI::RHITransferIndexBufferUnderlyingResource(FRHIIndexBuffer* DestIndexBuffer, FRHIIndexBuffer* SrcIndexBuffer)
+void FD3D11DynamicRHI::RHITransferBufferUnderlyingResource(FRHIBuffer* DestBuffer, FRHIBuffer* SrcBuffer)
 {
-	check(DestIndexBuffer);
-	FD3D11IndexBuffer* Dest = ResourceCast(DestIndexBuffer);
-	if (!SrcIndexBuffer)
+	check(DestBuffer);
+	FD3D11Buffer* Dest = ResourceCast(DestBuffer);
+	if (!SrcBuffer)
 	{
 		Dest->ReleaseUnderlyingResource();
 	}
 	else
 	{
-		FD3D11IndexBuffer* Src = ResourceCast(SrcIndexBuffer);
+		FD3D11Buffer* Src = ResourceCast(SrcBuffer);
 		Dest->Swap(*Src);
 	}
 }
diff --git a/Engine/Source/Runtime/Windows/D3D11RHI/Private/D3D11RHIPrivate.h b/Engine/Source/Runtime/Windows/D3D11RHI/Private/D3D11RHIPrivate.h
--- a/Engine/Source/Runtime/Windows/D3D11RHI/Private/D3D11RHIPrivate.h
+++ b/Engine/Source/Runtime/Windows/D3D11RHI/Private/D3D11RHIPrivate.h
@@ -439,25 +439,22 @@ public:
 	virtual FIndexBufferRHIRef RHICreateIndexBuffer(uint32 Stride, uint32 Size, uint32 InUsage, ERHIAccess InResourceState, FRHIResourceCreateInfo& CreateInfo) final override;
 	virtual void* LockIndexBuffer_BottomOfPipe(FRHICommandListImmediate& RHICmdList, FRHIIndexBuffer* IndexBuffer, uint32 Offset, uint32 Size, EResourceLockMode LockMode) final override;
 	virtual void UnlockIndexBuffer_BottomOfPipe(FRHICommandListImmediate& RHICmdList, FRHIIndexBuffer* IndexBuffer) final override;
-	virtual void RHITransferIndexBufferUnderlyingResource(FRHIIndexBuffer* DestIndexBuffer, FRHIIndexBuffer* SrcIndexBuffer) final override;
 	virtual FVertexBufferRHIRef RHICreateVertexBuffer(uint32 Size, uint32 InUsage, ERHIAccess InResourceState, FRHIResourceCreateInfo& CreateInfo) final override;
 	virtual void* LockVertexBuffer_BottomOfPipe(FRHICommandListImmediate& RHICmdList, FRHIVertexBuffer* VertexBuffer, uint32 Offset, uint32 SizeRHI, EResourceLockMode LockMode) final override;
 	virtual void UnlockVertexBuffer_BottomOfPipe(FRHICommandListImmediate& RHICmdList, FRHIVertexBuffer* VertexBuffer) final override;
 	virtual void RHICopyVertexBuffer(FRHIVertexBuffer* SourceBuffer, FRHIVertexBuffer* DestBuffer) final override;
-	virtual void RHITransferVertexBufferUnderlyingResource(FRHIVertexBuffer* DestVertexBuffer, FRHIVertexBuffer* SrcVertexBuffer) final override;
+	virtual void RHITransferBufferUnderlyingResource(FRHIBuffer* DestBuffer, FRHIBuffer* SrcBuffer) final override;
 	virtual FStructuredBufferRHIRef RHICreateStructuredBuffer(uint32 Stride, uint32 Size, uint32 InUsage, ERHIAccess InResourceState, FRHIResourceCreateInfo& CreateInfo) final override;
 	virtual void* LockStructuredBuffer_BottomOfPipe(FRHICommandListImmediate& RHICmdList, FRHIStructuredBuffer* StructuredBuffer, uint32 Offset, uint32 SizeRHI, EResourceLockMode LockMode) final override;
 	virtual void UnlockStructuredBuffer_BottomOfPipe(FRHICommandListImmediate& RHICmdList, FRHIStructuredBuffer* StructuredBuffer) final override;
-	virtual FUnorderedAccessViewRHIRef RHICreateUnorderedAccessView(FRHIStructuredBuffer* StructuredBuffer, bool bUseUAVCounter, bool bAppendBuffer) final override;
+	virtual FUnorderedAccessViewRHIRef RHICreateUnorderedAccessView(FRHIBuffer* Buffer, bool bUseUAVCounter, bool bAppendBuffer) final override;
 	virtual FUnorderedAccessViewRHIRef RHICreateUnorderedAccessView(FRHITexture* Texture, uint32 MipLevel) final override;
-	virtual FUnorderedAccessViewRHIRef RHICreateUnorderedAccessView(FRHIVertexBuffer* VertexBuffer, uint8 Format) final override;
-	virtual FUnorderedAccessViewRHIRef RHICreateUnorderedAccessView(FRHIIndexBuffer* IndexBuffer, uint8 Format) final override;
-	virtual FShaderResourceViewRHIRef RHICreateShaderResourceView(FRHIStructuredBuffer* StructuredBuffer) final override;
-	virtual FShaderResourceViewRHIRef RHICreateShaderResourceView(FRHIVertexBuffer* VertexBuffer, uint32 Stride, uint8 Format) final override;
+	virtual FUnorderedAccessViewRHIRef RHICreateUnorderedAccessView(FRHIBuffer* Buffer, uint8 Format) final override;
+	virtual FShaderResourceViewRHIRef RHICreateShaderResourceView(FRHIBuffer* Buffer) final override;
+	virtual FShaderResourceViewRHIRef RHICreateShaderResourceView(FRHIBuffer* Buffer, uint32 Stride, uint8 Format) final override;
 	virtual FShaderResourceViewRHIRef RHICreateShaderResourceView(const FShaderResourceViewInitializer& Initializer) final override;
-	virtual FShaderResourceViewRHIRef RHICreateShaderResourceView(FRHIIndexBuffer* Buffer) final override;
-	virtual void RHIUpdateShaderResourceView(FRHIShaderResourceView* SRV, FRHIVertexBuffer* VertexBuffer, uint32 Stride, uint8 Format) final override;
-	virtual void RHIUpdateShaderResourceView(FRHIShaderResourceView* SRV, FRHIIndexBuffer* IndexBuffer) final override;
+	virtual void RHIUpdateShaderResourceView(FRHIShaderResourceView* SRV, FRHIBuffer* VertexBuffer, uint32 Stride, uint8 Format) final override;
+	virtual void RHIUpdateShaderResourceView(FRHIShaderResourceView* SRV, FRHIBuffer* IndexBuffer) final override;
 	virtual uint64 RHICalcTexture2DPlatformSize(uint32 SizeX, uint32 SizeY, uint8 Format, uint32 NumMips, uint32 NumSamples, ETextureCreateFlags Flags, const FRHIResourceCreateInfo& CreateInfo, uint32& OutAlign) final override;
 	virtual uint64 RHICalcTexture3DPlatformSize(uint32 SizeX, uint32 SizeY, uint32 SizeZ, uint8 Format, uint32 NumMips, ETextureCreateFlags Flags, const FRHIResourceCreateInfo& CreateInfo, uint32& OutAlign) final override;
 	virtual uint64 RHICalcTextureCubePlatformSize(uint32 Size, uint8 Format, uint32 NumMips, ETextureCreateFlags Flags, const FRHIResourceCreateInfo& CreateInfo, uint32& OutAlign) final override;
@@ -530,7 +527,7 @@ public:
 
 	virtual void RHISetComputeShader(FRHIComputeShader* ComputeShader) final override;
 	virtual void RHIDispatchComputeShader(uint32 ThreadGroupCountX, uint32 ThreadGroupCountY, uint32 ThreadGroupCountZ) final override;
-	virtual void RHIDispatchIndirectComputeShader(FRHIVertexBuffer* ArgumentBuffer, uint32 ArgumentOffset) final override;
+	virtual void RHIDispatchIndirectComputeShader(FRHIBuffer* ArgumentBuffer, uint32 ArgumentOffset) final override;
 	virtual void RHIBeginUAVOverlap() final override;
 	virtual void RHIEndUAVOverlap() final override;
 	virtual void RHISetMultipleViewports(uint32 Count, const FViewportBounds* Data) final override;
@@ -582,10 +579,10 @@ public:
 	void InternalSetUAVPS(uint32 BindIndex, FD3D11UnorderedAccessView* UnorderedAccessViewRHI);
 	void SetRenderTargetsAndClear(const FRHISetRenderTargetsInfo& RenderTargetsInfo);
 	virtual void RHIDrawPrimitive(uint32 BaseVertexIndex, uint32 NumPrimitives, uint32 NumInstances) final override;
-	virtual void RHIDrawPrimitiveIndirect(FRHIVertexBuffer* ArgumentBuffer, uint32 ArgumentOffset) final override;
-	virtual void RHIDrawIndexedIndirect(FRHIIndexBuffer* IndexBufferRHI, FRHIStructuredBuffer* ArgumentsBufferRHI, int32 DrawArgumentsIndex, uint32 NumInstances) final override;
-	virtual void RHIDrawIndexedPrimitive(FRHIIndexBuffer* IndexBuffer, int32 BaseVertexIndex, uint32 FirstInstance, uint32 NumVertices, uint32 StartIndex, uint32 NumPrimitives, uint32 NumInstances) final override;
-	virtual void RHIDrawIndexedPrimitiveIndirect(FRHIIndexBuffer* IndexBuffer, FRHIVertexBuffer* ArgumentBuffer, uint32 ArgumentOffset) final override;
+	virtual void RHIDrawPrimitiveIndirect(FRHIBuffer* ArgumentBuffer, uint32 ArgumentOffset) final override;
+	virtual void RHIDrawIndexedIndirect(FRHIBuffer* IndexBufferRHI, FRHIBuffer* ArgumentsBufferRHI, int32 DrawArgumentsIndex, uint32 NumInstances) final override;
+	virtual void RHIDrawIndexedPrimitive(FRHIBuffer* IndexBuffer, int32 BaseVertexIndex, uint32 FirstInstance, uint32 NumVertices, uint32 StartIndex, uint32 NumPrimitives, uint32 NumInstances) final override;
+	virtual void RHIDrawIndexedPrimitiveIndirect(FRHIBuffer* IndexBuffer, FRHIBuffer* ArgumentBuffer, uint32 ArgumentOffset) final override;
 	virtual void RHIEnableDepthBoundsTest(bool bEnable) final override
 	{
 		if (GSupportsDepthBoundsTest && StateCache.bDepthBoundsEnabled != bEnable)
@@ -642,7 +639,6 @@ public:
 	virtual FUnorderedAccessViewRHIRef RHICreateUnorderedAccessView_RenderThread(class FRHICommandListImmediate& RHICmdList, FRHITexture* Texture, uint32 MipLevel) final override;
 	virtual FUnorderedAccessViewRHIRef RHICreateUnorderedAccessView_RenderThread(class FRHICommandListImmediate& RHICmdList, FRHIVertexBuffer* VertexBuffer, uint8 Format) final override;
 	virtual FShaderResourceViewRHIRef RHICreateShaderResourceView_RenderThread(class FRHICommandListImmediate& RHICmdList, const FShaderResourceViewInitializer& Initializer) final override;
-	virtual FUnorderedAccessViewRHIRef RHICreateUnorderedAccessView_RenderThread(class FRHICommandListImmediate& RHICmdList, FRHIIndexBuffer* IndexBuffer, uint8 Format) final override;
 	virtual FShaderResourceViewRHIRef RHICreateShaderResourceView_RenderThread(class FRHICommandListImmediate& RHICmdList, FRHITexture* Texture, const FRHITextureSRVCreateInfo& CreateInfo) final override;
 	virtual FShaderResourceViewRHIRef RHICreateShaderResourceView_RenderThread(class FRHICommandListImmediate& RHICmdList, FRHIStructuredBuffer* StructuredBuffer) final override;
 	virtual FTextureCubeRHIRef RHICreateTextureCube_RenderThread(class FRHICommandListImmediate& RHICmdList, uint32 Size, uint8 Format, uint32 NumMips, ETextureCreateFlags Flags, ERHIAccess InResourceState, FRHIResourceCreateInfo& CreateInfo) final override;
diff --git a/Engine/Source/Runtime/Windows/D3D11RHI/Private/D3D11StructuredBuffer.cpp b/Engine/Source/Runtime/Windows/D3D11RHI/Private/D3D11StructuredBuffer.cpp
--- a/Engine/Source/Runtime/Windows/D3D11RHI/Private/D3D11StructuredBuffer.cpp
+++ b/Engine/Source/Runtime/Windows/D3D11RHI/Private/D3D11StructuredBuffer.cpp
@@ -2,8 +2,10 @@
 
 #include "D3D11RHIPrivate.h"
 
-FStructuredBufferRHIRef FD3D11DynamicRHI::RHICreateStructuredBuffer(uint32 Stride,uint32 Size,uint32 InUsage, ERHIAccess InResourceState, FRHIResourceCreateInfo& CreateInfo)
+FStructuredBufferRHIRef FD3D11DynamicRHI::RHICreateStructuredBuffer(uint32 Stride, uint32 Size, uint32 InUsage, ERHIAccess InResourceState, FRHIResourceCreateInfo& CreateInfo)
 {
+	InUsage |= BUF_StructuredBuffer;
+
 	// Explicitly check that the size is nonzero before allowing CreateStructuredBuffer to opaquely fail.
 	check(Size > 0);
 	// Check for values that will cause D3D calls to fail
@@ -27,11 +29,6 @@ FStructuredBufferRHIRef FD3D11DynamicRHI::RHICreateStructuredBuffer(uint32 Strid
 		Desc.BindFlags |= D3D11_BIND_UNORDERED_ACCESS;
 	}
 
-	if (InUsage & BUF_StreamOutput)
-	{
-		Desc.BindFlags |= D3D11_BIND_STREAM_OUTPUT;
-	}
-
 	Desc.CPUAccessFlags = (InUsage & BUF_AnyDynamic) ? D3D11_CPU_ACCESS_WRITE : 0;
 	Desc.MiscFlags = 0;
 
@@ -89,7 +86,7 @@ FStructuredBufferRHIRef FD3D11DynamicRHI::RHICreateStructuredBuffer(uint32 Strid
 		CreateInfo.ResourceArray->Discard();
 	}
 
-	return new FD3D11StructuredBuffer(StructuredBufferResource,Stride,Size,InUsage);
+	return new FD3D11Buffer(StructuredBufferResource, Size, InUsage, Stride);
 }
 
 FStructuredBufferRHIRef FD3D11DynamicRHI::CreateStructuredBuffer_RenderThread(
@@ -105,7 +102,7 @@ FStructuredBufferRHIRef FD3D11DynamicRHI::CreateStructuredBuffer_RenderThread(
 
 void* FD3D11DynamicRHI::LockStructuredBuffer_BottomOfPipe(FRHICommandListImmediate& RHICmdList, FRHIStructuredBuffer* StructuredBufferRHI,uint32 Offset,uint32 Size,EResourceLockMode LockMode)
 {
-	FD3D11StructuredBuffer* StructuredBuffer = ResourceCast(StructuredBufferRHI);
+	FD3D11Buffer* StructuredBuffer = ResourceCast(StructuredBufferRHI);
 	
 	// If this resource is bound to the device, unbind it
 	ConditionalClearShaderResource(StructuredBuffer, true);
@@ -170,7 +167,7 @@ void* FD3D11DynamicRHI::LockStructuredBuffer_BottomOfPipe(FRHICommandListImmedia
 
 void FD3D11DynamicRHI::UnlockStructuredBuffer_BottomOfPipe(FRHICommandListImmediate& RHICmdList, FRHIStructuredBuffer* StructuredBufferRHI)
 {
-	FD3D11StructuredBuffer* StructuredBuffer = ResourceCast(StructuredBufferRHI);
+	FD3D11Buffer* StructuredBuffer = ResourceCast(StructuredBufferRHI);
 
 	// Determine whether the Structured buffer is dynamic or not.
 	D3D11_BUFFER_DESC Desc;
diff --git a/Engine/Source/Runtime/Windows/D3D11RHI/Private/D3D11Texture.cpp b/Engine/Source/Runtime/Windows/D3D11RHI/Private/D3D11Texture.cpp
--- a/Engine/Source/Runtime/Windows/D3D11RHI/Private/D3D11Texture.cpp
+++ b/Engine/Source/Runtime/Windows/D3D11RHI/Private/D3D11Texture.cpp
@@ -2724,8 +2724,8 @@ void FD3D11DynamicRHI::RHICopyBufferRegion(FRHIVertexBuffer* DstBuffer, uint64 D
 		return;
 	}
 
-	FD3D11VertexBuffer* DstBufferD3D11 = ResourceCast(DstBuffer);
-	FD3D11VertexBuffer* SrcBufferD3D11 = ResourceCast(SrcBuffer);
+	FD3D11Buffer* DstBufferD3D11 = ResourceCast(DstBuffer);
+	FD3D11Buffer* SrcBufferD3D11 = ResourceCast(SrcBuffer);
 
 	check(DstBufferD3D11 && SrcBufferD3D11);
 	check(DstOffset + NumBytes <= DstBuffer->GetSize() && SrcOffset + NumBytes <= SrcBuffer->GetSize());
diff --git a/Engine/Source/Runtime/Windows/D3D11RHI/Private/D3D11UAV.cpp b/Engine/Source/Runtime/Windows/D3D11RHI/Private/D3D11UAV.cpp
--- a/Engine/Source/Runtime/Windows/D3D11RHI/Private/D3D11UAV.cpp
+++ b/Engine/Source/Runtime/Windows/D3D11RHI/Private/D3D11UAV.cpp
@@ -35,12 +35,12 @@ static D3D11_UNORDERED_ACCESS_VIEW_DESC InitializeCommonUAVDesc(const D3D11_BUFF
 	return UAVDesc;
 }
 
-FUnorderedAccessViewRHIRef FD3D11DynamicRHI::RHICreateUnorderedAccessView(FRHIStructuredBuffer* StructuredBufferRHI, bool bUseUAVCounter, bool bAppendBuffer)
+FUnorderedAccessViewRHIRef FD3D11DynamicRHI::RHICreateUnorderedAccessView(FRHIBuffer* BufferRHI, bool bUseUAVCounter, bool bAppendBuffer)
 {
-	FD3D11StructuredBuffer* StructuredBuffer = ResourceCast(StructuredBufferRHI);
+	FD3D11Buffer* Buffer = ResourceCast(BufferRHI);
 
 	D3D11_BUFFER_DESC BufferDesc;
-	StructuredBuffer->Resource->GetDesc(&BufferDesc);
+	Buffer->Resource->GetDesc(&BufferDesc);
 
 	D3D11_UNORDERED_ACCESS_VIEW_DESC UAVDesc = InitializeCommonUAVDesc(BufferDesc, 0, BufferDesc.ByteWidth, PF_Unknown, true);
 
@@ -60,18 +60,18 @@ FUnorderedAccessViewRHIRef FD3D11DynamicRHI::RHICreateUnorderedAccessView(FRHISt
 	}
 
 	TRefCountPtr<ID3D11UnorderedAccessView> UnorderedAccessView;
-	VERIFYD3D11RESULT_EX(Direct3DDevice->CreateUnorderedAccessView(StructuredBuffer->Resource,&UAVDesc,(ID3D11UnorderedAccessView**)UnorderedAccessView.GetInitReference()), Direct3DDevice);
+	VERIFYD3D11RESULT_EX(Direct3DDevice->CreateUnorderedAccessView(Buffer->Resource, &UAVDesc, (ID3D11UnorderedAccessView**)UnorderedAccessView.GetInitReference()), Direct3DDevice);
 
-	return new FD3D11UnorderedAccessView(UnorderedAccessView,StructuredBuffer);
+	return new FD3D11UnorderedAccessView(UnorderedAccessView, Buffer);
 }
 
 FUnorderedAccessViewRHIRef FD3D11DynamicRHI::RHICreateUnorderedAccessView_RenderThread(
 	class FRHICommandListImmediate& RHICmdList,
-	FRHIStructuredBuffer* StructuredBuffer,
+	FRHIBuffer* Buffer,
 	bool bUseUAVCounter,
 	bool bAppendBuffer)
 {
-	return RHICreateUnorderedAccessView(StructuredBuffer, bUseUAVCounter, bAppendBuffer);
+	return RHICreateUnorderedAccessView(Buffer, bUseUAVCounter, bAppendBuffer);
 }
 
 FUnorderedAccessViewRHIRef FD3D11DynamicRHI::RHICreateUnorderedAccessView(FRHITexture* TextureRHI, uint32 MipLevel)
@@ -126,50 +126,34 @@ FUnorderedAccessViewRHIRef FD3D11DynamicRHI::RHICreateUnorderedAccessView_Render
 	return RHICreateUnorderedAccessView(Texture, MipLevel);
 }
 
-FUnorderedAccessViewRHIRef FD3D11DynamicRHI::RHICreateUnorderedAccessView(FRHIVertexBuffer* VertexBufferRHI, uint8 Format)
+FUnorderedAccessViewRHIRef FD3D11DynamicRHI::RHICreateUnorderedAccessView(FRHIBuffer* BufferRHI, uint8 Format)
 {
-	FD3D11VertexBuffer* VertexBuffer = ResourceCast(VertexBufferRHI);
+	FD3D11Buffer* Buffer = ResourceCast(BufferRHI);
 
 	D3D11_BUFFER_DESC BufferDesc;
-	VertexBuffer->Resource->GetDesc(&BufferDesc);
+	Buffer->Resource->GetDesc(&BufferDesc);
 
 	D3D11_UNORDERED_ACCESS_VIEW_DESC UAVDesc = InitializeCommonUAVDesc(BufferDesc, 0, BufferDesc.ByteWidth, (EPixelFormat)Format, false);
 
 	TRefCountPtr<ID3D11UnorderedAccessView> UnorderedAccessView;
-	VERIFYD3D11RESULT_EX(Direct3DDevice->CreateUnorderedAccessView(VertexBuffer->Resource,&UAVDesc,(ID3D11UnorderedAccessView**)UnorderedAccessView.GetInitReference()), Direct3DDevice);
+	VERIFYD3D11RESULT_EX(Direct3DDevice->CreateUnorderedAccessView(Buffer->Resource, &UAVDesc, (ID3D11UnorderedAccessView**)UnorderedAccessView.GetInitReference()), Direct3DDevice);
 
-	return new FD3D11UnorderedAccessView(UnorderedAccessView,VertexBuffer);
+	return new FD3D11UnorderedAccessView(UnorderedAccessView, Buffer);
 }
 
 FUnorderedAccessViewRHIRef FD3D11DynamicRHI::RHICreateUnorderedAccessView_RenderThread(
 	class FRHICommandListImmediate& RHICmdList,
-	FRHIVertexBuffer* VertexBuffer,
+	FRHIBuffer* Buffer,
 	uint8 Format)
 {
-	return RHICreateUnorderedAccessView(VertexBuffer, Format);
-}
-
-FUnorderedAccessViewRHIRef FD3D11DynamicRHI::RHICreateUnorderedAccessView(FRHIIndexBuffer* IndexBufferRHI, uint8 Format)
-{
-	FD3D11IndexBuffer* IndexBuffer = ResourceCast(IndexBufferRHI);
-
-	D3D11_BUFFER_DESC BufferDesc;
-	IndexBuffer->Resource->GetDesc(&BufferDesc);
-
-	D3D11_UNORDERED_ACCESS_VIEW_DESC UAVDesc = InitializeCommonUAVDesc(BufferDesc, 0, BufferDesc.ByteWidth, (EPixelFormat)Format, false);
-
-	TRefCountPtr<ID3D11UnorderedAccessView> UnorderedAccessView;
-	VERIFYD3D11RESULT_EX(Direct3DDevice->CreateUnorderedAccessView(IndexBuffer->Resource, &UAVDesc, (ID3D11UnorderedAccessView**)UnorderedAccessView.GetInitReference()), Direct3DDevice);
-
-	return new FD3D11UnorderedAccessView(UnorderedAccessView, IndexBuffer);
+	return RHICreateUnorderedAccessView(Buffer, Format);
 }
 
-FUnorderedAccessViewRHIRef FD3D11DynamicRHI::RHICreateUnorderedAccessView_RenderThread(
+FShaderResourceViewRHIRef FD3D11DynamicRHI::RHICreateShaderResourceView_RenderThread(
 	class FRHICommandListImmediate& RHICmdList,
-	FRHIIndexBuffer* IndexBuffer,
-	uint8 Format)
+	FRHIStructuredBuffer* StructuredBuffer)
 {
-	return RHICreateUnorderedAccessView(IndexBuffer, Format);
+	return RHICreateShaderResourceView(StructuredBuffer);
 }
 
 FShaderResourceViewRHIRef FD3D11DynamicRHI::RHICreateShaderResourceView(FRHIStructuredBuffer* StructuredBufferRHI)
@@ -177,13 +161,6 @@ FShaderResourceViewRHIRef FD3D11DynamicRHI::RHICreateShaderResourceView(FRHIStru
 	return FD3D11DynamicRHI::RHICreateShaderResourceView(FShaderResourceViewInitializer(StructuredBufferRHI));
 }
 
-FShaderResourceViewRHIRef FD3D11DynamicRHI::RHICreateShaderResourceView_RenderThread(
-	class FRHICommandListImmediate& RHICmdList,
-	FRHIStructuredBuffer* StructuredBuffer)
-{
-	return RHICreateShaderResourceView(StructuredBuffer);
-}
-
 uint64 FD3D11DynamicRHI::RHIGetMinimumAlignmentForBufferBackedSRV(EPixelFormat Format)
 {
 	return GPixelFormats[Format].BlockBytes;
@@ -238,40 +215,38 @@ static void CreateD3D11ShaderResourceViewOnBuffer(ID3D11Device* Direct3DDevice,
 	}
 }
 
-FShaderResourceViewRHIRef FD3D11DynamicRHI::RHICreateShaderResourceView(FRHIVertexBuffer* VertexBufferRHI, uint32 Stride, uint8 Format)
+FShaderResourceViewRHIRef FD3D11DynamicRHI::RHICreateShaderResourceView(FRHIBuffer* BufferRHI, uint32 Stride, uint8 Format)
 {
 	ensureMsgf(Stride == GPixelFormats[Format].BlockBytes, TEXT("provided stride: %i was not consitent with Pixelformat: %s"), Stride, GPixelFormats[Format].Name);
-	return FD3D11DynamicRHI::RHICreateShaderResourceView(FShaderResourceViewInitializer(VertexBufferRHI, EPixelFormat(Format)));
+	return FD3D11DynamicRHI::RHICreateShaderResourceView(FShaderResourceViewInitializer(BufferRHI, EPixelFormat(Format)));
 }
 
 FShaderResourceViewRHIRef FD3D11DynamicRHI::RHICreateShaderResourceView(const FShaderResourceViewInitializer& Initializer)
 {
+	FShaderResourceViewInitializer::FBufferShaderResourceViewInitializer Desc = Initializer.AsBufferSRV();
+
+	FD3D11Buffer* Buffer = ResourceCast(Desc.Buffer);
+
+	if (!Desc.Buffer || !Buffer->Resource)
+	{
+		return new FD3D11ShaderResourceView(nullptr, nullptr);
+	}
+
+	TRefCountPtr<ID3D11ShaderResourceView> ShaderResourceView;
+
 	switch (Initializer.GetType())
 	{
 		case FShaderResourceViewInitializer::EType::VertexBufferSRV:
 		{
-			FShaderResourceViewInitializer::FVertexBufferShaderResourceViewInitializer Desc = Initializer.AsVertexBufferSRV();
+			CreateD3D11ShaderResourceViewOnBuffer(Direct3DDevice, Buffer->Resource, Desc.StartOffsetBytes, Desc.NumElements, Desc.Format, ShaderResourceView.GetInitReference());
 
-			FD3D11VertexBuffer* VertexBuffer = ResourceCast(Desc.VertexBuffer);
-			if (!Desc.VertexBuffer || !VertexBuffer->Resource)
-			{
-				return new FD3D11ShaderResourceView(nullptr, nullptr);
-			}
-
-			TRefCountPtr<ID3D11ShaderResourceView> ShaderResourceView;
-			CreateD3D11ShaderResourceViewOnBuffer(Direct3DDevice, VertexBuffer->Resource, Desc.StartOffsetBytes, Desc.NumElements, Desc.Format, ShaderResourceView.GetInitReference());
-
-			return new FD3D11ShaderResourceView(ShaderResourceView, VertexBuffer);
+			return new FD3D11ShaderResourceView(ShaderResourceView, Buffer);
 		}
 
 		case FShaderResourceViewInitializer::EType::StructuredBufferSRV:
 		{
-			FShaderResourceViewInitializer::FStructuredBufferShaderResourceViewInitializer Desc = Initializer.AsStructuredBufferSRV();
-
-			FD3D11StructuredBuffer* StructuredBuffer = ResourceCast(Desc.StructuredBuffer);
-
 			D3D11_BUFFER_DESC BufferDesc;
-			StructuredBuffer->Resource->GetDesc(&BufferDesc);
+			Buffer->Resource->GetDesc(&BufferDesc);
 
 			const bool bByteAccessBuffer = (BufferDesc.MiscFlags & D3D11_RESOURCE_MISC_BUFFER_ALLOW_RAW_VIEWS) != 0;
 
@@ -302,29 +277,18 @@ FShaderResourceViewRHIRef FD3D11DynamicRHI::RHICreateShaderResourceView(const FS
 				SRVDesc.Format = DXGI_FORMAT_UNKNOWN;
 			}
 
-			TRefCountPtr<ID3D11ShaderResourceView> ShaderResourceView;
-			VERIFYD3D11RESULT_EX(Direct3DDevice->CreateShaderResourceView(StructuredBuffer->Resource, &SRVDesc, (ID3D11ShaderResourceView**)ShaderResourceView.GetInitReference()), Direct3DDevice);
+			VERIFYD3D11RESULT_EX(Direct3DDevice->CreateShaderResourceView(Buffer->Resource, &SRVDesc, (ID3D11ShaderResourceView**)ShaderResourceView.GetInitReference()), Direct3DDevice);
 
-			return new FD3D11ShaderResourceView(ShaderResourceView, StructuredBuffer);
+			return new FD3D11ShaderResourceView(ShaderResourceView, Buffer);
 		}
 
 		case FShaderResourceViewInitializer::EType::IndexBufferSRV:
 		{
-			FShaderResourceViewInitializer::FIndexBufferShaderResourceViewInitializer Desc = Initializer.AsIndexBufferSRV();
-
-			if (!Desc.IndexBuffer)
-			{
-				return new FD3D11ShaderResourceView(nullptr, nullptr);
-			}
-			FD3D11IndexBuffer* Buffer = ResourceCast(Desc.IndexBuffer);
-			check(Buffer);
-			check(Buffer->Resource);
-
 			// The stride in bytes of the index buffer; must be 2 or 4
-			const uint32 Stride = Desc.IndexBuffer->GetStride();
+			const uint32 Stride = Desc.Buffer->GetStride();
 			check(Stride == 2 || Stride == 4);
 			const EPixelFormat Format = (Stride == 2) ? PF_R16_UINT : PF_R32_UINT;
-			TRefCountPtr<ID3D11ShaderResourceView> ShaderResourceView;
+
 			CreateD3D11ShaderResourceViewOnBuffer(Direct3DDevice, Buffer->Resource, Desc.StartOffsetBytes, Desc.NumElements, Format, ShaderResourceView.GetInitReference());
 
 			return new FD3D11ShaderResourceView(ShaderResourceView, Buffer);
@@ -358,53 +322,48 @@ FShaderResourceViewRHIRef FD3D11DynamicRHI::CreateShaderResourceView_RenderThrea
 	return RHICreateShaderResourceView(Initializer);
 }
 
-void FD3D11DynamicRHI::RHIUpdateShaderResourceView(FRHIShaderResourceView* SRV, FRHIVertexBuffer* VertexBufferRHI, uint32 Stride, uint8 Format)
+void FD3D11DynamicRHI::RHIUpdateShaderResourceView(FRHIShaderResourceView* SRV, FRHIBuffer* BufferRHI, uint32 Stride, uint8 Format)
 {
 	check(SRV);
 	FD3D11ShaderResourceView* SRVD3D11 = ResourceCast(SRV);
-	if (!VertexBufferRHI)
+	if (!BufferRHI)
 	{
 		SRVD3D11->Rename(nullptr, nullptr);
 	}
 	else
 	{
-		FD3D11VertexBuffer* VertexBuffer = ResourceCast(VertexBufferRHI);
-		check(VertexBuffer->Resource);
+		FD3D11Buffer* Buffer = ResourceCast(BufferRHI);
+		check(Buffer->Resource);
 
 		TRefCountPtr<ID3D11ShaderResourceView> ShaderResourceView;
-		CreateD3D11ShaderResourceViewOnBuffer(Direct3DDevice, VertexBuffer->Resource, 0, UINT32_MAX, Format, ShaderResourceView.GetInitReference());
+		CreateD3D11ShaderResourceViewOnBuffer(Direct3DDevice, Buffer->Resource, 0, UINT32_MAX, Format, ShaderResourceView.GetInitReference());
 
-		SRVD3D11->Rename(ShaderResourceView, VertexBuffer);
+		SRVD3D11->Rename(ShaderResourceView, Buffer);
 	}
 }
 
-void FD3D11DynamicRHI::RHIUpdateShaderResourceView(FRHIShaderResourceView* SRV, FRHIIndexBuffer* IndexBufferRHI)
+void FD3D11DynamicRHI::RHIUpdateShaderResourceView(FRHIShaderResourceView* SRV, FRHIBuffer* BufferRHI)
 {
 	check(SRV);
 	FD3D11ShaderResourceView* SRVD3D11 = ResourceCast(SRV);
-	if (!IndexBufferRHI)
+	if (!BufferRHI)
 	{
 		SRVD3D11->Rename(nullptr, nullptr);
 	}
 	else
 	{
-		FD3D11IndexBuffer* IndexBuffer = ResourceCast(IndexBufferRHI);
-		check(IndexBuffer->Resource);
+		FD3D11Buffer* Buffer = ResourceCast(BufferRHI);
+		check(Buffer->Resource);
 
-		const uint32 Stride = IndexBuffer->GetStride();
+		const uint32 Stride = Buffer->GetStride();
 		const EPixelFormat Format = Stride == 2 ? PF_R16_UINT : PF_R32_UINT;
 		TRefCountPtr<ID3D11ShaderResourceView> ShaderResourceView;
-		CreateD3D11ShaderResourceViewOnBuffer(Direct3DDevice, IndexBuffer->Resource, 0, UINT32_MAX, Format, ShaderResourceView.GetInitReference());
+		CreateD3D11ShaderResourceViewOnBuffer(Direct3DDevice, Buffer->Resource, 0, UINT32_MAX, Format, ShaderResourceView.GetInitReference());
 
-		SRVD3D11->Rename(ShaderResourceView, IndexBuffer);
+		SRVD3D11->Rename(ShaderResourceView, Buffer);
 	}
 }
 
-FShaderResourceViewRHIRef FD3D11DynamicRHI::RHICreateShaderResourceView(FRHIIndexBuffer* BufferRHI)
-{
-	return FD3D11DynamicRHI::RHICreateShaderResourceView(FShaderResourceViewInitializer(BufferRHI));
-}
-
 FShaderResourceViewRHIRef FD3D11DynamicRHI::CreateShaderResourceView_RenderThread(
 	class FRHICommandListImmediate& RHICmdList,
 	FRHIIndexBuffer* Buffer)
diff --git a/Engine/Source/Runtime/Windows/D3D11RHI/Private/D3D11VertexBuffer.cpp b/Engine/Source/Runtime/Windows/D3D11RHI/Private/D3D11VertexBuffer.cpp
--- a/Engine/Source/Runtime/Windows/D3D11RHI/Private/D3D11VertexBuffer.cpp
+++ b/Engine/Source/Runtime/Windows/D3D11RHI/Private/D3D11VertexBuffer.cpp
@@ -13,11 +13,13 @@ TAutoConsoleVariable<int32> GCVarUseSharedKeyedMutex(
 	TEXT("with the D3D11_RESOURCE_MISC_SHARED_KEYEDMUTEX flag instead of D3D11_RESOURCE_MISC_SHARED (default).\n"),
 	ECVF_Default);
 
-FVertexBufferRHIRef FD3D11DynamicRHI::RHICreateVertexBuffer(uint32 Size,uint32 InUsage, ERHIAccess InResourceState, FRHIResourceCreateInfo& CreateInfo)
+FVertexBufferRHIRef FD3D11DynamicRHI::RHICreateVertexBuffer(uint32 Size, uint32 InUsage, ERHIAccess InResourceState, FRHIResourceCreateInfo& CreateInfo)
 {
+	InUsage |= BUF_VertexBuffer;
+
 	if (CreateInfo.bWithoutNativeResource)
 	{
-		return new FD3D11VertexBuffer();
+		return new FD3D11Buffer();
 	}
 
 	// Explicitly check that the size is nonzero before allowing CreateVertexBuffer to opaquely fail.
@@ -42,12 +44,7 @@ FVertexBufferRHIRef FD3D11DynamicRHI::RHICreateVertexBuffer(uint32 Size,uint32 I
 		Desc.MiscFlags |= D3D11_RESOURCE_MISC_BUFFER_ALLOW_RAW_VIEWS;
 	}
 
-	if (InUsage & BUF_StreamOutput)
-	{
-		Desc.BindFlags |= D3D11_BIND_STREAM_OUTPUT;
-	}
-
-	if(InUsage & BUF_DrawIndirect)
+	if (InUsage & BUF_DrawIndirect)
 	{
 		Desc.MiscFlags |= D3D11_RESOURCE_MISC_DRAWINDIRECT_ARGS;
 	}
@@ -116,7 +113,7 @@ FVertexBufferRHIRef FD3D11DynamicRHI::RHICreateVertexBuffer(uint32 Size,uint32 I
 		CreateInfo.ResourceArray->Discard();
 	}
 
-	return new FD3D11VertexBuffer(VertexBufferResource,Size,InUsage);
+	return new FD3D11Buffer(VertexBufferResource, Size, InUsage, 0);
 }
 
 FVertexBufferRHIRef FD3D11DynamicRHI::CreateVertexBuffer_RenderThread(
@@ -133,7 +130,7 @@ void* FD3D11DynamicRHI::LockVertexBuffer_BottomOfPipe(FRHICommandListImmediate&
 {
 	check(Size > 0);
 
-	FD3D11VertexBuffer* VertexBuffer = ResourceCast(VertexBufferRHI);
+	FD3D11Buffer* VertexBuffer = ResourceCast(VertexBufferRHI);
 	
 	// If this resource is bound to the device, unbind it
 	ConditionalClearShaderResource(VertexBuffer, true);
@@ -213,7 +210,7 @@ void* FD3D11DynamicRHI::LockVertexBuffer_BottomOfPipe(FRHICommandListImmediate&
 
 void FD3D11DynamicRHI::UnlockVertexBuffer_BottomOfPipe(FRHICommandListImmediate& RHICmdList, FRHIVertexBuffer* VertexBufferRHI)
 {
-	FD3D11VertexBuffer* VertexBuffer = ResourceCast(VertexBufferRHI);
+	FD3D11Buffer* VertexBuffer = ResourceCast(VertexBufferRHI);
 
 	// Determine whether the vertex buffer is dynamic or not.
 	D3D11_BUFFER_DESC Desc;
@@ -251,8 +248,8 @@ void FD3D11DynamicRHI::UnlockVertexBuffer_BottomOfPipe(FRHICommandListImmediate&
 
 void FD3D11DynamicRHI::RHICopyVertexBuffer(FRHIVertexBuffer* SourceBufferRHI, FRHIVertexBuffer* DestBufferRHI)
 {
-	FD3D11VertexBuffer* SourceBuffer = ResourceCast(SourceBufferRHI);
-	FD3D11VertexBuffer* DestBuffer = ResourceCast(DestBufferRHI);
+	FD3D11Buffer* SourceBuffer = ResourceCast(SourceBufferRHI);
+	FD3D11Buffer* DestBuffer = ResourceCast(DestBufferRHI);
 
 	D3D11_BUFFER_DESC SourceBufferDesc;
 	SourceBuffer->Resource->GetDesc(&SourceBufferDesc);
@@ -266,18 +263,3 @@ void FD3D11DynamicRHI::RHICopyVertexBuffer(FRHIVertexBuffer* SourceBufferRHI, FR
 
 	GPUProfilingData.RegisterGPUWork(1);
 }
-
-void FD3D11DynamicRHI::RHITransferVertexBufferUnderlyingResource(FRHIVertexBuffer* DestVertexBuffer, FRHIVertexBuffer* SrcVertexBuffer)
-{
-	check(DestVertexBuffer);
-	FD3D11VertexBuffer* Dest = ResourceCast(DestVertexBuffer);
-	if (!SrcVertexBuffer)
-	{
-		Dest->ReleaseUnderlyingResource();
-	}
-	else
-	{
-		FD3D11VertexBuffer* Src = ResourceCast(SrcVertexBuffer);
-		Dest->Swap(*Src);
-	}
-}
diff --git a/Engine/Source/Runtime/Windows/D3D11RHI/Public/D3D11Resources.h b/Engine/Source/Runtime/Windows/D3D11RHI/Public/D3D11Resources.h
--- a/Engine/Source/Runtime/Windows/D3D11RHI/Public/D3D11Resources.h
+++ b/Engine/Source/Runtime/Windows/D3D11RHI/Public/D3D11Resources.h
@@ -662,22 +662,21 @@ private:
 	class FD3D11DynamicRHI* D3D11RHI;
 };
 
-/** Index buffer resource class that stores stride information. */
-class FD3D11IndexBuffer : public FRHIIndexBuffer, public FD3D11BaseShaderResource
+/** Buffer resource class. */
+class FD3D11Buffer : public FRHIBuffer, public FD3D11BaseShaderResource
 {
 public:
 
-	/** The index buffer resource */
 	TRefCountPtr<ID3D11Buffer> Resource;
 
-	FD3D11IndexBuffer() = default;
+	FD3D11Buffer() = default;
 
-	FD3D11IndexBuffer(ID3D11Buffer* InResource, uint32 InStride, uint32 InSize, uint32 InUsage)
-	: FRHIIndexBuffer(InStride,InSize,InUsage)
+	FD3D11Buffer(ID3D11Buffer* InResource, uint32 InSize, uint32 InUsage, uint32 InStride)
+	: FRHIBuffer(InSize, InUsage, InStride)
 	, Resource(InResource)
 	{}
 
-	virtual ~FD3D11IndexBuffer()
+	virtual ~FD3D11Buffer()
 	{
 		if (Resource)
 		{
@@ -685,93 +684,9 @@ public:
 		}
 	}
 
-	void Swap(FD3D11IndexBuffer& Other)
+	void Swap(FD3D11Buffer& SrcBuffer)
 	{
-		FRHIIndexBuffer::Swap(Other);
-		Resource.Swap(Other.Resource);
-	}
-
-	void ReleaseUnderlyingResource()
-	{
-		check(Resource);
-		UpdateBufferStats(Resource, false);
-		Resource = nullptr;
-		FRHIIndexBuffer::ReleaseUnderlyingResource();
-	}
-
-	// IRefCountedObject interface.
-	virtual uint32 AddRef() const
-	{
-		return FRHIResource::AddRef();
-	}
-	virtual uint32 Release() const
-	{
-		return FRHIResource::Release();
-	}
-	virtual uint32 GetRefCount() const
-	{
-		return FRHIResource::GetRefCount();
-	}
-};
-
-/** Structured buffer resource class. */
-class FD3D11StructuredBuffer : public FRHIStructuredBuffer, public FD3D11BaseShaderResource
-{
-public:
-
-	TRefCountPtr<ID3D11Buffer> Resource;
-
-	FD3D11StructuredBuffer(ID3D11Buffer* InResource, uint32 InStride, uint32 InSize, uint32 InUsage)
-	: FRHIStructuredBuffer(InStride,InSize,InUsage)
-	, Resource(InResource)
-	{
-	}
-
-	virtual ~FD3D11StructuredBuffer()
-	{
-		UpdateBufferStats(Resource, false);
-	}
-	
-	// IRefCountedObject interface.
-	virtual uint32 AddRef() const
-	{
-		return FRHIResource::AddRef();
-	}
-	virtual uint32 Release() const
-	{
-		return FRHIResource::Release();
-	}
-	virtual uint32 GetRefCount() const
-	{
-		return FRHIResource::GetRefCount();
-	}
-};
-
-/** Vertex buffer resource class. */
-class FD3D11VertexBuffer : public FRHIVertexBuffer, public FD3D11BaseShaderResource
-{
-public:
-
-	TRefCountPtr<ID3D11Buffer> Resource;
-
-	FD3D11VertexBuffer() = default;
-
-	FD3D11VertexBuffer(ID3D11Buffer* InResource, uint32 InSize, uint32 InUsage)
-	: FRHIVertexBuffer(InSize,InUsage)
-	, Resource(InResource)
-	{}
-
-	virtual ~FD3D11VertexBuffer()
-	{
-		if (Resource)
-		{
-			UpdateBufferStats(Resource, false);
-		}
-	}
-
-	void Swap(FD3D11VertexBuffer& SrcBuffer)
-	{
-		FRHIVertexBuffer::Swap(SrcBuffer);
+		FRHIBuffer::Swap(SrcBuffer);
 		Resource.Swap(SrcBuffer.Resource);
 	}
 
@@ -780,7 +695,7 @@ public:
 		check(Resource);
 		UpdateBufferStats(Resource, false);
 		Resource = nullptr;
-		FRHIVertexBuffer::ReleaseUnderlyingResource();
+		FRHIBuffer::ReleaseUnderlyingResource();
 	}
 
 	// IRefCountedObject interface.
@@ -936,19 +851,9 @@ struct TD3D11ResourceTraits<FRHIUniformBuffer>
 	typedef FD3D11UniformBuffer TConcreteType;
 };
 template<>
-struct TD3D11ResourceTraits<FRHIIndexBuffer>
-{
-	typedef FD3D11IndexBuffer TConcreteType;
-};
-template<>
-struct TD3D11ResourceTraits<FRHIStructuredBuffer>
-{
-	typedef FD3D11StructuredBuffer TConcreteType;
-};
-template<>
-struct TD3D11ResourceTraits<FRHIVertexBuffer>
+struct TD3D11ResourceTraits<FRHIBuffer>
 {
-	typedef FD3D11VertexBuffer TConcreteType;
+	typedef FD3D11Buffer TConcreteType;
 };
 template<>
 struct TD3D11ResourceTraits<FRHIStagingBuffer>
