Implemented capability query in D3D11 for stencil ref export from pixel shader. Also renamed D3D11_FEATURE_FORMAT_SUPPORT3 to D3D11_FEATURE_D3D11_OPTIONS3, as the previous name was incorrect.

diff --git a/Engine/Source/Runtime/Windows/D3D11RHI/Private/D3D11Device.cpp b/Engine/Source/Runtime/Windows/D3D11RHI/Private/D3D11Device.cpp
--- a/Engine/Source/Runtime/Windows/D3D11RHI/Private/D3D11Device.cpp
+++ b/Engine/Source/Runtime/Windows/D3D11RHI/Private/D3D11Device.cpp
@@ -411,12 +411,44 @@ uint32 FD3D11DynamicRHI::GetMaxMSAAQuality(uint32 SampleCount)
 static bool GFormatSupportsTypedUAVLoad[PF_MAX];
 
 // The D3D11 header is not recent enough to check support for this
-static const D3D11_FEATURE D3D11_FEATURE_FORMAT_SUPPORT3 = (D3D11_FEATURE)15;
+static const D3D11_FEATURE D3D11_FEATURE_D3D11_OPTIONS2 = (D3D11_FEATURE)14;
+static const D3D11_FEATURE D3D11_FEATURE_D3D11_OPTIONS3 = (D3D11_FEATURE)15;
+
 #if !PLATFORM_HOLOLENS
+
+typedef enum D3D11_CONSERVATIVE_RASTERIZATION_TIER
+{
+	D3D11_CONSERVATIVE_RASTERIZATION_NOT_SUPPORTED = 0,
+	D3D11_CONSERVATIVE_RASTERIZATION_TIER_1 = 1,
+	D3D11_CONSERVATIVE_RASTERIZATION_TIER_2 = 2,
+	D3D11_CONSERVATIVE_RASTERIZATION_TIER_3 = 3
+} D3D11_CONSERVATIVE_RASTERIZATION_TIER;
+
+typedef enum D3D11_TILED_RESOURCES_TIER
+{
+	D3D11_TILED_RESOURCES_NOT_SUPPORTED = 0,
+	D3D11_TILED_RESOURCES_TIER_1 = 1,
+	D3D11_TILED_RESOURCES_TIER_2 = 2,
+	D3D11_TILED_RESOURCES_TIER_3 = 3
+} D3D11_TILED_RESOURCES_TIER;
+
+typedef struct D3D11_FEATURE_DATA_D3D11_OPTIONS2
+{
+	BOOL PSSpecifiedStencilRefSupported;
+	BOOL TypedUAVLoadAdditionalFormats;
+	BOOL ROVsSupported;
+	D3D11_CONSERVATIVE_RASTERIZATION_TIER ConservativeRasterizationTier;
+	D3D11_TILED_RESOURCES_TIER TiledResourcesTier;
+	BOOL MapOnDefaultTextures;
+	BOOL StandardSwizzle;
+	BOOL UnifiedMemoryArchitecture;
+} D3D11_FEATURE_DATA_D3D11_OPTIONS2;
+
 typedef struct D3D11_FEATURE_DATA_D3D11_OPTIONS3
 {
 	BOOL VPAndRTArrayIndexFromAnyShaderFeedingRasterizer;
 } D3D11_FEATURE_DATA_D3D11_OPTIONS3;
+
 #endif
 
 void FD3D11DynamicRHI::SetupAfterDeviceCreation()
@@ -436,9 +468,19 @@ void FD3D11DynamicRHI::SetupAfterDeviceCreation()
 			D3D11RHI_ShouldAllowAsyncResourceCreation() ? TEXT("no driver support") : TEXT("disabled by user"));
 	}
 
+	{
+		D3D11_FEATURE_DATA_D3D11_OPTIONS2 Data;
+		HRESULT Result = Direct3DDevice->CheckFeatureSupport(D3D11_FEATURE_D3D11_OPTIONS2, &Data, sizeof(Data));
+		GRHISupportsStencilRefFromPixelShader = SUCCEEDED(Result) && Data.PSSpecifiedStencilRefSupported;
+		if (GRHISupportsStencilRefFromPixelShader)
+		{
+			UE_LOG(LogD3D11RHI, Log, TEXT("Stencil ref from pixel shader is supported"));
+		}
+	}
+
 	{
 		D3D11_FEATURE_DATA_D3D11_OPTIONS3 Data;
-		HRESULT Result = Direct3DDevice->CheckFeatureSupport(D3D11_FEATURE_FORMAT_SUPPORT3, &Data, sizeof(Data));
+		HRESULT Result = Direct3DDevice->CheckFeatureSupport(D3D11_FEATURE_D3D11_OPTIONS3, &Data, sizeof(Data));
 		GRHISupportsArrayIndexFromAnyShader = SUCCEEDED(Result) && Data.VPAndRTArrayIndexFromAnyShaderFeedingRasterizer;
 		if (GRHISupportsArrayIndexFromAnyShader)
 		{
