Disable __try/__finally statement for Clang build to prevent compiler crash.

#rnx

diff --git a/Engine/Source/Runtime/D3D12RHI/Private/D3D12Texture.h b/Engine/Source/Runtime/D3D12RHI/Private/D3D12Texture.h
--- a/Engine/Source/Runtime/D3D12RHI/Private/D3D12Texture.h
+++ b/Engine/Source/Runtime/D3D12RHI/Private/D3D12Texture.h
@@ -6,7 +6,7 @@ D3D12Texture.h: Implementation of D3D12 Texture
 #pragma once
 
 /** If true, guard texture creates with SEH to log more information about a driver crash we are seeing during texture streaming. */
-#define GUARDED_TEXTURE_CREATES (PLATFORM_WINDOWS && !(UE_BUILD_SHIPPING || UE_BUILD_TEST))
+#define GUARDED_TEXTURE_CREATES (PLATFORM_WINDOWS && !(UE_BUILD_SHIPPING || UE_BUILD_TEST || PLATFORM_COMPILER_CLANG))
 
 
 void SafeCreateTexture2D(FD3D12Device* pDevice, 
diff --git a/Engine/Source/Runtime/Windows/D3D11RHI/Private/D3D11Texture.cpp b/Engine/Source/Runtime/Windows/D3D11RHI/Private/D3D11Texture.cpp
--- a/Engine/Source/Runtime/Windows/D3D11RHI/Private/D3D11Texture.cpp
+++ b/Engine/Source/Runtime/Windows/D3D11RHI/Private/D3D11Texture.cpp
@@ -507,7 +507,7 @@ DXGI_FORMAT FD3D11DynamicRHI::GetPlatformTextureResourceFormat(DXGI_FORMAT InFor
 }
 
 /** If true, guard texture creates with SEH to log more information about a driver crash we are seeing during texture streaming. */
-#define GUARDED_TEXTURE_CREATES (PLATFORM_WINDOWS && !(UE_BUILD_SHIPPING || UE_BUILD_TEST))
+#define GUARDED_TEXTURE_CREATES (PLATFORM_WINDOWS && !(UE_BUILD_SHIPPING || UE_BUILD_TEST || PLATFORM_COMPILER_CLANG))
 
 /**
  * Creates a 2D texture optionally guarded by a structured exception handler.
