Removing extraneous feature level checks from the DX RHIs. Since we require 11_0 as a minimum, there's no need to have checks for if we're running with 11_0 or higher.

diff --git a/Engine/Source/Runtime/D3D12RHI/Private/D3D12RenderTarget.cpp b/Engine/Source/Runtime/D3D12RHI/Private/D3D12RenderTarget.cpp
--- a/Engine/Source/Runtime/D3D12RHI/Private/D3D12RenderTarget.cpp
+++ b/Engine/Source/Runtime/D3D12RHI/Private/D3D12RenderTarget.cpp
@@ -238,8 +238,7 @@ void FD3D12CommandContext::RHICopyToResolveTarget(FRHITexture* SourceTextureRHI,
 				GetParentDevice()->RegisterGPUWork();
 			}
 
-			if (FeatureLevel >= D3D_FEATURE_LEVEL_11_0
-				&& DestTexture2D->GetDepthStencilView(FExclusiveDepthStencil::DepthWrite_StencilWrite)
+			if (DestTexture2D->GetDepthStencilView(FExclusiveDepthStencil::DepthWrite_StencilWrite)
 				&& SourceTextureRHI->IsMultisampled()
 				&& !DestTextureRHI->IsMultisampled())
 			{
diff --git a/Engine/Source/Runtime/D3D12RHI/Private/D3D12Texture.cpp b/Engine/Source/Runtime/D3D12RHI/Private/D3D12Texture.cpp
--- a/Engine/Source/Runtime/D3D12RHI/Private/D3D12Texture.cpp
+++ b/Engine/Source/Runtime/D3D12RHI/Private/D3D12Texture.cpp
@@ -752,9 +752,6 @@ TD3D12Texture2D<BaseResourceType>* FD3D12DynamicRHI::CreateD3D12Texture2D(FRHICo
 		check(SizeZ <= GetMaxTextureArrayLayers());
 	}
 
-	// Render target allocation with UAV flag will silently fail in feature level 10
-	check(FeatureLevel >= D3D_FEATURE_LEVEL_11_0 || !(Flags & TexCreate_UAV));
-
 	SCOPE_CYCLE_COUNTER(STAT_D3D12CreateTextureTime);
 
 	const bool bSRGB = (Flags & TexCreate_SRGB) != 0;
diff --git a/Engine/Source/Runtime/Windows/D3D11RHI/Private/D3D11Device.cpp b/Engine/Source/Runtime/Windows/D3D11RHI/Private/D3D11Device.cpp
--- a/Engine/Source/Runtime/Windows/D3D11RHI/Private/D3D11Device.cpp
+++ b/Engine/Source/Runtime/Windows/D3D11RHI/Private/D3D11Device.cpp
@@ -96,7 +96,7 @@ FD3D11DynamicRHI::FD3D11DynamicRHI(IDXGIFactory1* InDXGIFactory1,D3D_FEATURE_LEV
 	GConfig->GetInt( TEXT( "TextureStreaming" ), TEXT( "PoolSizeVRAMPercentage" ), GPoolSizeVRAMPercentage, GEngineIni );	
 
 	// Initialize the RHI capabilities.
-	check(FeatureLevel == D3D_FEATURE_LEVEL_11_1 || FeatureLevel == D3D_FEATURE_LEVEL_11_0);
+	check(FeatureLevel >= D3D_FEATURE_LEVEL_11_0);
 
    	
 	TRefCountPtr<IDXGIFactory5> Factory5;
@@ -120,7 +120,7 @@ FD3D11DynamicRHI::FD3D11DynamicRHI(IDXGIFactory1* InDXGIFactory1,D3D_FEATURE_LEV
 			GMaxRHIShaderPlatform = SP_PCD3D_ES3_1;
 		}
 	}
-	else if(FeatureLevel == D3D_FEATURE_LEVEL_11_0 || FeatureLevel == D3D_FEATURE_LEVEL_11_1)
+	else
 	{
 		GMaxRHIFeatureLevel = ERHIFeatureLevel::SM5;
 		GMaxRHIShaderPlatform = SP_PCD3D_SM5;
@@ -208,15 +208,12 @@ FD3D11DynamicRHI::FD3D11DynamicRHI(IDXGIFactory1* InDXGIFactory1,D3D_FEATURE_LEV
 	GPixelFormats[PF_NV12].PlatformFormat = DXGI_FORMAT_NV12;
 	GPixelFormats[PF_NV12].Supported = true;
 
-	if (FeatureLevel >= D3D_FEATURE_LEVEL_11_0)
-	{
-		GSupportsSeparateRenderTargetBlendState = true;
-		GMaxTextureDimensions = D3D11_REQ_TEXTURE2D_U_OR_V_DIMENSION;
-		GMaxCubeTextureDimensions = D3D11_REQ_TEXTURECUBE_DIMENSION;
-		GMaxTextureArrayLayers = D3D11_REQ_TEXTURE2D_ARRAY_AXIS_DIMENSION;
-		GRHISupportsMSAADepthSampleAccess = true;
-		GRHISupportsRHIThread = !!EXPERIMENTAL_D3D11_RHITHREAD;
-	}
+	GSupportsSeparateRenderTargetBlendState = true;
+	GMaxTextureDimensions = D3D11_REQ_TEXTURE2D_U_OR_V_DIMENSION;
+	GMaxCubeTextureDimensions = D3D11_REQ_TEXTURECUBE_DIMENSION;
+	GMaxTextureArrayLayers = D3D11_REQ_TEXTURE2D_ARRAY_AXIS_DIMENSION;
+	GRHISupportsMSAADepthSampleAccess = true;
+	GRHISupportsRHIThread = !!EXPERIMENTAL_D3D11_RHITHREAD;
 
 	GMaxTextureMipCount = FMath::CeilLogTwo( GMaxTextureDimensions ) + 1;
 	GMaxTextureMipCount = FMath::Min<int32>( MAX_TEXTURE_MIP_COUNT, GMaxTextureMipCount );
@@ -466,12 +463,9 @@ void FD3D11DynamicRHI::SetupAfterDeviceCreation()
 		}
 	}
 
-	if (FeatureLevel >= D3D_FEATURE_LEVEL_11_0)
-	{
-		GFormatSupportsTypedUAVLoad[PF_R32_FLOAT] = true;
-		GFormatSupportsTypedUAVLoad[PF_R32_UINT] = true;
-		GFormatSupportsTypedUAVLoad[PF_R32_SINT] = true;
-	}
+	GFormatSupportsTypedUAVLoad[PF_R32_FLOAT] = true;
+	GFormatSupportsTypedUAVLoad[PF_R32_UINT] = true;
+	GFormatSupportsTypedUAVLoad[PF_R32_SINT] = true;
 }
 
 bool FD3D11DynamicRHI::RHIIsTypedUAVLoadSupported(EPixelFormat PixelFormat)
diff --git a/Engine/Source/Runtime/Windows/D3D11RHI/Private/D3D11RenderTarget.cpp b/Engine/Source/Runtime/Windows/D3D11RHI/Private/D3D11RenderTarget.cpp
--- a/Engine/Source/Runtime/Windows/D3D11RHI/Private/D3D11RenderTarget.cpp
+++ b/Engine/Source/Runtime/Windows/D3D11RHI/Private/D3D11RenderTarget.cpp
@@ -207,8 +207,7 @@ void FD3D11DynamicRHI::RHICopyToResolveTarget(FRHITexture* SourceTextureRHI, FRH
 		{
 			GPUProfilingData.RegisterGPUWork();
 		
-			if((FeatureLevel == D3D_FEATURE_LEVEL_11_0 || FeatureLevel == D3D_FEATURE_LEVEL_11_1)
-				&& DestTexture2D->GetDepthStencilView(FExclusiveDepthStencil::DepthWrite_StencilWrite)
+			if(DestTexture2D->GetDepthStencilView(FExclusiveDepthStencil::DepthWrite_StencilWrite)
 				&& SourceTextureRHI->IsMultisampled()
 				&& !DestTextureRHI->IsMultisampled())
 			{
diff --git a/Engine/Source/Runtime/Windows/D3D11RHI/Private/D3D11Texture.cpp b/Engine/Source/Runtime/Windows/D3D11RHI/Private/D3D11Texture.cpp
--- a/Engine/Source/Runtime/Windows/D3D11RHI/Private/D3D11Texture.cpp
+++ b/Engine/Source/Runtime/Windows/D3D11RHI/Private/D3D11Texture.cpp
@@ -582,9 +582,6 @@ TD3D11Texture2D<BaseResourceType>* FD3D11DynamicRHI::CreateD3D11Texture2D(uint32
 		checkf(SizeZ <= GetMaxTextureArrayLayers(), TEXT("Requested texture array size too large: %i, %i"), SizeZ, GetMaxTextureArrayLayers());
 	}
 
-	// Render target allocation with UAV flag will silently fail in feature level 10
-	check(FeatureLevel >= D3D_FEATURE_LEVEL_11_0 || !(Flags & TexCreate_UAV));
-
 	SCOPE_CYCLE_COUNTER(STAT_D3D11CreateTextureTime);
 
 	bool bPooledTexture = true;
@@ -948,13 +945,10 @@ TD3D11Texture2D<BaseResourceType>* FD3D11DynamicRHI::CreateD3D11Texture2D(uint32
 				// Create a read-only access views for the texture.
 				// Read-only DSVs are not supported in Feature Level 10 so 
 				// a dummy DSV is created in order reduce logic complexity at a higher-level.
-				if(Direct3DDevice->GetFeatureLevel() == D3D_FEATURE_LEVEL_11_0 || Direct3DDevice->GetFeatureLevel() == D3D_FEATURE_LEVEL_11_1)
+				DSVDesc.Flags = (AccessType & FExclusiveDepthStencil::DepthRead_StencilWrite) ? D3D11_DSV_READ_ONLY_DEPTH : 0;
+				if(HasStencilBits(DSVDesc.Format))
 				{
-					DSVDesc.Flags = (AccessType & FExclusiveDepthStencil::DepthRead_StencilWrite) ? D3D11_DSV_READ_ONLY_DEPTH : 0;
-					if(HasStencilBits(DSVDesc.Format))
-					{
-						DSVDesc.Flags |= (AccessType & FExclusiveDepthStencil::DepthWrite_StencilRead) ? D3D11_DSV_READ_ONLY_STENCIL : 0;
-					}
+					DSVDesc.Flags |= (AccessType & FExclusiveDepthStencil::DepthWrite_StencilRead) ? D3D11_DSV_READ_ONLY_STENCIL : 0;
 				}
 				VERIFYD3D11RESULT_EX(Direct3DDevice->CreateDepthStencilView(TextureResource,&DSVDesc,DepthStencilViews[AccessType].GetInitReference()), Direct3DDevice);
 			}
@@ -2425,13 +2419,10 @@ TD3D11Texture2D<BaseResourceType>* FD3D11DynamicRHI::CreateTextureFromResource(b
 			// Create a read-only access views for the texture.
 			// Read-only DSVs are not supported in Feature Level 10 so 
 			// a dummy DSV is created in order reduce logic complexity at a higher-level.
-			if(Direct3DDevice->GetFeatureLevel() == D3D_FEATURE_LEVEL_11_0 || Direct3DDevice->GetFeatureLevel() == D3D_FEATURE_LEVEL_11_1)
+			DSVDesc.Flags = (AccessType & FExclusiveDepthStencil::DepthRead_StencilWrite) ? D3D11_DSV_READ_ONLY_DEPTH : 0;
+			if (HasStencilBits(DSVDesc.Format))
 			{
-				DSVDesc.Flags = (AccessType & FExclusiveDepthStencil::DepthRead_StencilWrite) ? D3D11_DSV_READ_ONLY_DEPTH : 0;
-				if(HasStencilBits(DSVDesc.Format))
-				{
-					DSVDesc.Flags |= (AccessType & FExclusiveDepthStencil::DepthWrite_StencilRead) ? D3D11_DSV_READ_ONLY_STENCIL : 0;
-				}
+				DSVDesc.Flags |= (AccessType & FExclusiveDepthStencil::DepthWrite_StencilRead) ? D3D11_DSV_READ_ONLY_STENCIL : 0;
 			}
 			VERIFYD3D11RESULT_EX(Direct3DDevice->CreateDepthStencilView(TextureResource,&DSVDesc,DepthStencilViews[AccessType].GetInitReference()), Direct3DDevice);
 		}
diff --git a/Engine/Source/Runtime/Windows/D3D11RHI/Private/Windows/WindowsD3D11Device.cpp b/Engine/Source/Runtime/Windows/D3D11RHI/Private/Windows/WindowsD3D11Device.cpp
--- a/Engine/Source/Runtime/Windows/D3D11RHI/Private/Windows/WindowsD3D11Device.cpp
+++ b/Engine/Source/Runtime/Windows/D3D11RHI/Private/Windows/WindowsD3D11Device.cpp
@@ -258,9 +258,7 @@ static void SafeCreateDXGIFactory(IDXGIFactory1** DXGIFactory1, bool bWithDebug)
  */
 static D3D_FEATURE_LEVEL GetMinAllowedD3DFeatureLevel()
 {
-	// Default to 11.0
-	D3D_FEATURE_LEVEL AllowedFeatureLevel = D3D_FEATURE_LEVEL_11_0;
-	return AllowedFeatureLevel;
+	return D3D_FEATURE_LEVEL_11_0;
 }
 
 /**
@@ -269,9 +267,7 @@ static D3D_FEATURE_LEVEL GetMinAllowedD3DFeatureLevel()
  */
 static D3D_FEATURE_LEVEL GetMaxAllowedD3DFeatureLevel()
 {
-	// Default to 11.0
-	D3D_FEATURE_LEVEL AllowedFeatureLevel = D3D_FEATURE_LEVEL_11_0;
-	return AllowedFeatureLevel;
+	return D3D_FEATURE_LEVEL_11_1;
 }
 
 /**
@@ -802,11 +798,9 @@ bool FD3D11DynamicRHIModule::IsSupported()
 		FindAdapter();
 	}
 
-	// The hardware must support at least 10.0 (usually 11_0, 10_0 or 10_1).
+	// The hardware must support at least 11_0.
 	return ChosenAdapter.IsValid()
-		&& ChosenAdapter.MaxSupportedFeatureLevel != D3D_FEATURE_LEVEL_9_1
-		&& ChosenAdapter.MaxSupportedFeatureLevel != D3D_FEATURE_LEVEL_9_2
-		&& ChosenAdapter.MaxSupportedFeatureLevel != D3D_FEATURE_LEVEL_9_3;
+		&& ChosenAdapter.MaxSupportedFeatureLevel >= D3D_FEATURE_LEVEL_11_0;
 }
 
 
