Pacified MSVC's static analyzer. Warning C28199 was getting raised due to an implementation detail of UE_LOG(). MS' documentation states that this warning "can be mistaken if the variable is initialized outside of the function" which was the case here for DeviceDetails, which is initialised both by zero-initialisation and the earlier call the GetDeviceDetails().

#rnx

diff --git a/Engine/Source/Runtime/Windows/XAudio2/Private/XAudio2Device.cpp b/Engine/Source/Runtime/Windows/XAudio2/Private/XAudio2Device.cpp
--- a/Engine/Source/Runtime/Windows/XAudio2/Private/XAudio2Device.cpp
+++ b/Engine/Source/Runtime/Windows/XAudio2/Private/XAudio2Device.cpp
@@ -249,12 +249,16 @@ bool FXAudio2Device::InitializeHardware()
 		FXAudioDeviceProperties::DeviceDetails.OutputFormat.Format.nSamplesPerSec = SampleRate;
 	}
 
-	UE_LOG(LogInit, Log, TEXT( "XAudio2 using '%s' : %d channels at %g kHz using %d bits per sample (channel mask 0x%x)" ), 
-		FXAudioDeviceProperties::DeviceDetails.DisplayName,
-		FXAudioDeviceProperties::NumSpeakers, 
-		( float )SampleRate / 1000.0f, 
-		FXAudioDeviceProperties::DeviceDetails.OutputFormat.Format.wBitsPerSample,
-		(uint32)UE4_XAUDIO2_CHANNELMASK );
+	{
+		const TCHAR* DisplayName = FXAudioDeviceProperties::DeviceDetails.DisplayName;
+		int32 BitsPerSample = int32(FXAudioDeviceProperties::DeviceDetails.OutputFormat.Format.wBitsPerSample);
+		UE_LOG(LogInit, Log, TEXT( "XAudio2 using '%s' : %d channels at %g kHz using %d bits per sample (channel mask 0x%x)" ), 
+			DisplayName,
+			FXAudioDeviceProperties::NumSpeakers,
+			( float )SampleRate / 1000.0f,
+			BitsPerSample,
+			(uint32)UE4_XAUDIO2_CHANNELMASK );
+	}
 
 	if( !GetOutputMatrix( UE4_XAUDIO2_CHANNELMASK, FXAudioDeviceProperties::NumSpeakers ) )
 	{
