Cis fixes for xbox

diff --git a/Engine/Source/Runtime/Windows/AudioMixerXAudio2/Private/AudioMixerPlatformXAudio2.cpp b/Engine/Source/Runtime/Windows/AudioMixerXAudio2/Private/AudioMixerPlatformXAudio2.cpp
--- a/Engine/Source/Runtime/Windows/AudioMixerXAudio2/Private/AudioMixerPlatformXAudio2.cpp
+++ b/Engine/Source/Runtime/Windows/AudioMixerXAudio2/Private/AudioMixerPlatformXAudio2.cpp
@@ -1461,7 +1461,7 @@ namespace Audio
 		}
 
 #if WITH_XMA2 && USE_XMA2_FOR_STREAMING
-		if (InSoundWave.NumChannels <= 2)
+		if (InSoundWave.GetNumChannels() <= 2)
 		{
 			return XMA2_INFO_NEW();
 		}
