Fix crash running UE5 with -dx11 (UE-104870)

diff --git a/Engine/Source/Runtime/Windows/D3D11RHI/Private/D3D11Device.cpp b/Engine/Source/Runtime/Windows/D3D11RHI/Private/D3D11Device.cpp
--- a/Engine/Source/Runtime/Windows/D3D11RHI/Private/D3D11Device.cpp
+++ b/Engine/Source/Runtime/Windows/D3D11RHI/Private/D3D11Device.cpp
@@ -425,6 +425,7 @@ void FD3D11DynamicRHI::SetupAfterDeviceCreation()
 	}
 
 	{
+#if 0
 		D3D11_FEATURE_DATA_D3D11_OPTIONS2 Data;
 		HRESULT Result = Direct3DDevice->CheckFeatureSupport(D3D11_FEATURE_D3D11_OPTIONS2, &Data, sizeof(Data));
 		GRHISupportsStencilRefFromPixelShader = SUCCEEDED(Result) && Data.PSSpecifiedStencilRefSupported;
@@ -432,6 +433,9 @@ void FD3D11DynamicRHI::SetupAfterDeviceCreation()
 		{
 			UE_LOG(LogD3D11RHI, Log, TEXT("Stencil ref from pixel shader is supported"));
 		}
+#else
+		GRHISupportsStencilRefFromPixelShader = false; // this boolean is later used to choose a code path that requires DXIL shaders. Cannot set to true without fixing that first.
+#endif
 	}
 
 	{
