Better texture limit message

diff --git a/Engine/Source/Runtime/Windows/D3D11RHI/Private/D3D11Texture.cpp b/Engine/Source/Runtime/Windows/D3D11RHI/Private/D3D11Texture.cpp
--- a/Engine/Source/Runtime/Windows/D3D11RHI/Private/D3D11Texture.cpp
+++ b/Engine/Source/Runtime/Windows/D3D11RHI/Private/D3D11Texture.cpp
@@ -568,18 +568,18 @@ TD3D11Texture2D<BaseResourceType>* FD3D11DynamicRHI::CreateD3D11Texture2D(uint32
 
 	if (bCubeTexture)
 	{
-		checkf(SizeX <= GetMaxCubeTextureDimension(), TEXT("Requested cube texture size too large: %i, %i"), SizeX, GetMaxCubeTextureDimension());
+		checkf(SizeX <= GetMaxCubeTextureDimension(), TEXT("Requested cube texture size too large: %i, Max: %i"), SizeX, GetMaxCubeTextureDimension());
 		check(SizeX == SizeY);
 	}
 	else
 	{
-		checkf(SizeX <= GetMax2DTextureDimension(), TEXT("Requested texture2d x size too large: %i, %i"), SizeX, GetMax2DTextureDimension());
-		checkf(SizeY <= GetMax2DTextureDimension(), TEXT("Requested texture2d y size too large: %i, %i"), SizeY, GetMax2DTextureDimension());
+		checkf(SizeX <= GetMax2DTextureDimension(), TEXT("Requested texture2d x size too large: %i, Max: %i"), SizeX, GetMax2DTextureDimension());
+		checkf(SizeY <= GetMax2DTextureDimension(), TEXT("Requested texture2d y size too large: %i, Max: %i"), SizeY, GetMax2DTextureDimension());
 	}
 
 	if (bTextureArray)
 	{
-		checkf(SizeZ <= GetMaxTextureArrayLayers(), TEXT("Requested texture array size too large: %i, %i"), SizeZ, GetMaxTextureArrayLayers());
+		checkf(SizeZ <= GetMaxTextureArrayLayers(), TEXT("Requested texture array size too large: %i, Max: %i"), SizeZ, GetMaxTextureArrayLayers());
 	}
 
 	SCOPE_CYCLE_COUNTER(STAT_D3D11CreateTextureTime);
