Merged 'Fixed regression in D3D11 RHI CPU performance in development builds by removing FName lookups.'

diff --git a/Engine/Source/Runtime/Windows/D3D11RHI/Private/D3D11Commands.cpp b/Engine/Source/Runtime/Windows/D3D11RHI/Private/D3D11Commands.cpp
--- a/Engine/Source/Runtime/Windows/D3D11RHI/Private/D3D11Commands.cpp
+++ b/Engine/Source/Runtime/Windows/D3D11RHI/Private/D3D11Commands.cpp
@@ -1312,7 +1312,7 @@ FORCEINLINE void SetResource(FD3D11DynamicRHI* RESTRICT D3D11RHI, FD3D11StateCac
 }
 
 template <EShaderFrequency ShaderFrequency>
-inline int32 SetShaderResourcesFromBuffer_Surface(FD3D11DynamicRHI* RESTRICT D3D11RHI, FD3D11StateCache* RESTRICT StateCache, FD3D11UniformBuffer* RESTRICT Buffer, const uint32* RESTRICT ResourceMap, int32 BufferIndex, FName LayoutName)
+inline int32 SetShaderResourcesFromBuffer_Surface(FD3D11DynamicRHI* RESTRICT D3D11RHI, FD3D11StateCache* RESTRICT StateCache, FD3D11UniformBuffer* RESTRICT Buffer, const uint32* RESTRICT ResourceMap, int32 BufferIndex, const TCHAR* LayoutName)
 {
 	const TRefCountPtr<FRHIResource>* RESTRICT Resources = Buffer->ResourceTable.GetData();
 	const int32 NumResourcesInTable = Buffer->ResourceTable.Num();
@@ -1336,7 +1336,7 @@ inline int32 SetShaderResourcesFromBuffer_Surface(FD3D11DynamicRHI* RESTRICT D3D
 			FRHITexture* TextureRHI = (FRHITexture*)Resources[ResourceIndex].GetReference();
 			if (!TextureRHI)
 			{
-				UE_LOG(LogD3D11RHI, Fatal, TEXT("Null texture (resource %d bind %d) on UB Layout %s"), ResourceIndex, BindIndex, *LayoutName.ToString());
+				UE_LOG(LogD3D11RHI, Fatal, TEXT("Null texture (resource %d bind %d) on UB Layout %s"), ResourceIndex, BindIndex, LayoutName);
 			}
 			TextureRHI->SetLastRenderTime(CurrentTime);
 			FD3D11TextureBase* TextureD3D11 = GetD3D11TextureFromRHITexture(TextureRHI);
@@ -1367,7 +1367,7 @@ inline int32 SetShaderResourcesFromBuffer_Surface(FD3D11DynamicRHI* RESTRICT D3D
 
 
 template <EShaderFrequency ShaderFrequency>
-inline int32 SetShaderResourcesFromBufferUAVPS(FD3D11DynamicRHI* RESTRICT D3D11RHI, FD3D11StateCache* RESTRICT StateCache, FD3D11UniformBuffer* RESTRICT Buffer, const uint32* RESTRICT ResourceMap, int32 BufferIndex, FName LayoutName)
+inline int32 SetShaderResourcesFromBufferUAVPS(FD3D11DynamicRHI* RESTRICT D3D11RHI, FD3D11StateCache* RESTRICT StateCache, FD3D11UniformBuffer* RESTRICT Buffer, const uint32* RESTRICT ResourceMap, int32 BufferIndex, const TCHAR* LayoutName)
 {
 	const TRefCountPtr<FRHIResource>* RESTRICT Resources = Buffer->ResourceTable.GetData();
 	float CurrentTime = FApp::GetCurrentTime();
@@ -1386,7 +1386,7 @@ inline int32 SetShaderResourcesFromBufferUAVPS(FD3D11DynamicRHI* RESTRICT D3D11R
 			FD3D11UnorderedAccessView* UnorderedAccessViewRHI = (FD3D11UnorderedAccessView*)Resources[ResourceIndex].GetReference();
 			if (!UnorderedAccessViewRHI)
 			{
-				UE_LOG(LogD3D11RHI, Fatal, TEXT("Null UAV (resource %d bind %d) on UB Layout %s"), ResourceIndex, BindIndex, *LayoutName.ToString());
+				UE_LOG(LogD3D11RHI, Fatal, TEXT("Null UAV (resource %d bind %d) on UB Layout %s"), ResourceIndex, BindIndex, LayoutName);
 			}
 
 #if ENABLE_RHI_VALIDATION
@@ -1406,7 +1406,7 @@ inline int32 SetShaderResourcesFromBufferUAVPS(FD3D11DynamicRHI* RESTRICT D3D11R
 
 
 template <EShaderFrequency ShaderFrequency>
-inline int32 SetShaderResourcesFromBuffer_SRV(FD3D11DynamicRHI* RESTRICT D3D11RHI, FD3D11StateCache* RESTRICT StateCache, FD3D11UniformBuffer* RESTRICT Buffer, const uint32* RESTRICT ResourceMap, int32 BufferIndex, FName LayoutName)
+inline int32 SetShaderResourcesFromBuffer_SRV(FD3D11DynamicRHI* RESTRICT D3D11RHI, FD3D11StateCache* RESTRICT StateCache, FD3D11UniformBuffer* RESTRICT Buffer, const uint32* RESTRICT ResourceMap, int32 BufferIndex, const TCHAR* LayoutName)
 {
 	const TRefCountPtr<FRHIResource>* RESTRICT Resources = Buffer->ResourceTable.GetData();
 	float CurrentTime = FApp::GetCurrentTime();
@@ -1428,7 +1428,7 @@ inline int32 SetShaderResourcesFromBuffer_SRV(FD3D11DynamicRHI* RESTRICT D3D11RH
 			FD3D11ShaderResourceView* ShaderResourceViewRHI = (FD3D11ShaderResourceView*)Resources[ResourceIndex].GetReference();
 			if (!ShaderResourceViewRHI)
 			{
-				UE_LOG(LogD3D11RHI, Fatal, TEXT("Null SRV (resource %d bind %d) on UB Layout %s"), ResourceIndex, BindIndex, *LayoutName.ToString());
+				UE_LOG(LogD3D11RHI, Fatal, TEXT("Null SRV (resource %d bind %d) on UB Layout %s"), ResourceIndex, BindIndex, LayoutName);
 			}
 
 #if ENABLE_RHI_VALIDATION
@@ -1542,9 +1542,9 @@ void FD3D11DynamicRHI::SetResourcesFromTables(const ShaderType* RESTRICT Shader)
 #endif
 
 #if !(UE_BUILD_SHIPPING || UE_BUILD_TEST)
-		const FName LayoutName = *Buffer->GetLayout().GetDebugName();
+		const TCHAR* LayoutName = *Buffer->GetLayout().GetDebugName();
 #else 
-		const FName LayoutName = NAME_None;
+		const TCHAR* LayoutName = nullptr;
 #endif
 
 		// todo: could make this two pass: gather then set
@@ -1572,7 +1572,7 @@ int32 FD3D11DynamicRHI::SetUAVPSResourcesFromTables(const ShaderType* RESTRICT S
 		FD3D11UniformBuffer* Buffer = (FD3D11UniformBuffer*)BoundUniformBuffers[ShaderType::StaticFrequency][BufferIndex].GetReference();
 
 		check(BufferIndex < Shader->ShaderResourceTable.ResourceTableLayoutHashes.Num());
-		FName LayoutName = *Buffer->GetLayout().GetDebugName();
+		const TCHAR* LayoutName = *Buffer->GetLayout().GetDebugName();
 
 		if ((EShaderFrequency)ShaderType::StaticFrequency == SF_Pixel)
 		{
