Adding more logging on UAV creation failures.

diff --git a/Engine/Source/Runtime/Windows/D3D11RHI/Private/D3D11UAV.cpp b/Engine/Source/Runtime/Windows/D3D11RHI/Private/D3D11UAV.cpp
--- a/Engine/Source/Runtime/Windows/D3D11RHI/Private/D3D11UAV.cpp
+++ b/Engine/Source/Runtime/Windows/D3D11RHI/Private/D3D11UAV.cpp
@@ -60,7 +60,7 @@ FUnorderedAccessViewRHIRef FD3D11DynamicRHI::RHICreateUnorderedAccessView(FRHIBu
 	}
 
 	TRefCountPtr<ID3D11UnorderedAccessView> UnorderedAccessView;
-	VERIFYD3D11RESULT_EX(Direct3DDevice->CreateUnorderedAccessView(Buffer->Resource, &UAVDesc, (ID3D11UnorderedAccessView**)UnorderedAccessView.GetInitReference()), Direct3DDevice);
+	VERIFYD3D11CREATEVIEWRESULT(Direct3DDevice->CreateUnorderedAccessView(Buffer->Resource, &UAVDesc, (ID3D11UnorderedAccessView**)UnorderedAccessView.GetInitReference()), Direct3DDevice, BufferRHI, UAVDesc);
 
 	return new FD3D11UnorderedAccessView(UnorderedAccessView, Buffer);
 }
@@ -113,7 +113,7 @@ FUnorderedAccessViewRHIRef FD3D11DynamicRHI::RHICreateUnorderedAccessView(FRHITe
 	UAVDesc.Format = FindShaderResourceDXGIFormat((DXGI_FORMAT)GPixelFormats[TextureRHI->GetFormat()].PlatformFormat, false);
 
 	TRefCountPtr<ID3D11UnorderedAccessView> UnorderedAccessView;
-	VERIFYD3D11RESULT_EX(Direct3DDevice->CreateUnorderedAccessView(Texture->GetResource(),&UAVDesc,(ID3D11UnorderedAccessView**)UnorderedAccessView.GetInitReference()), Direct3DDevice);
+	VERIFYD3D11CREATEVIEWRESULT(Direct3DDevice->CreateUnorderedAccessView(Texture->GetResource(),&UAVDesc,(ID3D11UnorderedAccessView**)UnorderedAccessView.GetInitReference()), Direct3DDevice, TextureRHI, UAVDesc);
 
 	return new FD3D11UnorderedAccessView(UnorderedAccessView,Texture);
 }
@@ -136,7 +136,7 @@ FUnorderedAccessViewRHIRef FD3D11DynamicRHI::RHICreateUnorderedAccessView(FRHIBu
 	D3D11_UNORDERED_ACCESS_VIEW_DESC UAVDesc = InitializeCommonUAVDesc(BufferDesc, 0, BufferDesc.ByteWidth, (EPixelFormat)Format, false);
 
 	TRefCountPtr<ID3D11UnorderedAccessView> UnorderedAccessView;
-	VERIFYD3D11RESULT_EX(Direct3DDevice->CreateUnorderedAccessView(Buffer->Resource, &UAVDesc, (ID3D11UnorderedAccessView**)UnorderedAccessView.GetInitReference()), Direct3DDevice);
+	VERIFYD3D11CREATEVIEWRESULT(Direct3DDevice->CreateUnorderedAccessView(Buffer->Resource, &UAVDesc, (ID3D11UnorderedAccessView**)UnorderedAccessView.GetInitReference()), Direct3DDevice, BufferRHI, UAVDesc);
 
 	return new FD3D11UnorderedAccessView(UnorderedAccessView, Buffer);
 }
diff --git a/Engine/Source/Runtime/Windows/D3D11RHI/Private/D3D11Util.cpp b/Engine/Source/Runtime/Windows/D3D11RHI/Private/D3D11Util.cpp
--- a/Engine/Source/Runtime/Windows/D3D11RHI/Private/D3D11Util.cpp
+++ b/Engine/Source/Runtime/Windows/D3D11RHI/Private/D3D11Util.cpp
@@ -367,6 +367,60 @@ void VerifyD3D11ResizeViewportResult(
 		NewState.Format);
 }
 
+void VerifyD3D11CreateViewResult(HRESULT D3DResult, const ANSICHAR* Code, const ANSICHAR* Filename, uint32 Line, ID3D11Device* Device, FRHITexture* Texture, const D3D11_UNORDERED_ACCESS_VIEW_DESC& Desc)
+{
+	check(FAILED(D3DResult));
+
+	UINT FormatSupport = 0;
+	Device->CheckFormatSupport(Desc.Format, &FormatSupport);
+
+	const FString TextureName = Texture ? Texture->GetName().ToString() : FString(TEXT("<Unknown>"));
+	const TCHAR* ViewFormat = GetD3D11TextureFormatString(Desc.Format);
+
+	const FString& ErrorString = GetD3D11ErrorString(D3DResult, Device);
+
+	UE_LOG(LogD3D11RHI, Error, TEXT("%s failed with error %s\n at %s:%u\n (Name=%s, Format=%s(0x%08X), FormatSupport=0x%08X"),
+		ANSI_TO_TCHAR(Code), *ErrorString, ANSI_TO_TCHAR(Filename), Line,
+		*TextureName, ViewFormat, Desc.Format, FormatSupport);
+
+	TerminateOnDeviceRemoved(D3DResult, Device);
+	TerminateOnOutOfMemory(D3DResult, false);
+
+	UE_LOG(LogD3D11RHI, Fatal, TEXT("%s failed with error %s\n at %s:%u\n (Name=%s, Format=%s(0x%08X), FormatSupport=0x%08X"),
+		ANSI_TO_TCHAR(Code), *ErrorString, ANSI_TO_TCHAR(Filename), Line,
+		*TextureName, ViewFormat, Desc.Format, FormatSupport);
+}
+
+void VerifyD3D11CreateViewResult(HRESULT D3DResult, const ANSICHAR* Code, const ANSICHAR* Filename, uint32 Line, ID3D11Device* Device, FRHIBuffer* Buffer, const D3D11_UNORDERED_ACCESS_VIEW_DESC& Desc)
+{
+	check(FAILED(D3DResult));
+
+	UINT FormatSupport = 0;
+	Device->CheckFormatSupport(Desc.Format, &FormatSupport);
+
+	FString BufferName = FString(TEXT("<Unknown>"));
+#if ENABLE_RHI_VALIDATION
+	if (Buffer)
+	{
+		BufferName = Buffer->GetDebugName();
+	}
+#endif
+	const TCHAR* ViewFormat = GetD3D11TextureFormatString(Desc.Format);
+
+	const FString& ErrorString = GetD3D11ErrorString(D3DResult, Device);
+
+	UE_LOG(LogD3D11RHI, Error, TEXT("%s failed with error %s\n at %s:%u\n (Name=%s, Format=%s(0x%08X), FormatSupport=0x%08X"),
+		ANSI_TO_TCHAR(Code), *ErrorString, ANSI_TO_TCHAR(Filename), Line,
+		*BufferName, ViewFormat, Desc.Format, FormatSupport);
+
+	TerminateOnDeviceRemoved(D3DResult, Device);
+	TerminateOnOutOfMemory(D3DResult, false);
+
+	UE_LOG(LogD3D11RHI, Fatal, TEXT("%s failed with error %s\n at %s:%u\n (Name=%s, Format=%s(0x%08X), FormatSupport=0x%08X"),
+		ANSI_TO_TCHAR(Code), *ErrorString, ANSI_TO_TCHAR(Filename), Line,
+		*BufferName, ViewFormat, Desc.Format, FormatSupport);
+}
+
 void VerifyComRefCount(IUnknown* Object,int32 ExpectedRefs,const TCHAR* Code,const TCHAR* Filename,int32 Line)
 {
 	int32 NumRefs;
diff --git a/Engine/Source/Runtime/Windows/D3D11RHI/Public/D3D11Util.h b/Engine/Source/Runtime/Windows/D3D11RHI/Public/D3D11Util.h
--- a/Engine/Source/Runtime/Windows/D3D11RHI/Public/D3D11Util.h
+++ b/Engine/Source/Runtime/Windows/D3D11RHI/Public/D3D11Util.h
@@ -83,6 +83,8 @@ struct FD3D11ResizeViewportState
 extern D3D11RHI_API void VerifyD3D11ResizeViewportResult(HRESULT D3DResult, const ANSICHAR* Code, const ANSICHAR* Filename, uint32 Line,
 	const FD3D11ResizeViewportState& OldState, const FD3D11ResizeViewportState& NewState, ID3D11Device* Device);
 
+extern D3D11RHI_API void VerifyD3D11CreateViewResult(HRESULT D3DResult, const ANSICHAR* Code, const ANSICHAR* Filename, uint32 Line, ID3D11Device* Device, FRHITexture* Texture, const D3D11_UNORDERED_ACCESS_VIEW_DESC& Desc);
+extern D3D11RHI_API void VerifyD3D11CreateViewResult(HRESULT D3DResult, const ANSICHAR* Code, const ANSICHAR* Filename, uint32 Line, ID3D11Device* Device, FRHIBuffer* Buffer, const D3D11_UNORDERED_ACCESS_VIEW_DESC& Desc);
 
 /**
  * A macro for using VERIFYD3D11RESULT that automatically passes in the code and filename/line.
@@ -94,6 +96,8 @@ extern D3D11RHI_API void VerifyD3D11ResizeViewportResult(HRESULT D3DResult, cons
 #define VERIFYD3D11RESULT_NOEXIT(x)		{HRESULT hr = x; if (FAILED(hr)) { VerifyD3D11ResultNoExit(hr,#x,__FILE__,__LINE__, 0); }}
 #define VERIFYD3D11CREATETEXTURERESULT(x,UEFormat,SizeX,SizeY,SizeZ,Format,NumMips,Flags,Usage,CPUAccessFlags,MiscFlags,SampleCount,SampleQuality,SubResPtr,SubResPitch,SubResSlicePitch,Device) {HRESULT hr = x; if (FAILED(hr)) { VerifyD3D11CreateTextureResult(hr, UEFormat,#x,__FILE__,__LINE__,SizeX,SizeY,SizeZ,Format,NumMips,Flags,Usage,CPUAccessFlags,MiscFlags,SampleCount,SampleQuality,SubResPtr,SubResPitch,SubResSlicePitch,Device); }}
 #define VERIFYD3D11RESIZEVIEWPORTRESULT(x, OldState, NewState, Device) { HRESULT hr = x; if (FAILED(hr)) { VerifyD3D11ResizeViewportResult(hr, #x, __FILE__, __LINE__, OldState, NewState, Device); }}
+#define VERIFYD3D11CREATEVIEWRESULT(x, Device, Resource, Desc) {HRESULT hr = x; if (FAILED(hr)) { VerifyD3D11CreateViewResult(hr, #x, __FILE__, __LINE__, Device, Resource, Desc); }}
+
 /**
  * Checks that a COM object has the expected number of references.
  */
