Fixing UAVMask support.
- Removed OutputMask since nothing was using it anymore.
- Made UAVMask an optional output from FXC compiles, only d3d11 reads it.
- UAVMask now created from entire set of UAVs and not just UAVs in Uniform Buffers.

diff --git a/Engine/Shaders/Public/ShaderVersion.ush b/Engine/Shaders/Public/ShaderVersion.ush
--- a/Engine/Shaders/Public/ShaderVersion.ush
+++ b/Engine/Shaders/Public/ShaderVersion.ush
@@ -3,4 +3,4 @@
 // in Platform.ush (which should be included in any shader) it allows to invalidate the shader DDC.
 //
 // If you are merging streams and there is a conflict with this GUID you should make a new GUID rather than taking one or the other.
-// GUID = B1C09C8B-698B-404B-84C8-81AE0AA9340D
+// GUID = 6224F14D-6734-4B88-BF90-06B539F4EAC5
diff --git a/Engine/Source/Developer/Windows/ShaderFormatD3D/Private/D3DShaderCompiler.cpp b/Engine/Source/Developer/Windows/ShaderFormatD3D/Private/D3DShaderCompiler.cpp
--- a/Engine/Source/Developer/Windows/ShaderFormatD3D/Private/D3DShaderCompiler.cpp
+++ b/Engine/Source/Developer/Windows/ShaderFormatD3D/Private/D3DShaderCompiler.cpp
@@ -665,7 +665,6 @@ bool CompileAndProcessD3DShaderFXC(FString& PreprocessedShaderSource, const FStr
 		uint32 NumSRVs = 0;
 		uint32 NumCBs = 0;
 		uint32 NumUAVs = 0;
-		uint32 OutputMask = 0;
 
 		TArray<FString> UniformBufferNames;
 		TArray<FString> ShaderOutputs;
@@ -722,7 +721,6 @@ bool CompileAndProcessD3DShaderFXC(FString& PreprocessedShaderSource, const FStr
 					D3D11_SIGNATURE_PARAMETER_DESC ParamDescs[3];
 					D3D11_SIGNATURE_PARAMETER_DESC& ParamDesc = ParamDescs[1];
 					Reflector->GetOutputParameterDesc(Index, &ParamDesc);
-					OutputMask |= (1 << ParamDesc.Register);
 				}
 
 				bool bFoundUnused = false;
@@ -851,7 +849,23 @@ bool CompileAndProcessD3DShaderFXC(FString& PreprocessedShaderSource, const FStr
 				CompressedData = Shader;
 			}
 
-			FShaderCodePackedResourceCounts PackedResourceCounts = { bGlobalUniformBufferUsed, static_cast<uint8>(NumSamplers), static_cast<uint8>(NumSRVs), static_cast<uint8>(NumCBs), static_cast<uint8>(NumUAVs), static_cast<uint16>(OutputMask)};
+			// Add resource masks before the parameters are pulled for the uniform buffers
+			FShaderCodeResourceMasks ResourceMasks{};
+			for (const auto& Param : Output.ParameterMap.GetParameterMap())
+			{
+				const FParameterAllocation& ParamAlloc = Param.Value;
+				if (ParamAlloc.Type == EShaderParameterType::UAV)
+				{
+					ResourceMasks.UAVMask |= 1u << ParamAlloc.BaseIndex;
+				}
+			}
+
+			auto AddOptionalDataCallback = [&](FShaderCode& ShaderCode)
+			{
+				Output.ShaderCode.AddOptionalData(ResourceMasks);
+			};
+
+			FShaderCodePackedResourceCounts PackedResourceCounts = { bGlobalUniformBufferUsed, static_cast<uint8>(NumSamplers), static_cast<uint8>(NumSRVs), static_cast<uint8>(NumCBs), static_cast<uint8>(NumUAVs) };
 			GenerateFinalOutput(CompressedData,
 				Input, VendorExtensions,
 				UsedUniformBufferSlots, UniformBufferNames,
@@ -859,7 +873,7 @@ bool CompileAndProcessD3DShaderFXC(FString& PreprocessedShaderSource, const FStr
 				PackedResourceCounts, NumInstructions,
 				Output,
 				[](FMemoryWriter&){},
-				[](FShaderCode&){});
+				AddOptionalDataCallback);
 		}
 	}
 
diff --git a/Engine/Source/Developer/Windows/ShaderFormatD3D/Private/D3DShaderCompilerDXC.cpp b/Engine/Source/Developer/Windows/ShaderFormatD3D/Private/D3DShaderCompilerDXC.cpp
--- a/Engine/Source/Developer/Windows/ShaderFormatD3D/Private/D3DShaderCompilerDXC.cpp
+++ b/Engine/Source/Developer/Windows/ShaderFormatD3D/Private/D3DShaderCompilerDXC.cpp
@@ -877,7 +877,7 @@ bool CompileAndProcessD3DShaderDXC(FString& PreprocessedShaderSource,
 
 			//#todo-rco: Should compress ShaderCode?
 
-			FShaderCodePackedResourceCounts PackedResourceCounts = { bGlobalUniformBufferUsed, static_cast<uint8>(NumSamplers), static_cast<uint8>(NumSRVs), static_cast<uint8>(NumCBs), static_cast<uint8>(NumUAVs), 0 };
+			FShaderCodePackedResourceCounts PackedResourceCounts = { bGlobalUniformBufferUsed, static_cast<uint8>(NumSamplers), static_cast<uint8>(NumSRVs), static_cast<uint8>(NumCBs), static_cast<uint8>(NumUAVs) };
 			GenerateFinalOutput(ShaderBlob,
 				Input, VendorExtensions,
 				UsedUniformBufferSlots, UniformBufferNames,
diff --git a/Engine/Source/Runtime/RenderCore/Public/ShaderCore.h b/Engine/Source/Runtime/RenderCore/Public/ShaderCore.h
--- a/Engine/Source/Runtime/RenderCore/Public/ShaderCore.h
+++ b/Engine/Source/Runtime/RenderCore/Public/ShaderCore.h
@@ -524,7 +524,14 @@ struct FShaderCodePackedResourceCounts
 	uint8 NumSRVs;
 	uint8 NumCBs;
 	uint8 NumUAVs;
-	uint16 OutputMask; //Mask of rendertargets bound
+};
+
+struct FShaderCodeResourceMasks
+{
+	// for FindOptionalData() and AddOptionalData()
+	static const uint8 Key = 'm';
+
+	uint32 UAVMask; // Mask of UAVs bound
 };
 
 // if this changes you need to make sure all shaders get invalidated
diff --git a/Engine/Source/Runtime/Windows/D3D11RHI/Private/D3D11Shaders.cpp b/Engine/Source/Runtime/Windows/D3D11RHI/Private/D3D11Shaders.cpp
--- a/Engine/Source/Runtime/Windows/D3D11RHI/Private/D3D11Shaders.cpp
+++ b/Engine/Source/Runtime/Windows/D3D11RHI/Private/D3D11Shaders.cpp
@@ -16,55 +16,13 @@ static inline void ReadShaderOptionalData(FShaderCodeReader& InShaderCode, TShad
 {
 	auto PackedResourceCounts = InShaderCode.FindOptionalData<FShaderCodePackedResourceCounts>();
 	check(PackedResourceCounts);
-	OutShader.OutputMask = PackedResourceCounts->OutputMask;
-	
-	uint32 UAVMask = 0;
-	uint32 MinOffset = OutShader.ShaderResourceTable.UnorderedAccessViewMap.Num();
-	// If the token stream isn't empty, it has a length of at least 2, because it's always terminated with 0xffffffff. If it has
-	// a length of 2, it means it only contains an offset entry for a single uniform buffer, which must be 0, because there's nothing
-	// for it to offset into. Therefore, we only care about streams which have 3 or more tokens (including the terminator).
-	if (MinOffset > 2)
-	{
-		// Ignore the terminator.
-		--MinOffset;
-
-		// The token stream starts with a table of offsets, followed by a list of UAV bindings, like this:
-		//		O1 O2 O3 ... On B1 B2 B3 ... Bm
-		// The offsets indicate where to find the bindings for each active uniform buffer, relative to the start of the
-		// data (not the start of the binding list). Therefore, this buffer:
-		//		9 0 7 6 0 0 B1 B2 B3 B4
-		// means that:
-		//		buffer 0 starts at index 9 (B4)
-		//		buffers 1, 4 and 5 are empty (offsets are 0)
-		//		buffer 2 starts at index 7 (B2)
-		//		buffer 3 starts at index 6 (B1)
-		// We can also infer that buffer 2 has two elements in this case (B2 and B3). Since we don't know how many buffers there are,
-		// we will parse the data like this:
-		//		* read an offset from the start of the list
-		//		* process every binding from that offset until the end, since the bindings list is contiguous
-		//		* read the next offset, process bindings from that until the first biding we've already processed
-		//		* stop when the next offset is at a location we've already processed as a binding (we've found the first binding in that case).
-		for (uint32 BufferIdx = 0; BufferIdx < MinOffset; ++BufferIdx)
-		{
-			uint32 BufferOffset = OutShader.ShaderResourceTable.UnorderedAccessViewMap[BufferIdx];
-			if (BufferOffset == 0 || BufferOffset >= MinOffset)
-			{
-				continue;
-			}
 
-			for (uint32 ElemIdx = BufferOffset; ElemIdx < MinOffset; ++ElemIdx)
-			{
-				uint32 UAVBinding = OutShader.ShaderResourceTable.UnorderedAccessViewMap[ElemIdx];
-				const uint8 BindIndex = FRHIResourceTableEntry::GetBindIndex(UAVBinding);
-				UAVMask |= (1 << BindIndex);
-			}
-
-			MinOffset = BufferOffset;
-		}
+	OutShader.UAVMask = 0;
+	if (auto ResourceMasks = InShaderCode.FindOptionalData<FShaderCodeResourceMasks>())
+	{
+		OutShader.UAVMask = ResourceMasks->UAVMask;
 	}
-
-	OutShader.UAVMask = UAVMask;
-
+	
 	OutShader.bShaderNeedsGlobalConstantBuffer = PackedResourceCounts->bGlobalUniformBufferUsed;
 #if !(UE_BUILD_SHIPPING || UE_BUILD_TEST)
 	OutShader.ShaderName = InShaderCode.FindOptionalData(FShaderCodeName::Key);
diff --git a/Engine/Source/Runtime/Windows/D3D11RHI/Public/D3D11Resources.h b/Engine/Source/Runtime/Windows/D3D11RHI/Public/D3D11Resources.h
--- a/Engine/Source/Runtime/Windows/D3D11RHI/Public/D3D11Resources.h
+++ b/Engine/Source/Runtime/Windows/D3D11RHI/Public/D3D11Resources.h
@@ -46,7 +46,6 @@ struct FD3D11ShaderData
 	TArray<FShaderCodeVendorExtension>	VendorExtensions;
 	bool								bShaderNeedsGlobalConstantBuffer;
 	bool								bIsSm6Shader;
-	uint16								OutputMask;
 	uint16								UAVMask;
 };
 
