Fixed rampant FName allocation/copies in FD3D11DynamicRHI::SetUAVPSResourcesFromTables, matching behavior of FD3D11DynamicRHI::SetResourcesFromTables.

diff --git a/Engine/Source/Runtime/Windows/D3D11RHI/Private/D3D11Commands.cpp b/Engine/Source/Runtime/Windows/D3D11RHI/Private/D3D11Commands.cpp
--- a/Engine/Source/Runtime/Windows/D3D11RHI/Private/D3D11Commands.cpp
+++ b/Engine/Source/Runtime/Windows/D3D11RHI/Private/D3D11Commands.cpp
@@ -1315,13 +1315,13 @@ FORCEINLINE void SetResource(FD3D11DynamicRHI* RESTRICT D3D11RHI, FD3D11StateCac
 {
 	// We set the resource through the RHI to track state for the purposes of unbinding SRVs when a UAV or RTV is bound.
 	// todo: need to support SRV_Static for faster calls when possible
-	D3D11RHI->SetShaderResourceView<Frequency>(ShaderResource, SRV, BindIndex, ResourceName,FD3D11StateCache::SRV_Unknown);
+	D3D11RHI->SetShaderResourceView<Frequency>(ShaderResource, SRV, BindIndex, ResourceName, FD3D11StateCache::SRV_Unknown);
 }
 
 template <EShaderFrequency Frequency>
 FORCEINLINE void SetResource(FD3D11DynamicRHI* RESTRICT D3D11RHI, FD3D11StateCache* RESTRICT StateCache, uint32 BindIndex, ID3D11SamplerState* RESTRICT SamplerState)
 {
-	StateCache->SetSamplerState<Frequency>(SamplerState,BindIndex);
+	StateCache->SetSamplerState<Frequency>(SamplerState, BindIndex);
 }
 
 template <EShaderFrequency ShaderFrequency>
@@ -1585,7 +1585,11 @@ int32 FD3D11DynamicRHI::SetUAVPSResourcesFromTables(const ShaderType* RESTRICT S
 		FD3D11UniformBuffer* Buffer = (FD3D11UniformBuffer*)BoundUniformBuffers[ShaderType::StaticFrequency][BufferIndex].GetReference();
 
 		check(BufferIndex < Shader->ShaderResourceTable.ResourceTableLayoutHashes.Num());
+	#if !(UE_BUILD_SHIPPING || UE_BUILD_TEST)
 		const TCHAR* LayoutName = *Buffer->GetLayout().GetDebugName();
+	#else
+		const TCHAR* LayoutName = nullptr;
+	#endif
 
 		if ((EShaderFrequency)ShaderType::StaticFrequency == SF_Pixel)
 		{
