D3D11 RHI - Work around an issue with the WARP device (-AllowSoftwareRendering) & BC7
Creating two views with different formats (DXGI_FORMAT_BC7_UNORM vs DXGI_FORMAT_BC7_UNORM_SRGB) will result in a crash inside d3d10warp.dll when creating the second view

diff --git a/Engine/Source/Runtime/Windows/D3D11RHI/Private/D3D11Device.cpp b/Engine/Source/Runtime/Windows/D3D11RHI/Private/D3D11Device.cpp
--- a/Engine/Source/Runtime/Windows/D3D11RHI/Private/D3D11Device.cpp
+++ b/Engine/Source/Runtime/Windows/D3D11RHI/Private/D3D11Device.cpp
@@ -52,7 +52,7 @@ TAutoConsoleVariable<int32> CVarD3D11ZeroBufferSizeInMB(
 	);
 
 
-FD3D11DynamicRHI::FD3D11DynamicRHI(IDXGIFactory1* InDXGIFactory1,D3D_FEATURE_LEVEL InFeatureLevel, int32 InChosenAdapter, const DXGI_ADAPTER_DESC& InChosenDescription) :
+FD3D11DynamicRHI::FD3D11DynamicRHI(IDXGIFactory1* InDXGIFactory1,D3D_FEATURE_LEVEL InFeatureLevel, int32 InChosenAdapter, const DXGI_ADAPTER_DESC& InChosenDescription, bool bInSoftwareAdapter) :
 	DXGIFactory1(InDXGIFactory1),
 #if NV_AFTERMATH
 	NVAftermathIMContextHandle(nullptr),
@@ -78,6 +78,7 @@ FD3D11DynamicRHI::FD3D11DynamicRHI(IDXGIFactory1* InDXGIFactory1,D3D_FEATURE_LEV
 	bUsingTessellation(false),
 	GPUProfilingData(this),
 	ChosenAdapter(InChosenAdapter),
+	bSoftwareAdapter(bInSoftwareAdapter),
 	ChosenDescription(InChosenDescription),
 	bAllowVendorDevice(!FParse::Param(FCommandLine::Get(), TEXT("novendordevice")))
 {
diff --git a/Engine/Source/Runtime/Windows/D3D11RHI/Private/D3D11RHIPrivate.h b/Engine/Source/Runtime/Windows/D3D11RHI/Private/D3D11RHIPrivate.h
--- a/Engine/Source/Runtime/Windows/D3D11RHI/Private/D3D11RHIPrivate.h
+++ b/Engine/Source/Runtime/Windows/D3D11RHI/Private/D3D11RHIPrivate.h
@@ -367,7 +367,7 @@ public:
 	friend class FWindowsMixedRealityViewport;
 
 	/** Initialization constructor. */
-	FD3D11DynamicRHI(IDXGIFactory1* InDXGIFactory1,D3D_FEATURE_LEVEL InFeatureLevel,int32 InChosenAdapter, const DXGI_ADAPTER_DESC& ChosenDescription);
+	FD3D11DynamicRHI(IDXGIFactory1* InDXGIFactory1,D3D_FEATURE_LEVEL InFeatureLevel,int32 InChosenAdapter, const DXGI_ADAPTER_DESC& ChosenDescription, bool bInSoftwareAdapter);
 
 	/** Destructor */
 	virtual ~FD3D11DynamicRHI();
@@ -923,6 +923,8 @@ protected:
 	FD3DGPUProfiler GPUProfilingData;
 	// >= 0, was computed before, unless hardware was changed during engine init it should be the same
 	int32 ChosenAdapter;
+	// if this is a software adapter
+	bool bSoftwareAdapter;
 	// we don't use AdapterDesc.Description as there is a bug with Optimus where it can report the wrong name
 	DXGI_ADAPTER_DESC ChosenDescription;
 
@@ -1054,11 +1056,14 @@ struct FD3D11Adapter
 	int32 AdapterIndex;
 	/** The maximum D3D11 feature level supported. 0 if not supported or FindAdpater() wasn't called */
 	D3D_FEATURE_LEVEL MaxSupportedFeatureLevel;
+	/** Whether this is a software adapter */
+	bool bSoftwareAdapter;
 
 	// constructor
-	FD3D11Adapter(int32 InAdapterIndex = -1, D3D_FEATURE_LEVEL InMaxSupportedFeatureLevel = (D3D_FEATURE_LEVEL)0)
+	FD3D11Adapter(int32 InAdapterIndex = -1, D3D_FEATURE_LEVEL InMaxSupportedFeatureLevel = (D3D_FEATURE_LEVEL)0, bool bInSoftwareAdatper = false)
 		: AdapterIndex(InAdapterIndex)
 		, MaxSupportedFeatureLevel(InMaxSupportedFeatureLevel)
+		, bSoftwareAdapter(bInSoftwareAdatper)
 	{
 	}
 
diff --git a/Engine/Source/Runtime/Windows/D3D11RHI/Private/D3D11Texture.cpp b/Engine/Source/Runtime/Windows/D3D11RHI/Private/D3D11Texture.cpp
--- a/Engine/Source/Runtime/Windows/D3D11RHI/Private/D3D11Texture.cpp
+++ b/Engine/Source/Runtime/Windows/D3D11RHI/Private/D3D11Texture.cpp
@@ -506,6 +506,26 @@ DXGI_FORMAT FD3D11DynamicRHI::GetPlatformTextureResourceFormat(DXGI_FORMAT InFor
 	return InFormat;
 }
 
+// Work around an issue with the WARP device & BC7
+// Creating two views with different formats (DXGI_FORMAT_BC7_UNORM vs DXGI_FORMAT_BC7_UNORM_SRGB)
+// will result in a crash inside d3d10warp.dll when creating the second view
+void ApplyBC7SoftwareAdapterWorkaround(bool bSoftwareAdapter, D3D11_TEXTURE2D_DESC& Desc)
+{
+	if (bSoftwareAdapter)
+	{
+		bool bApplyWorkaround =	Desc.Format == DXGI_FORMAT_BC7_TYPELESS
+							 && Desc.Usage == D3D11_USAGE_DEFAULT
+							 && Desc.MipLevels == 1
+							 && Desc.ArraySize == 1
+							 && Desc.CPUAccessFlags == 0;
+
+		if (bApplyWorkaround)
+		{
+			Desc.MiscFlags |= D3D11_RESOURCE_MISC_SHARED;
+		}
+	}
+}
+
 /** If true, guard texture creates with SEH to log more information about a driver crash we are seeing during texture streaming. */
 #define GUARDED_TEXTURE_CREATES (PLATFORM_WINDOWS && !(UE_BUILD_SHIPPING || UE_BUILD_TEST || PLATFORM_COMPILER_CLANG))
 
@@ -653,6 +673,8 @@ TD3D11Texture2D<BaseResourceType>* FD3D11DynamicRHI::CreateD3D11Texture2D(uint32
 	TextureDesc.CPUAccessFlags = CPUAccessFlags;
 	TextureDesc.MiscFlags = bCubeTexture ? D3D11_RESOURCE_MISC_TEXTURECUBE : 0;
 
+	ApplyBC7SoftwareAdapterWorkaround(bSoftwareAdapter, TextureDesc);
+
 	// NV12 doesn't support SRV in NV12 format so don't create SRV for it.
 	// Todo: add support for SRVs of underneath luminance & chrominance textures.
 	if (Format == PF_NV12)
@@ -1269,6 +1291,8 @@ FTexture2DRHIRef FD3D11DynamicRHI::RHIAsyncCreateTexture2D(uint32 SizeX,uint32 S
 	TextureDesc.CPUAccessFlags = 0;
 	TextureDesc.MiscFlags = 0;
 
+	ApplyBC7SoftwareAdapterWorkaround(bSoftwareAdapter, TextureDesc);
+
 	for (uint32 MipIndex = 0; MipIndex < NumInitialMips; ++MipIndex)
 	{
 		uint32 NumBlocksX = FMath::Max<uint32>(1,(SizeX >> MipIndex) / GPixelFormats[Format].BlockSizeX);
diff --git a/Engine/Source/Runtime/Windows/D3D11RHI/Private/Windows/WindowsD3D11Device.cpp b/Engine/Source/Runtime/Windows/D3D11RHI/Private/Windows/WindowsD3D11Device.cpp
--- a/Engine/Source/Runtime/Windows/D3D11RHI/Private/Windows/WindowsD3D11Device.cpp
+++ b/Engine/Source/Runtime/Windows/D3D11RHI/Private/Windows/WindowsD3D11Device.cpp
@@ -963,17 +963,19 @@ void FD3D11DynamicRHIModule::FindAdapter()
 						bIsNonLocalMemoryPresent = NonLocalVideoMemoryInfo.Budget != 0;
 					}
 				}
+
+				const bool bIsSoftware = bIsMicrosoft;
 				const bool bIsIntegrated = bIsIntel && !bIsNonLocalMemoryPresent;
 				// PerfHUD is for performance profiling
 				const bool bIsPerfHUD = !FCString::Stricmp(AdapterDesc.Description,TEXT("NVIDIA PerfHUD"));
 
-				FD3D11Adapter CurrentAdapter(AdapterIndex, ActualFeatureLevel);
+				FD3D11Adapter CurrentAdapter(AdapterIndex, ActualFeatureLevel, bIsSoftware);
 
 				// Add special check to support HMDs, which do not have associated outputs.
 				// To reject the software emulation, unless the cvar wants it.
 				// https://msdn.microsoft.com/en-us/library/windows/desktop/bb205075(v=vs.85).aspx#WARP_new_for_Win8
 				// Before we tested for no output devices but that failed where a laptop had a Intel (with output) and NVidia (with no output)
-				const bool bSkipSoftwareAdapter = bIsMicrosoft && !bAllowSoftwareFallback && CVarExplicitAdapterValue < 0 && HmdGraphicsAdapterLuid == 0;
+				const bool bSkipSoftwareAdapter = bIsSoftware && !bAllowSoftwareFallback && CVarExplicitAdapterValue < 0 && HmdGraphicsAdapterLuid == 0;
 				
 				// we don't allow the PerfHUD adapter
 				const bool bSkipPerfHUDAdapter = bIsPerfHUD && !bAllowPerfHUD;
@@ -1080,7 +1082,7 @@ FDynamicRHI* FD3D11DynamicRHIModule::CreateRHI(ERHIFeatureLevel::Type RequestedF
 	GMaxRHIShaderPlatform = GShaderPlatformForFeatureLevel[GMaxRHIFeatureLevel];
 	check(GMaxRHIShaderPlatform != SP_NumPlatforms);
 
-	GD3D11RHI = new FD3D11DynamicRHI(DXGIFactory1,ChosenAdapter.MaxSupportedFeatureLevel,ChosenAdapter.AdapterIndex,ChosenDescription);
+	GD3D11RHI = new FD3D11DynamicRHI(DXGIFactory1,ChosenAdapter.MaxSupportedFeatureLevel,ChosenAdapter.AdapterIndex,ChosenDescription,ChosenAdapter.bSoftwareAdapter);
 	FDynamicRHI* FinalRHI = GD3D11RHI;
 
 #if ENABLE_RHI_VALIDATION
