D3D11 - Clear reference to CurrentComputeShader in CleanupD3DDevice()
Prior to this change, it was only cleared by the desctructor, after the device was actually destroyed

diff --git a/Engine/Source/Runtime/Windows/D3D11RHI/Private/D3D11Device.cpp b/Engine/Source/Runtime/Windows/D3D11RHI/Private/D3D11Device.cpp
--- a/Engine/Source/Runtime/Windows/D3D11RHI/Private/D3D11Device.cpp
+++ b/Engine/Source/Runtime/Windows/D3D11RHI/Private/D3D11Device.cpp
@@ -508,6 +508,8 @@ void FD3D11DynamicRHI::CleanupD3DDevice()
 
 		check(!GIsCriticalError);
 
+		CurrentComputeShader = nullptr;
+
 		// Ask all initialized FRenderResources to release their RHI resources.
 		FRenderResource::ReleaseRHIForAllResources();
 
