Rename UE4_ prefix to UE

diff --git a/Engine/Source/Runtime/Windows/XAudio2/Private/XAudio2Device.cpp b/Engine/Source/Runtime/Windows/XAudio2/Private/XAudio2Device.cpp
--- a/Engine/Source/Runtime/Windows/XAudio2/Private/XAudio2Device.cpp
+++ b/Engine/Source/Runtime/Windows/XAudio2/Private/XAudio2Device.cpp
@@ -109,7 +109,7 @@ bool FXAudio2Device::InitializeHardware()
 	LoadVorbisLibraries();
 #endif
 
-	SampleRate = UE4_XAUDIO2_SAMPLERATE;
+	SampleRate = UE_XAUDIO2_SAMPLERATE;
 
 #if PLATFORM_WINDOWS || PLATFORM_HOLOLENS
 	bComInitialized = FPlatformMisc::CoInitialize();
@@ -239,7 +239,7 @@ bool FXAudio2Device::InitializeHardware()
 	DeviceProperties->XAudio2->SetDebugConfiguration(&DebugConfig);
 #endif // #if DEBUG_XAUDIO2
 
-	FXAudioDeviceProperties::NumSpeakers = UE4_XAUDIO2_NUMCHANNELS;
+	FXAudioDeviceProperties::NumSpeakers = UE_XAUDIO2_NUMCHANNELS;
 	SampleRate = FXAudioDeviceProperties::DeviceDetails.OutputFormat.Format.nSamplesPerSec;
 
 	// Clamp the output frequency to the limits of the reverb XAPO
@@ -257,10 +257,10 @@ bool FXAudio2Device::InitializeHardware()
 			FXAudioDeviceProperties::NumSpeakers,
 			( float )SampleRate / 1000.0f,
 			BitsPerSample,
-			(uint32)UE4_XAUDIO2_CHANNELMASK );
+			(uint32)UE_XAUDIO2_CHANNELMASK );
 	}
 
-	if( !GetOutputMatrix( UE4_XAUDIO2_CHANNELMASK, FXAudioDeviceProperties::NumSpeakers ) )
+	if( !GetOutputMatrix( UE_XAUDIO2_CHANNELMASK, FXAudioDeviceProperties::NumSpeakers ) )
 	{
 		UE_LOG(LogInit, Log, TEXT( "Unsupported speaker configuration for this number of channels" ) );
 		DeviceProperties->XAudio2 = NULL;
@@ -278,7 +278,7 @@ bool FXAudio2Device::InitializeHardware()
 #else	// #if XAUDIO_SUPPORTS_DEVICE_DETAILS
 	// Create the final output voice
 	if (!ValidateAPICall(TEXT("CreateMasteringVoice"),
-		DeviceProperties->XAudio2->CreateMasteringVoice(&DeviceProperties->MasteringVoice, UE4_XAUDIO2_NUMCHANNELS, UE4_XAUDIO2_SAMPLERATE, 0, 0, nullptr )))
+		DeviceProperties->XAudio2->CreateMasteringVoice(&DeviceProperties->MasteringVoice, UE_XAUDIO2_NUMCHANNELS, UE_XAUDIO2_SAMPLERATE, 0, 0, nullptr )))
 	{
 		UE_LOG(LogInit, Warning, TEXT( "Failed to create the mastering voice for XAudio2" ) );
 		DeviceProperties->XAudio2 = nullptr;
diff --git a/Engine/Source/Runtime/Windows/XAudio2/Private/XAudio2Effects.cpp b/Engine/Source/Runtime/Windows/XAudio2/Private/XAudio2Effects.cpp
--- a/Engine/Source/Runtime/Windows/XAudio2/Private/XAudio2Effects.cpp
+++ b/Engine/Source/Runtime/Windows/XAudio2/Private/XAudio2Effects.cpp
@@ -314,7 +314,7 @@ XAPO_REGISTRATION_PROPERTIES FXAudio2RadioEffect::Registration =
  */
 bool FXAudio2EffectsManager::CreateEQPremasterVoices()
 {
-	uint32 SampleRate = UE4_XAUDIO2_SAMPLERATE;
+	uint32 SampleRate = UE_XAUDIO2_SAMPLERATE;
 
 	// Create the EQ effect
 	if( !XAudio2Device->ValidateAPICall( TEXT( "CreateFX (EQ)" ), 
@@ -349,8 +349,8 @@ bool FXAudio2EffectsManager::CreateEQPremasterVoices()
 	}
 
 	// Set the output matrix catering for a potential downmix
-	const uint32 NumChannels = UE4_XAUDIO2_NUMCHANNELS;
-	FXAudio2Device::GetOutputMatrix(UE4_XAUDIO2_CHANNELMASK, NumChannels);
+	const uint32 NumChannels = UE_XAUDIO2_NUMCHANNELS;
+	FXAudio2Device::GetOutputMatrix(UE_XAUDIO2_CHANNELMASK, NumChannels);
 
 	if( !XAudio2Device->ValidateAPICall( TEXT( "SetOutputMatrix (EQPremaster)" ), 
 		EQPremasterVoice->SetOutputMatrix( NULL, SPEAKER_COUNT, NumChannels, FXAudioDeviceProperties::OutputMixMatrix ) ) )
@@ -374,7 +374,7 @@ bool FXAudio2EffectsManager::CreateReverbVoice()
 {
 	UINT32 Flags;
 
-	uint32 SampleRate = UE4_XAUDIO2_SAMPLERATE;
+	uint32 SampleRate = UE_XAUDIO2_SAMPLERATE;
 
 	Flags = 0;		// XAUDIO2FX_DEBUG
 
@@ -438,7 +438,7 @@ bool FXAudio2EffectsManager::CreateReverbVoice()
 bool FXAudio2EffectsManager::CreateRadioVoice()
 {
 	// Grab the sample rate, which is needed to configure the radio distortion effect settings.
-	const uint32 SampleRate = UE4_XAUDIO2_SAMPLERATE;
+	const uint32 SampleRate = UE_XAUDIO2_SAMPLERATE;
 	
 	// Create the custom XAPO radio distortion effect
 	FXAudio2RadioEffect* NewRadioEffect = new FXAudio2RadioEffect( SampleRate );
@@ -468,8 +468,8 @@ bool FXAudio2EffectsManager::CreateRadioVoice()
 		return false;
 	}
 
-	const uint32 NumChannels = UE4_XAUDIO2_NUMCHANNELS;
-	FXAudio2Device::GetOutputMatrix( UE4_XAUDIO2_CHANNELMASK, NumChannels );
+	const uint32 NumChannels = UE_XAUDIO2_NUMCHANNELS;
+	FXAudio2Device::GetOutputMatrix( UE_XAUDIO2_CHANNELMASK, NumChannels );
 
 	// Designate the radio-distorted audio to route to the master voice.
 	if( !XAudio2Device->ValidateAPICall( TEXT( "SetOutputMatrix (Radio)" ), 
diff --git a/Engine/Source/Runtime/Windows/XAudio2/Private/XAudio2Source.cpp b/Engine/Source/Runtime/Windows/XAudio2/Private/XAudio2Source.cpp
--- a/Engine/Source/Runtime/Windows/XAudio2/Private/XAudio2Source.cpp
+++ b/Engine/Source/Runtime/Windows/XAudio2/Private/XAudio2Source.cpp
@@ -1676,7 +1676,7 @@ void FSpatializationHelper::Init()
 	//
 	//  Speaker geometry configuration on the final mix, specifies assignment of channels
 	//  to speaker positions, defined as per WAVEFORMATEXTENSIBLE.dwChannelMask
-	X3DAudioInitialize( UE4_XAUDIO2_CHANNELMASK, X3DAUDIO_SPEED_OF_SOUND, X3DInstance );
+	X3DAudioInitialize( UE_XAUDIO2_CHANNELMASK, X3DAUDIO_SPEED_OF_SOUND, X3DInstance );
 
 	// Initialize 3D audio parameters
 #if X3DAUDIO_VECTOR_IS_A_D3DVECTOR
diff --git a/Engine/Source/Runtime/Windows/XAudio2/Private/XAudio2Support.h b/Engine/Source/Runtime/Windows/XAudio2/Private/XAudio2Support.h
--- a/Engine/Source/Runtime/Windows/XAudio2/Private/XAudio2Support.h
+++ b/Engine/Source/Runtime/Windows/XAudio2/Private/XAudio2Support.h
@@ -965,11 +965,11 @@ struct FXAudioDeviceProperties final : public IDeviceChangedListener
 };
 
 #if XAUDIO_SUPPORTS_DEVICE_DETAILS
-	#define UE4_XAUDIO2_NUMCHANNELS		FXAudioDeviceProperties::DeviceDetails.OutputFormat.Format.nChannels > 6 ? 6 : FXAudioDeviceProperties::DeviceDetails.OutputFormat.Format.nChannels
-	#define UE4_XAUDIO2_CHANNELMASK		FXAudioDeviceProperties::DeviceDetails.OutputFormat.Format.nChannels > 6 ? SPEAKER_5POINT1 : FXAudioDeviceProperties::DeviceDetails.OutputFormat.dwChannelMask
-	#define UE4_XAUDIO2_SAMPLERATE		FXAudioDeviceProperties::DeviceDetails.OutputFormat.Format.nSamplesPerSec
+	#define UE_XAUDIO2_NUMCHANNELS		FXAudioDeviceProperties::DeviceDetails.OutputFormat.Format.nChannels > 6 ? 6 : FXAudioDeviceProperties::DeviceDetails.OutputFormat.Format.nChannels
+	#define UE_XAUDIO2_CHANNELMASK		FXAudioDeviceProperties::DeviceDetails.OutputFormat.Format.nChannels > 6 ? SPEAKER_5POINT1 : FXAudioDeviceProperties::DeviceDetails.OutputFormat.dwChannelMask
+	#define UE_XAUDIO2_SAMPLERATE		FXAudioDeviceProperties::DeviceDetails.OutputFormat.Format.nSamplesPerSec
 #else	//XAUDIO_SUPPORTS_DEVICE_DETAILS
-#define UE4_XAUDIO2_NUMCHANNELS		6
-#define UE4_XAUDIO2_CHANNELMASK		SPEAKER_5POINT1
-#define UE4_XAUDIO2_SAMPLERATE		44100
+#define UE_XAUDIO2_NUMCHANNELS		6
+#define UE_XAUDIO2_CHANNELMASK		SPEAKER_5POINT1
+#define UE_XAUDIO2_SAMPLERATE		44100
 #endif	//XAUDIO_SUPPORTS_DEVICE_DETAILS
