Tessellation Removal: Removing Tessellation from D3D11RHI.

diff --git a/Engine/Source/Runtime/Windows/D3D11RHI/Private/D3D11Commands.cpp b/Engine/Source/Runtime/Windows/D3D11RHI/Private/D3D11Commands.cpp
--- a/Engine/Source/Runtime/Windows/D3D11RHI/Private/D3D11Commands.cpp
+++ b/Engine/Source/Runtime/Windows/D3D11RHI/Private/D3D11Commands.cpp
@@ -38,8 +38,6 @@
 DECLARE_ISBOUNDSHADER(VertexShader)
 DECLARE_ISBOUNDSHADER(PixelShader)
 DECLARE_ISBOUNDSHADER(GeometryShader)
-DECLARE_ISBOUNDSHADER(HullShader)
-DECLARE_ISBOUNDSHADER(DomainShader)
 DECLARE_ISBOUNDSHADER(ComputeShader)
 
 
@@ -196,8 +194,6 @@ void FD3D11DynamicRHI::RHISetGraphicsPipelineState(FRHIGraphicsPipelineState* Gr
 	if (bApplyAdditionalState)
 	{
 		ApplyStaticUniformBuffers(static_cast<FD3D11VertexShader*>(PsoInit.BoundShaderState.VertexShaderRHI));
-		ApplyStaticUniformBuffers(static_cast<FD3D11HullShader*>(PsoInit.BoundShaderState.HullShaderRHI));
-		ApplyStaticUniformBuffers(static_cast<FD3D11DomainShader*>(PsoInit.BoundShaderState.DomainShaderRHI));
 		ApplyStaticUniformBuffers(static_cast<FD3D11GeometryShader*>(PsoInit.BoundShaderState.GeometryShaderRHI));
 		ApplyStaticUniformBuffers(static_cast<FD3D11PixelShader*>(PsoInit.BoundShaderState.PixelShaderRHI));
 	}
@@ -324,19 +320,8 @@ void FD3D11DynamicRHI::RHISetBoundShaderState(FRHIBoundShaderState* BoundShaderS
 	StateCache.SetVertexShader(BoundShaderState->VertexShader);
 	StateCache.SetPixelShader(BoundShaderState->PixelShader);
 
-	StateCache.SetHullShader(BoundShaderState->HullShader);
-	StateCache.SetDomainShader(BoundShaderState->DomainShader);
 	StateCache.SetGeometryShader(BoundShaderState->GeometryShader);
 
-	if(BoundShaderState->HullShader != NULL && BoundShaderState->DomainShader != NULL)
-	{
-		bUsingTessellation = true;
-	}
-	else
-	{
-		bUsingTessellation = false;
-	}
-
 	// @TODO : really should only discard the constants if the shader state has actually changed.
 	bDiscardSharedConstants = true;
 
@@ -397,20 +382,6 @@ void FD3D11DynamicRHI::RHISetShaderTexture(FRHIGraphicsShader* ShaderRHI,uint32
 		SetShaderTexture<SF_Vertex>(NewTexture, ShaderResourceView, TextureIndex, NewTextureRHI);
 	}
 	break;
-	case SF_Hull:
-	{
-		FD3D11HullShader* HullShader = static_cast<FD3D11HullShader*>(ShaderRHI);
-		VALIDATE_BOUND_SHADER(HullShader);
-		SetShaderTexture<SF_Hull>(NewTexture, ShaderResourceView, TextureIndex, NewTextureRHI);
-	}
-	break;
-	case SF_Domain:
-	{
-		FD3D11DomainShader* DomainShader = static_cast<FD3D11DomainShader*>(ShaderRHI);
-		VALIDATE_BOUND_SHADER(DomainShader);
-		SetShaderTexture<SF_Domain>(NewTexture, ShaderResourceView, TextureIndex, NewTextureRHI);
-	}
-	break;
 	case SF_Geometry:
 	{
 		FD3D11GeometryShader* GeometryShader = static_cast<FD3D11GeometryShader*>(ShaderRHI);
@@ -514,20 +485,6 @@ void FD3D11DynamicRHI::RHISetShaderResourceViewParameter(FRHIGraphicsShader* Sha
 		SetShaderResourceView<SF_Vertex>(Resource, D3D11SRV, TextureIndex, NAME_None);
 	}
 	break;
-	case SF_Hull:
-	{
-		FD3D11HullShader* HullShader = static_cast<FD3D11HullShader*>(ShaderRHI);
-		VALIDATE_BOUND_SHADER(HullShader);
-		SetShaderResourceView<SF_Hull>(Resource, D3D11SRV, TextureIndex, NAME_None);
-	}
-	break;
-	case SF_Domain:
-	{
-		FD3D11DomainShader* DomainShader = static_cast<FD3D11DomainShader*>(ShaderRHI);
-		VALIDATE_BOUND_SHADER(DomainShader);
-		SetShaderResourceView<SF_Domain>(Resource, D3D11SRV, TextureIndex, NAME_None);
-	}
-	break;
 	case SF_Geometry:
 	{
 		FD3D11GeometryShader* GeometryShader = static_cast<FD3D11GeometryShader*>(ShaderRHI);
@@ -576,20 +533,6 @@ void FD3D11DynamicRHI::RHISetShaderSampler(FRHIGraphicsShader* ShaderRHI,uint32
 		FD3D11VertexShader* VertexShader = static_cast<FD3D11VertexShader*>(ShaderRHI);
 		VALIDATE_BOUND_SHADER(VertexShader);
 		StateCache.SetSamplerState<SF_Vertex>(StateResource, SamplerIndex);
-	}
-		break;
-	case SF_Hull:
-	{
-		FD3D11HullShader* HullShader = static_cast<FD3D11HullShader*>(ShaderRHI);
-		VALIDATE_BOUND_SHADER(HullShader);
-		StateCache.SetSamplerState<SF_Hull>(StateResource, SamplerIndex);
-	}
-	break;
-	case SF_Domain:
-	{
-		FD3D11DomainShader* DomainShader = static_cast<FD3D11DomainShader*>(ShaderRHI);
-		VALIDATE_BOUND_SHADER(DomainShader);
-		StateCache.SetSamplerState<SF_Domain>(StateResource, SamplerIndex);
 	}
 	break;
 	case SF_Geometry:
@@ -647,22 +590,6 @@ void FD3D11DynamicRHI::RHISetShaderUniformBuffer(FRHIGraphicsShader* ShaderRHI,u
 		StateCache.SetConstantBuffer<SF_Vertex>(ConstantBuffer, BufferIndex);
 	}
 	break;
-	case SF_Hull:
-	{
-		FD3D11HullShader* HullShader = static_cast<FD3D11HullShader*>(ShaderRHI);
-		VALIDATE_BOUND_SHADER(HullShader);
-		Stage = SF_Hull;
-		StateCache.SetConstantBuffer<SF_Hull>(ConstantBuffer, BufferIndex);
-	}
-	break;
-	case SF_Domain:
-	{
-		FD3D11DomainShader* DomainShader = static_cast<FD3D11DomainShader*>(ShaderRHI);
-		VALIDATE_BOUND_SHADER(DomainShader);
-		Stage = SF_Domain;
-		StateCache.SetConstantBuffer<SF_Domain>(ConstantBuffer, BufferIndex);
-	}
-	break;
 	case SF_Geometry:
 	{
 		FD3D11GeometryShader* GeometryShader = static_cast<FD3D11GeometryShader*>(ShaderRHI);
@@ -714,22 +641,6 @@ void FD3D11DynamicRHI::RHISetShaderParameter(FRHIGraphicsShader* ShaderRHI,uint3
 		VSConstantBuffers[BufferIndex]->UpdateConstant((const uint8*)NewValue, BaseIndex, NumBytes);
 	}
 	break;
-	case SF_Hull:
-	{
-		FD3D11HullShader* HullShader = static_cast<FD3D11HullShader*>(ShaderRHI);
-		VALIDATE_BOUND_SHADER(HullShader);
-		checkSlow(HSConstantBuffers[BufferIndex]);
-		HSConstantBuffers[BufferIndex]->UpdateConstant((const uint8*)NewValue, BaseIndex, NumBytes);
-	}
-	break;
-	case SF_Domain:
-	{
-		FD3D11DomainShader* DomainShader = static_cast<FD3D11DomainShader*>(ShaderRHI);
-		VALIDATE_BOUND_SHADER(DomainShader);
-		checkSlow(DSConstantBuffers[BufferIndex]);
-		DSConstantBuffers[BufferIndex]->UpdateConstant((const uint8*)NewValue, BaseIndex, NumBytes);
-	}
-	break;
 	case SF_Geometry:
 	{
 		FD3D11GeometryShader* GeometryShader = static_cast<FD3D11GeometryShader*>(ShaderRHI);
@@ -1157,31 +1068,8 @@ void FD3D11DynamicRHI::SetRenderTargetsAndClear(const FRHISetRenderTargetsInfo&
 
 // Primitive drawing.
 
-static D3D11_PRIMITIVE_TOPOLOGY GetD3D11PrimitiveType(EPrimitiveType PrimitiveType, bool bUsingTessellation)
+static D3D11_PRIMITIVE_TOPOLOGY GetD3D11PrimitiveType(EPrimitiveType PrimitiveType)
 {
-	if(bUsingTessellation)
-	{
-		switch(PrimitiveType)
-		{
-		case PT_1_ControlPointPatchList: return D3D11_PRIMITIVE_TOPOLOGY_1_CONTROL_POINT_PATCHLIST;
-		case PT_2_ControlPointPatchList: return D3D11_PRIMITIVE_TOPOLOGY_2_CONTROL_POINT_PATCHLIST;
-
-		// This is the case for tessellation without AEN or other buffers, so just flip to 3 CPs
-		case PT_TriangleList: return D3D11_PRIMITIVE_TOPOLOGY_3_CONTROL_POINT_PATCHLIST;
-
-		case PT_LineList:
-		case PT_TriangleStrip:
-		case PT_QuadList:
-		case PT_PointList:
-		case PT_RectList:
-			UE_LOG(LogD3D11RHI, Fatal,TEXT("Invalid type specified for tessellated render, probably missing a case in FStaticMeshSceneProxy::GetMeshElement"));
-			break;
-		default:
-			// Other cases are valid.
-			break;
-		};
-	}
-
 	switch(PrimitiveType)
 	{
 	case PT_TriangleList: return D3D11_PRIMITIVE_TOPOLOGY_TRIANGLELIST;
@@ -1248,33 +1136,6 @@ void FD3D11DynamicRHI::CommitNonComputeShaderConstants()
 		}
 	}
 
-	// Skip HS/DS CB updates in cases where tessellation isn't being used
-	// Note that this is *potentially* unsafe because bDiscardSharedConstants is cleared at the
-	// end of the function, however we're OK for now because bDiscardSharedConstants
-	// is always reset whenever bUsingTessellation changes in SetBoundShaderState()
-	if(bUsingTessellation)
-	{
-		if (CurrentBoundShaderState->bShaderNeedsGlobalConstantBuffer[SF_Hull])
-		{
-			// Commit and bind hull shader constants
-			for(uint32 i=0;i<MAX_CONSTANT_BUFFER_SLOTS; i++)
-			{
-				FD3D11ConstantBuffer* ConstantBuffer = HSConstantBuffers[i];
-				FD3DRHIUtil::CommitConstants<SF_Hull>(ConstantBuffer, StateCache, i, bDiscardSharedConstants);
-			}
-		}
-
-		if (CurrentBoundShaderState->bShaderNeedsGlobalConstantBuffer[SF_Domain])
-		{
-			// Commit and bind domain shader constants
-			for(uint32 i=0;i<MAX_CONSTANT_BUFFER_SLOTS; i++)
-			{
-				FD3D11ConstantBuffer* ConstantBuffer = DSConstantBuffers[i];
-				FD3DRHIUtil::CommitConstants<SF_Domain>(ConstantBuffer, StateCache, i, bDiscardSharedConstants);
-			}
-		}
-	}
-
 	if (CurrentBoundShaderState->bShaderNeedsGlobalConstantBuffer[SF_Geometry])
 	{
 		// Commit and bind geometry shader constants
@@ -1647,14 +1508,6 @@ void FD3D11DynamicRHI::CommitGraphicsResourceTables()
 	{
 		SetResourcesFromTables(PixelShader);
 	}
-	if (auto* Shader = CurrentBoundShaderState->GetHullShader())
-	{
-		SetResourcesFromTables(Shader);
-	}
-	if (auto* Shader = CurrentBoundShaderState->GetDomainShader())
-	{
-		SetResourcesFromTables(Shader);
-	}
 	if (auto* Shader = CurrentBoundShaderState->GetGeometryShader())
 	{
 		SetResourcesFromTables(Shader);
@@ -1678,7 +1531,7 @@ void FD3D11DynamicRHI::RHIDrawPrimitive(uint32 BaseVertexIndex,uint32 NumPrimiti
 	uint32 VertexCount = GetVertexCountForPrimitiveCount(NumPrimitives,PrimitiveType);
 
 	GPUProfilingData.RegisterGPUWork(NumPrimitives * NumInstances, VertexCount * NumInstances);
-	StateCache.SetPrimitiveTopology(GetD3D11PrimitiveType(PrimitiveType,bUsingTessellation));
+	StateCache.SetPrimitiveTopology(GetD3D11PrimitiveType(PrimitiveType));
 	if(NumInstances > 1)
 	{
 		Direct3DDeviceIMContext->DrawInstanced(VertexCount,NumInstances,BaseVertexIndex,0);
@@ -1702,7 +1555,7 @@ void FD3D11DynamicRHI::RHIDrawPrimitiveIndirect(FRHIBuffer* ArgumentBufferRHI, u
 	CommitGraphicsResourceTables();
 	CommitNonComputeShaderConstants();
 
-	StateCache.SetPrimitiveTopology(GetD3D11PrimitiveType(PrimitiveType,bUsingTessellation));
+	StateCache.SetPrimitiveTopology(GetD3D11PrimitiveType(PrimitiveType));
 	Direct3DDeviceIMContext->DrawInstancedIndirect(ArgumentBuffer->Resource,ArgumentOffset);
 
 	EnableUAVOverlap();
@@ -1725,7 +1578,7 @@ void FD3D11DynamicRHI::RHIDrawIndexedIndirect(FRHIBuffer* IndexBufferRHI, FRHIBu
 
 	TrackResourceBoundAsIB(IndexBuffer);
 	StateCache.SetIndexBuffer(IndexBuffer->Resource, Format, 0);
-	StateCache.SetPrimitiveTopology(GetD3D11PrimitiveType(PrimitiveType,bUsingTessellation));
+	StateCache.SetPrimitiveTopology(GetD3D11PrimitiveType(PrimitiveType));
 
 	if(NumInstances > 1)
 	{
@@ -1765,7 +1618,7 @@ void FD3D11DynamicRHI::RHIDrawIndexedPrimitive(FRHIBuffer* IndexBufferRHI, int32
 
 	TrackResourceBoundAsIB(IndexBuffer);
 	StateCache.SetIndexBuffer(IndexBuffer->Resource, Format, 0);
-	StateCache.SetPrimitiveTopology(GetD3D11PrimitiveType(PrimitiveType,bUsingTessellation));
+	StateCache.SetPrimitiveTopology(GetD3D11PrimitiveType(PrimitiveType));
 
 	if (NumInstances > 1 || FirstInstance != 0)
 	{
@@ -1798,7 +1651,7 @@ void FD3D11DynamicRHI::RHIDrawIndexedPrimitiveIndirect(FRHIBuffer* IndexBufferRH
 	const DXGI_FORMAT Format = (IndexBuffer->GetStride() == sizeof(uint16) ? DXGI_FORMAT_R16_UINT : DXGI_FORMAT_R32_UINT);
 	TrackResourceBoundAsIB(IndexBuffer);
 	StateCache.SetIndexBuffer(IndexBuffer->Resource, Format, 0);
-	StateCache.SetPrimitiveTopology(GetD3D11PrimitiveType(PrimitiveType,bUsingTessellation));
+	StateCache.SetPrimitiveTopology(GetD3D11PrimitiveType(PrimitiveType));
 	Direct3DDeviceIMContext->DrawIndexedInstancedIndirect(ArgumentBuffer->Resource,ArgumentOffset);
 
 	EnableUAVOverlap();
diff --git a/Engine/Source/Runtime/Windows/D3D11RHI/Private/D3D11Device.cpp b/Engine/Source/Runtime/Windows/D3D11RHI/Private/D3D11Device.cpp
--- a/Engine/Source/Runtime/Windows/D3D11RHI/Private/D3D11Device.cpp
+++ b/Engine/Source/Runtime/Windows/D3D11RHI/Private/D3D11Device.cpp
@@ -74,8 +74,7 @@ FD3D11DynamicRHI::FD3D11DynamicRHI(IDXGIFactory1* InDXGIFactory1,D3D_FEATURE_LEV
 	PresentCounter(0),
 	ResourceTableFrameCounter(INDEX_NONE),
 	CurrentDSVAccessType(FExclusiveDepthStencil::DepthWrite_StencilWrite),
-	bDiscardSharedConstants(false),
-	bUsingTessellation(false),
+	bDiscardSharedConstants(false),	
 	GPUProfilingData(this),
 	ChosenAdapter(InChosenAdapter),
 	bSoftwareAdapter(bInSoftwareAdapter),
diff --git a/Engine/Source/Runtime/Windows/D3D11RHI/Private/D3D11RHI.cpp b/Engine/Source/Runtime/Windows/D3D11RHI/Private/D3D11RHI.cpp
--- a/Engine/Source/Runtime/Windows/D3D11RHI/Private/D3D11RHI.cpp
+++ b/Engine/Source/Runtime/Windows/D3D11RHI/Private/D3D11RHI.cpp
@@ -53,8 +53,6 @@ void ClearShaderResource(ID3D11DeviceContext* Direct3DDeviceIMContext, uint32 Re
 	case SF_Pixel:   Direct3DDeviceIMContext->PSSetShaderResources(ResourceIndex,1,&NullView); break;
 	case SF_Compute: Direct3DDeviceIMContext->CSSetShaderResources(ResourceIndex,1,&NullView); break;
 	case SF_Geometry:Direct3DDeviceIMContext->GSSetShaderResources(ResourceIndex,1,&NullView); break;
-	case SF_Domain:  Direct3DDeviceIMContext->DSSetShaderResources(ResourceIndex,1,&NullView); break;
-	case SF_Hull:    Direct3DDeviceIMContext->HSSetShaderResources(ResourceIndex,1,&NullView); break;
 	case SF_Vertex:  Direct3DDeviceIMContext->VSSetShaderResources(ResourceIndex,1,&NullView); break;
 	};
 }
@@ -184,8 +182,6 @@ void FD3D11DynamicRHI::InternalSetShaderResourceView(FD3D11BaseShaderResource* R
 }
 
 template void FD3D11DynamicRHI::InternalSetShaderResourceView<SF_Vertex>(FD3D11BaseShaderResource* Resource, ID3D11ShaderResourceView* SRV, int32 ResourceIndex, FName SRVName, FD3D11StateCache::ESRV_Type SrvType);
-template void FD3D11DynamicRHI::InternalSetShaderResourceView<SF_Hull>(FD3D11BaseShaderResource* Resource, ID3D11ShaderResourceView* SRV, int32 ResourceIndex, FName SRVName, FD3D11StateCache::ESRV_Type SrvType);
-template void FD3D11DynamicRHI::InternalSetShaderResourceView<SF_Domain>(FD3D11BaseShaderResource* Resource, ID3D11ShaderResourceView* SRV, int32 ResourceIndex, FName SRVName, FD3D11StateCache::ESRV_Type SrvType);
 template void FD3D11DynamicRHI::InternalSetShaderResourceView<SF_Pixel>(FD3D11BaseShaderResource* Resource, ID3D11ShaderResourceView* SRV, int32 ResourceIndex, FName SRVName, FD3D11StateCache::ESRV_Type SrvType);
 template void FD3D11DynamicRHI::InternalSetShaderResourceView<SF_Geometry>(FD3D11BaseShaderResource* Resource, ID3D11ShaderResourceView* SRV, int32 ResourceIndex, FName SRVName, FD3D11StateCache::ESRV_Type SrvType);
 template void FD3D11DynamicRHI::InternalSetShaderResourceView<SF_Compute>(FD3D11BaseShaderResource* Resource, ID3D11ShaderResourceView* SRV, int32 ResourceIndex, FName SRVName, FD3D11StateCache::ESRV_Type SrvType);
@@ -242,8 +238,6 @@ void FD3D11DynamicRHI::ConditionalClearShaderResource(FD3D11BaseShaderResource*
 	SCOPE_CYCLE_COUNTER(STAT_D3D11ClearShaderResourceTime);
 	check(Resource);
 	ClearShaderResourceViews<SF_Vertex>(Resource);
-	ClearShaderResourceViews<SF_Hull>(Resource);
-	ClearShaderResourceViews<SF_Domain>(Resource);
 	ClearShaderResourceViews<SF_Pixel>(Resource);
 	ClearShaderResourceViews<SF_Geometry>(Resource);
 	ClearShaderResourceViews<SF_Compute>(Resource);
@@ -286,8 +280,6 @@ void FD3D11DynamicRHI::ClearAllShaderResourcesForFrequency()
 void FD3D11DynamicRHI::ClearAllShaderResources()
 {
 	ClearAllShaderResourcesForFrequency<SF_Vertex>();
-	ClearAllShaderResourcesForFrequency<SF_Hull>();
-	ClearAllShaderResourcesForFrequency<SF_Domain>();
 	ClearAllShaderResourcesForFrequency<SF_Geometry>();
 	ClearAllShaderResourcesForFrequency<SF_Pixel>();
 	ClearAllShaderResourcesForFrequency<SF_Compute>();
diff --git a/Engine/Source/Runtime/Windows/D3D11RHI/Private/D3D11RHIPrivate.h b/Engine/Source/Runtime/Windows/D3D11RHI/Private/D3D11RHIPrivate.h
--- a/Engine/Source/Runtime/Windows/D3D11RHI/Private/D3D11RHIPrivate.h
+++ b/Engine/Source/Runtime/Windows/D3D11RHI/Private/D3D11RHIPrivate.h
@@ -425,8 +425,6 @@ public:
 	virtual FVertexDeclarationRHIRef RHICreateVertexDeclaration(const FVertexDeclarationElementList& Elements) final override;
 	virtual FPixelShaderRHIRef RHICreatePixelShader(TArrayView<const uint8> Code, const FSHAHash& Hash) final override;
 	virtual FVertexShaderRHIRef RHICreateVertexShader(TArrayView<const uint8> Code, const FSHAHash& Hash) final override;
-	virtual FHullShaderRHIRef RHICreateHullShader(TArrayView<const uint8> Code, const FSHAHash& Hash) final override;
-	virtual FDomainShaderRHIRef RHICreateDomainShader(TArrayView<const uint8> Code, const FSHAHash& Hash) final override;
 	virtual FGeometryShaderRHIRef RHICreateGeometryShader(TArrayView<const uint8> Code, const FSHAHash& Hash) final override;
 	virtual FComputeShaderRHIRef RHICreateComputeShader(TArrayView<const uint8> Code, const FSHAHash& Hash) final override;
 	virtual FStagingBufferRHIRef RHICreateStagingBuffer() final override;
@@ -617,8 +615,6 @@ public:
 	virtual FPixelShaderRHIRef CreatePixelShader_RenderThread(class FRHICommandListImmediate& RHICmdList, TArrayView<const uint8> Code, const FSHAHash& Hash) final override;
 	virtual FGeometryShaderRHIRef CreateGeometryShader_RenderThread(class FRHICommandListImmediate& RHICmdList, TArrayView<const uint8> Code, const FSHAHash& Hash) final override;
 	virtual FComputeShaderRHIRef CreateComputeShader_RenderThread(class FRHICommandListImmediate& RHICmdList, TArrayView<const uint8> Code, const FSHAHash& Hash) final override;
-	virtual FHullShaderRHIRef CreateHullShader_RenderThread(class FRHICommandListImmediate& RHICmdList, TArrayView<const uint8> Code, const FSHAHash& Hash) final override;
-	virtual FDomainShaderRHIRef CreateDomainShader_RenderThread(class FRHICommandListImmediate& RHICmdList, TArrayView<const uint8> Code, const FSHAHash& Hash);
 	virtual void* LockTexture2D_RenderThread(class FRHICommandListImmediate& RHICmdList, FRHITexture2D* Texture, uint32 MipIndex, EResourceLockMode LockMode, uint32& DestStride, bool bLockWithinMiptail, bool bNeedsDefaultRHIFlush = true) final override;
 	virtual void UnlockTexture2D_RenderThread(class FRHICommandListImmediate& RHICmdList, FRHITexture2D* Texture, uint32 MipIndex, bool bLockWithinMiptail, bool bNeedsDefaultRHIFlush = true) final override;
 	virtual void UpdateTexture2D_RenderThread(class FRHICommandListImmediate& RHICmdList, FRHITexture2D* Texture, uint32 MipIndex, const struct FUpdateTextureRegion2D& UpdateRegion, uint32 SourcePitch, const uint8* SourceData) final override;
@@ -888,13 +884,8 @@ protected:
 	/** When a new shader is set, we discard all old constants set for the previous shader. */
 	bool bDiscardSharedConstants;
 
-	/** Set to true when the current shading setup uses tessellation */
-	bool bUsingTessellation;
-
 	/** A list of all D3D constant buffers RHIs that have been created. */
 	TArray<TRefCountPtr<FD3D11ConstantBuffer> > VSConstantBuffers;
-	TArray<TRefCountPtr<FD3D11ConstantBuffer> > HSConstantBuffers;
-	TArray<TRefCountPtr<FD3D11ConstantBuffer> > DSConstantBuffers;
 	TArray<TRefCountPtr<FD3D11ConstantBuffer> > PSConstantBuffers;
 	TArray<TRefCountPtr<FD3D11ConstantBuffer> > GSConstantBuffers;
 	TArray<TRefCountPtr<FD3D11ConstantBuffer> > CSConstantBuffers;
diff --git a/Engine/Source/Runtime/Windows/D3D11RHI/Private/D3D11Shaders.cpp b/Engine/Source/Runtime/Windows/D3D11RHI/Private/D3D11Shaders.cpp
--- a/Engine/Source/Runtime/Windows/D3D11RHI/Private/D3D11Shaders.cpp
+++ b/Engine/Source/Runtime/Windows/D3D11RHI/Private/D3D11Shaders.cpp
@@ -206,66 +206,6 @@ FGeometryShaderRHIRef FD3D11DynamicRHI::CreateGeometryShader_RenderThread(
 	return RHICreateGeometryShader(Code, Hash);
 }
 
-FHullShaderRHIRef FD3D11DynamicRHI::RHICreateHullShader(TArrayView<const uint8> Code, const FSHAHash& Hash)
-{ 
-	FShaderCodeReader ShaderCode(Code);
-
-	FD3D11HullShader* Shader = new FD3D11HullShader;
-
-	FMemoryReaderView Ar( Code, true );
-	Ar << Shader->ShaderResourceTable;
-	int32 Offset = Ar.Tell();
-	const uint8* CodePtr = Code.GetData() + Offset;
-	const size_t CodeSize = ShaderCode.GetActualShaderCodeSize() - Offset;
-
-	ReadShaderOptionalData(ShaderCode, *Shader);
-	if (!Shader->bIsSm6Shader && ApplyVendorExtensions(Direct3DDevice, SF_Hull, Shader->VendorExtensions))
-	{
-		VERIFYD3D11SHADERRESULT(Direct3DDevice->CreateHullShader((void*)CodePtr, CodeSize, nullptr, Shader->Resource.GetInitReference()), Shader, Direct3DDevice);
-		ResetVendorExtensions(Direct3DDevice, SF_Hull, Shader->VendorExtensions);
-		InitUniformBufferStaticSlots(Shader);
-	}
-
-	return Shader;
-}
-
-FHullShaderRHIRef FD3D11DynamicRHI::CreateHullShader_RenderThread(
-	class FRHICommandListImmediate& RHICmdList,
-	TArrayView<const uint8> Code, const FSHAHash& Hash)
-{
-	return RHICreateHullShader(Code, Hash);
-}
-
-FDomainShaderRHIRef FD3D11DynamicRHI::RHICreateDomainShader(TArrayView<const uint8> Code, const FSHAHash& Hash)
-{ 
-	FShaderCodeReader ShaderCode(Code);
-
-	FD3D11DomainShader* Shader = new FD3D11DomainShader;
-
-	FMemoryReaderView Ar(Code, true);
-	Ar << Shader->ShaderResourceTable;
-	int32 Offset = Ar.Tell();
-	const uint8* CodePtr = Code.GetData() + Offset;
-	const size_t CodeSize = ShaderCode.GetActualShaderCodeSize() - Offset;
-
-	ReadShaderOptionalData(ShaderCode, *Shader);
-	if (!Shader->bIsSm6Shader && ApplyVendorExtensions(Direct3DDevice, SF_Domain, Shader->VendorExtensions))
-	{
-		VERIFYD3D11SHADERRESULT(Direct3DDevice->CreateDomainShader((void*)CodePtr, CodeSize, nullptr, Shader->Resource.GetInitReference()), Shader, Direct3DDevice);
-		ResetVendorExtensions(Direct3DDevice, SF_Domain, Shader->VendorExtensions);
-		InitUniformBufferStaticSlots(Shader);
-	}
-
-	return Shader;
-}
-
-FDomainShaderRHIRef FD3D11DynamicRHI::CreateDomainShader_RenderThread(
-	class FRHICommandListImmediate& RHICmdList,
-	TArrayView<const uint8> Code, const FSHAHash& Hash)
-{
-	return RHICreateDomainShader(Code, Hash);
-}
-
 FPixelShaderRHIRef FD3D11DynamicRHI::RHICreatePixelShader(TArrayView<const uint8> Code, const FSHAHash& Hash)
 {
 	FShaderCodeReader ShaderCode(Code);
@@ -402,8 +342,6 @@ FD3D11BoundShaderState::~FD3D11BoundShaderState()
 * @param VertexDeclaration - existing vertex decl
 * @param StreamStrides - optional stream strides
 * @param VertexShader - existing vertex shader
-* @param HullShader - existing hull shader
-* @param DomainShader - existing domain shader
 * @param PixelShader - existing pixel shader
 * @param GeometryShader - existing geometry shader
 */
diff --git a/Engine/Source/Runtime/Windows/D3D11RHI/Private/D3D11StateCache.cpp b/Engine/Source/Runtime/Windows/D3D11RHI/Private/D3D11StateCache.cpp
--- a/Engine/Source/Runtime/Windows/D3D11RHI/Private/D3D11StateCache.cpp
+++ b/Engine/Source/Runtime/Windows/D3D11RHI/Private/D3D11StateCache.cpp
@@ -41,8 +41,6 @@ void FD3D11StateCacheBase::VerifySamplerStates()
 	switch (ShaderFrequency)
 	{
 	case SF_Vertex:		Direct3DDeviceIMContext->VSGetSamplers(0, D3D11_COMMONSHADER_SAMPLER_SLOT_COUNT, SamplerStates); break;
-	case SF_Hull:		Direct3DDeviceIMContext->HSGetSamplers(0, D3D11_COMMONSHADER_SAMPLER_SLOT_COUNT, SamplerStates); break;
-	case SF_Domain:		Direct3DDeviceIMContext->DSGetSamplers(0, D3D11_COMMONSHADER_SAMPLER_SLOT_COUNT, SamplerStates); break;
 	case SF_Geometry:	Direct3DDeviceIMContext->GSGetSamplers(0, D3D11_COMMONSHADER_SAMPLER_SLOT_COUNT, SamplerStates); break;
 	case SF_Pixel:		Direct3DDeviceIMContext->PSGetSamplers(0, D3D11_COMMONSHADER_SAMPLER_SLOT_COUNT, SamplerStates); break;
 	case SF_Compute:	Direct3DDeviceIMContext->CSGetSamplers(0, D3D11_COMMONSHADER_SAMPLER_SLOT_COUNT, SamplerStates); break;
@@ -67,8 +65,6 @@ void FD3D11StateCacheBase::VerifyConstantBuffers()
 	switch (ShaderFrequency)
 	{
 	case SF_Vertex:		Direct3DDeviceIMContext->VSGetConstantBuffers(0, D3D11_COMMONSHADER_CONSTANT_BUFFER_API_SLOT_COUNT, Buffers); break;
-	case SF_Hull:		Direct3DDeviceIMContext->HSGetConstantBuffers(0, D3D11_COMMONSHADER_CONSTANT_BUFFER_API_SLOT_COUNT, Buffers); break;
-	case SF_Domain:		Direct3DDeviceIMContext->DSGetConstantBuffers(0, D3D11_COMMONSHADER_CONSTANT_BUFFER_API_SLOT_COUNT, Buffers); break;
 	case SF_Geometry:	Direct3DDeviceIMContext->GSGetConstantBuffers(0, D3D11_COMMONSHADER_CONSTANT_BUFFER_API_SLOT_COUNT, Buffers); break;
 	case SF_Pixel:		Direct3DDeviceIMContext->PSGetConstantBuffers(0, D3D11_COMMONSHADER_CONSTANT_BUFFER_API_SLOT_COUNT, Buffers); break;
 	case SF_Compute:	Direct3DDeviceIMContext->CSGetConstantBuffers(0, D3D11_COMMONSHADER_CONSTANT_BUFFER_API_SLOT_COUNT, Buffers); break;
@@ -92,8 +88,6 @@ void FD3D11StateCacheBase::VerifyShaderResourceViews()
 	switch (ShaderFrequency)
 	{
 	case SF_Vertex:		Direct3DDeviceIMContext->VSGetShaderResources(0, D3D11_COMMONSHADER_INPUT_RESOURCE_SLOT_COUNT, Views); break;
-	case SF_Hull:		Direct3DDeviceIMContext->HSGetShaderResources(0, D3D11_COMMONSHADER_INPUT_RESOURCE_SLOT_COUNT, Views); break;
-	case SF_Domain:		Direct3DDeviceIMContext->DSGetShaderResources(0, D3D11_COMMONSHADER_INPUT_RESOURCE_SLOT_COUNT, Views); break;
 	case SF_Geometry:	Direct3DDeviceIMContext->GSGetShaderResources(0, D3D11_COMMONSHADER_INPUT_RESOURCE_SLOT_COUNT, Views); break;
 	case SF_Pixel:		Direct3DDeviceIMContext->PSGetShaderResources(0, D3D11_COMMONSHADER_INPUT_RESOURCE_SLOT_COUNT, Views); break;
 	case SF_Compute:	Direct3DDeviceIMContext->CSGetShaderResources(0, D3D11_COMMONSHADER_INPUT_RESOURCE_SLOT_COUNT, Views); break;
@@ -135,22 +129,16 @@ void FD3D11StateCacheBase::VerifyCacheState()
 	// Verify Shader States
 	{
 		TRefCountPtr<ID3D11VertexShader> VertexShader;
-		TRefCountPtr<ID3D11HullShader> HullShader;
-		TRefCountPtr<ID3D11DomainShader> DomainShader;
 		TRefCountPtr<ID3D11GeometryShader> GeometryShader;
 		TRefCountPtr<ID3D11PixelShader> PixelShader;
 		TRefCountPtr<ID3D11ComputeShader> ComputeShader;
 
 		Direct3DDeviceIMContext->VSGetShader(VertexShader.GetInitReference(), nullptr, nullptr);
-		Direct3DDeviceIMContext->HSGetShader(HullShader.GetInitReference(), nullptr, nullptr);
-		Direct3DDeviceIMContext->DSGetShader(DomainShader.GetInitReference(), nullptr, nullptr);
 		Direct3DDeviceIMContext->GSGetShader(GeometryShader.GetInitReference(), nullptr, nullptr);
 		Direct3DDeviceIMContext->PSGetShader(PixelShader.GetInitReference(), nullptr, nullptr);
 		Direct3DDeviceIMContext->CSGetShader(ComputeShader.GetInitReference(), nullptr, nullptr);
 
 		check(VertexShader.GetReference() == CurrentVertexShader);
-		check(HullShader.GetReference() == CurrentHullShader);
-		check(DomainShader.GetReference() == CurrentDomainShader);
 		check(GeometryShader.GetReference() == CurrentGeometryShader);
 		check(PixelShader.GetReference() == CurrentPixelShader);
 		check(ComputeShader.GetReference() == CurrentComputeShader);
@@ -208,8 +196,6 @@ void FD3D11StateCacheBase::VerifyCacheState()
 	// Verify Sampler States
 	{
 		VerifySamplerStates<SF_Vertex>();
-		VerifySamplerStates<SF_Hull>();
-		VerifySamplerStates<SF_Domain>();
 		VerifySamplerStates<SF_Geometry>();
 		VerifySamplerStates<SF_Pixel>();
 		VerifySamplerStates<SF_Compute>();
@@ -258,8 +244,6 @@ void FD3D11StateCacheBase::VerifyCacheState()
 	// Verify Constant Buffers
 	{
 		((FD3D11StateCache*)this)->VerifyConstantBuffers<SF_Vertex>();
-		((FD3D11StateCache*)this)->VerifyConstantBuffers<SF_Hull>();
-		((FD3D11StateCache*)this)->VerifyConstantBuffers<SF_Domain>();
 		((FD3D11StateCache*)this)->VerifyConstantBuffers<SF_Geometry>();
 		((FD3D11StateCache*)this)->VerifyConstantBuffers<SF_Pixel>();
 		((FD3D11StateCache*)this)->VerifyConstantBuffers<SF_Compute>();
@@ -268,8 +252,6 @@ void FD3D11StateCacheBase::VerifyCacheState()
 	// Verify Shader Resource Views
 	{
 		VerifyShaderResourceViews<SF_Vertex>();
-		VerifyShaderResourceViews<SF_Hull>();
-		VerifyShaderResourceViews<SF_Domain>();
 		VerifyShaderResourceViews<SF_Geometry>();
 		VerifyShaderResourceViews<SF_Pixel>();
 		VerifyShaderResourceViews<SF_Compute>();
@@ -310,8 +292,6 @@ void FD3D11StateCacheBase::ClearState()
 
 	// Shader Cache
 	CurrentVertexShader = nullptr;
-	CurrentHullShader = nullptr;
-	CurrentDomainShader = nullptr;
 	CurrentGeometryShader = nullptr;
 	CurrentPixelShader = nullptr;
 	CurrentComputeShader = nullptr;
diff --git a/Engine/Source/Runtime/Windows/D3D11RHI/Private/D3D11StateCachePrivate.h b/Engine/Source/Runtime/Windows/D3D11RHI/Private/D3D11StateCachePrivate.h
--- a/Engine/Source/Runtime/Windows/D3D11RHI/Private/D3D11StateCachePrivate.h
+++ b/Engine/Source/Runtime/Windows/D3D11RHI/Private/D3D11StateCachePrivate.h
@@ -82,8 +82,6 @@ protected:
 
 	// Shader Cache
 	ID3D11VertexShader* CurrentVertexShader;
-	ID3D11HullShader* CurrentHullShader;
-	ID3D11DomainShader* CurrentDomainShader;
 	ID3D11GeometryShader* CurrentGeometryShader;
 	ID3D11PixelShader* CurrentPixelShader;
 	ID3D11ComputeShader* CurrentComputeShader;
@@ -142,8 +140,6 @@ protected:
 		switch (ShaderFrequency)
 		{
 		case SF_Vertex:		Direct3DDeviceIMContext->VSSetShaderResources(ResourceIndex, 1, &SRV); break;
-		case SF_Hull:		Direct3DDeviceIMContext->HSSetShaderResources(ResourceIndex, 1, &SRV); break;
-		case SF_Domain:		Direct3DDeviceIMContext->DSSetShaderResources(ResourceIndex, 1, &SRV); break;
 		case SF_Geometry:	Direct3DDeviceIMContext->GSSetShaderResources(ResourceIndex, 1, &SRV); break;
 		case SF_Pixel:		Direct3DDeviceIMContext->PSSetShaderResources(ResourceIndex, 1, &SRV); break;
 		case SF_Compute:	Direct3DDeviceIMContext->CSSetShaderResources(ResourceIndex, 1, &SRV); break;
@@ -157,8 +153,6 @@ protected:
 		switch (ShaderFrequency)
 		{
 		case SF_Vertex:		Direct3DDeviceIMContext->VSSetSamplers(SamplerIndex, 1, &SamplerState); break;
-		case SF_Hull:		Direct3DDeviceIMContext->HSSetSamplers(SamplerIndex, 1, &SamplerState); break;
-		case SF_Domain:		Direct3DDeviceIMContext->DSSetSamplers(SamplerIndex, 1, &SamplerState); break;
 		case SF_Geometry:	Direct3DDeviceIMContext->GSSetSamplers(SamplerIndex, 1, &SamplerState); break;
 		case SF_Pixel:		Direct3DDeviceIMContext->PSSetSamplers(SamplerIndex, 1, &SamplerState); break;
 		case SF_Compute:	Direct3DDeviceIMContext->CSSetSamplers(SamplerIndex, 1, &SamplerState); break;
@@ -172,8 +166,6 @@ protected:
 		switch (ShaderFrequency)
 		{
 		case SF_Vertex:		Direct3DDeviceIMContext->VSSetConstantBuffers(SlotIndex, 1, &ConstantBuffer); break;
-		case SF_Hull:		Direct3DDeviceIMContext->HSSetConstantBuffers(SlotIndex, 1, &ConstantBuffer); break;
-		case SF_Domain:		Direct3DDeviceIMContext->DSSetConstantBuffers(SlotIndex, 1, &ConstantBuffer); break;
 		case SF_Geometry:	Direct3DDeviceIMContext->GSSetConstantBuffers(SlotIndex, 1, &ConstantBuffer); break;
 		case SF_Pixel:		Direct3DDeviceIMContext->PSSetConstantBuffers(SlotIndex, 1, &ConstantBuffer); break;
 		case SF_Compute:	Direct3DDeviceIMContext->CSSetConstantBuffers(SlotIndex, 1, &ConstantBuffer); break;
@@ -303,8 +295,6 @@ public:
 		switch (ShaderFrequency)
 		{
 		case SF_Vertex:		Direct3DDeviceIMContext->VSSetConstantBuffers(0, D3D11_COMMONSHADER_CONSTANT_BUFFER_API_SLOT_COUNT, Empty); break;
-		case SF_Hull:		Direct3DDeviceIMContext->HSSetConstantBuffers(0, D3D11_COMMONSHADER_CONSTANT_BUFFER_API_SLOT_COUNT, Empty); break;
-		case SF_Domain:		Direct3DDeviceIMContext->DSSetConstantBuffers(0, D3D11_COMMONSHADER_CONSTANT_BUFFER_API_SLOT_COUNT, Empty); break;
 		case SF_Geometry:	Direct3DDeviceIMContext->GSSetConstantBuffers(0, D3D11_COMMONSHADER_CONSTANT_BUFFER_API_SLOT_COUNT, Empty); break;
 		case SF_Pixel:		Direct3DDeviceIMContext->PSSetConstantBuffers(0, D3D11_COMMONSHADER_CONSTANT_BUFFER_API_SLOT_COUNT, Empty); break;
 		case SF_Compute:	Direct3DDeviceIMContext->CSSetConstantBuffers(0, D3D11_COMMONSHADER_CONSTANT_BUFFER_API_SLOT_COUNT, Empty); break;
@@ -339,12 +329,6 @@ public:
 			case SF_Vertex:		
 				Direct3DDeviceIMContext->VSGetShaderResources(StartResourceIndex, NumResources, SRV); 
 				break;
-			case SF_Hull:		
-				Direct3DDeviceIMContext->HSGetShaderResources(StartResourceIndex, NumResources, SRV); 
-				break;
-			case SF_Domain:		
-				Direct3DDeviceIMContext->DSGetShaderResources(StartResourceIndex, NumResources, SRV); 
-				break;
 			case SF_Geometry:	
 				Direct3DDeviceIMContext->GSGetShaderResources(StartResourceIndex, NumResources, SRV); 
 				break;
@@ -457,12 +441,6 @@ public:
 				case SF_Vertex:		
 					Direct3DDeviceIMContext->VSGetSamplers(StartSamplerIndex, NumSamplerIndexes, SamplerStates); 
 					break;
-				case SF_Hull:		
-					Direct3DDeviceIMContext->HSGetSamplers(StartSamplerIndex, NumSamplerIndexes, SamplerStates); 
-					break;
-				case SF_Domain:		
-					Direct3DDeviceIMContext->DSGetSamplers(StartSamplerIndex, NumSamplerIndexes, SamplerStates); 
-					break;
 				case SF_Geometry:	
 					Direct3DDeviceIMContext->GSGetSamplers(StartSamplerIndex, NumSamplerIndexes, SamplerStates); 
 					break;
@@ -515,8 +493,6 @@ template <EShaderFrequency ShaderFrequency>
 			switch (ShaderFrequency)
 			{
 			case SF_Vertex:		Direct3DDeviceIMContext1->VSGetConstantBuffers(StartSlot, NumBuffers, ConstantBuffers); break;
-			case SF_Hull:		Direct3DDeviceIMContext1->HSGetConstantBuffers(StartSlot, NumBuffers, ConstantBuffers); break;
-			case SF_Domain:		Direct3DDeviceIMContext1->DSGetConstantBuffers(StartSlot, NumBuffers, ConstantBuffers); break;
 			case SF_Geometry:	Direct3DDeviceIMContext1->GSGetConstantBuffers(StartSlot, NumBuffers, ConstantBuffers); break;
 			case SF_Pixel:		Direct3DDeviceIMContext1->PSGetConstantBuffers(StartSlot, NumBuffers, ConstantBuffers); break;
 			case SF_Compute:	Direct3DDeviceIMContext1->CSGetConstantBuffers(StartSlot, NumBuffers, ConstantBuffers); break;
@@ -675,62 +651,6 @@ template <EShaderFrequency ShaderFrequency>
 #endif
 	}
 
-	D3D11_STATE_CACHE_INLINE void SetHullShader(ID3D11HullShader* Shader)
-	{
-#if D3D11_ALLOW_STATE_CACHE
-		D3D11_STATE_CACHE_VERIFY_PRE();
-		if ((CurrentHullShader != Shader) || GD3D11SkipStateCaching)
-		{
-			CurrentHullShader = Shader;
-			Direct3DDeviceIMContext->HSSetShader(Shader, nullptr, 0);
-		}
-		D3D11_STATE_CACHE_VERIFY_POST();
-#else
-		Direct3DDeviceIMContext->HSSetShader(Shader, nullptr, 0);
-#endif
-	}
-
-	D3D11_STATE_CACHE_INLINE void GetHullShader(ID3D11HullShader** HullShader)
-	{
-#if D3D11_ALLOW_STATE_CACHE
-		*HullShader = CurrentHullShader;
-		if (CurrentHullShader)
-		{
-			CurrentHullShader->AddRef();
-		}
-#else
-		Direct3DDeviceIMContext->HSGetShader(HullShader, nullptr, nullptr);
-#endif
-	}
-
-	D3D11_STATE_CACHE_INLINE void SetDomainShader(ID3D11DomainShader* Shader)
-	{
-#if D3D11_ALLOW_STATE_CACHE
-		D3D11_STATE_CACHE_VERIFY_PRE();
-		if ((CurrentDomainShader != Shader) || GD3D11SkipStateCaching)
-		{
-			CurrentDomainShader = Shader;
-			Direct3DDeviceIMContext->DSSetShader(Shader, nullptr, 0);
-		}
-		D3D11_STATE_CACHE_VERIFY_POST();
-#else
-		Direct3DDeviceIMContext->DSSetShader(Shader, nullptr, 0);
-#endif
-	}
-
-	D3D11_STATE_CACHE_INLINE void GetDomainShader(ID3D11DomainShader** DomainShader)
-	{
-#if D3D11_ALLOW_STATE_CACHE
-		*DomainShader = CurrentDomainShader;
-		if (CurrentDomainShader)
-		{
-			CurrentDomainShader->AddRef();
-		}
-#else
-		Direct3DDeviceIMContext->DSGetShader(DomainShader, nullptr, nullptr);
-#endif
-	}
-
 	D3D11_STATE_CACHE_INLINE void SetGeometryShader(ID3D11GeometryShader* Shader)
 	{
 #if D3D11_ALLOW_STATE_CACHE
diff --git a/Engine/Source/Runtime/Windows/D3D11RHI/Private/D3D11Viewport.cpp b/Engine/Source/Runtime/Windows/D3D11RHI/Private/D3D11Viewport.cpp
--- a/Engine/Source/Runtime/Windows/D3D11RHI/Private/D3D11Viewport.cpp
+++ b/Engine/Source/Runtime/Windows/D3D11RHI/Private/D3D11Viewport.cpp
@@ -769,8 +769,6 @@ void FD3D11DynamicRHI::RHIEndDrawingViewport(FRHIViewport* ViewportRHI,bool bPre
 	MaxBoundVertexBufferIndex = INDEX_NONE;
 	
 	StateCache.SetPixelShader(nullptr);
-	StateCache.SetHullShader(nullptr);
-	StateCache.SetDomainShader(nullptr);
 	StateCache.SetGeometryShader(nullptr);
 	// Compute Shader is set to NULL after each Dispatch call, so no need to clear it here
 
diff --git a/Engine/Source/Runtime/Windows/D3D11RHI/Private/HoloLens/HoloLensD3D11ConstantBuffer.cpp b/Engine/Source/Runtime/Windows/D3D11RHI/Private/HoloLens/HoloLensD3D11ConstantBuffer.cpp
--- a/Engine/Source/Runtime/Windows/D3D11RHI/Private/HoloLens/HoloLensD3D11ConstantBuffer.cpp
+++ b/Engine/Source/Runtime/Windows/D3D11RHI/Private/HoloLens/HoloLensD3D11ConstantBuffer.cpp
@@ -101,8 +101,6 @@ void FD3D11DynamicRHI::InitConstantBuffers()
 	//	constants
 	VSConstantBuffers.Empty(MAX_CONSTANT_BUFFER_SLOTS);
 	PSConstantBuffers.Empty(MAX_CONSTANT_BUFFER_SLOTS);
-	HSConstantBuffers.Empty(MAX_CONSTANT_BUFFER_SLOTS);
-	DSConstantBuffers.Empty(MAX_CONSTANT_BUFFER_SLOTS);
 	GSConstantBuffers.Empty(MAX_CONSTANT_BUFFER_SLOTS);
 	CSConstantBuffers.Empty(MAX_CONSTANT_BUFFER_SLOTS);
 	for(int32 BufferIndex = 0;BufferIndex < MAX_CONSTANT_BUFFER_SLOTS;BufferIndex++)
@@ -120,8 +118,6 @@ void FD3D11DynamicRHI::InitConstantBuffers()
 		// New circular buffer system for faster constant uploads.  Avoids CopyResource and speeds things up considerably
 		VSConstantBuffers.Add(new FWinD3D11ConstantBuffer(this, Size, SubBuffers));
 		PSConstantBuffers.Add(new FWinD3D11ConstantBuffer(this, Size, SubBuffers));
-		HSConstantBuffers.Add(new FWinD3D11ConstantBuffer(this, Size));
-		DSConstantBuffers.Add(new FWinD3D11ConstantBuffer(this, Size));
 		GSConstantBuffers.Add(new FWinD3D11ConstantBuffer(this, Size));
 		CSConstantBuffers.Add(new FWinD3D11ConstantBuffer(this, Size));
 	}
diff --git a/Engine/Source/Runtime/Windows/D3D11RHI/Private/Windows/WindowsD3D11ConstantBuffer.cpp b/Engine/Source/Runtime/Windows/D3D11RHI/Private/Windows/WindowsD3D11ConstantBuffer.cpp
--- a/Engine/Source/Runtime/Windows/D3D11RHI/Private/Windows/WindowsD3D11ConstantBuffer.cpp
+++ b/Engine/Source/Runtime/Windows/D3D11RHI/Private/Windows/WindowsD3D11ConstantBuffer.cpp
@@ -102,8 +102,6 @@ void FD3D11DynamicRHI::InitConstantBuffers()
 	//	constants
 	VSConstantBuffers.Empty(MAX_CONSTANT_BUFFER_SLOTS);
 	PSConstantBuffers.Empty(MAX_CONSTANT_BUFFER_SLOTS);
-	HSConstantBuffers.Empty(MAX_CONSTANT_BUFFER_SLOTS);
-	DSConstantBuffers.Empty(MAX_CONSTANT_BUFFER_SLOTS);
 	GSConstantBuffers.Empty(MAX_CONSTANT_BUFFER_SLOTS);
 	CSConstantBuffers.Empty(MAX_CONSTANT_BUFFER_SLOTS);
 	for(int32 BufferIndex = 0;BufferIndex < MAX_CONSTANT_BUFFER_SLOTS;BufferIndex++)
@@ -121,8 +119,6 @@ void FD3D11DynamicRHI::InitConstantBuffers()
 		// New circular buffer system for faster constant uploads.  Avoids CopyResource and speeds things up considerably
 		VSConstantBuffers.Add(new FWinD3D11ConstantBuffer(this, Size, SubBuffers));
 		PSConstantBuffers.Add(new FWinD3D11ConstantBuffer(this, Size, SubBuffers));
-		HSConstantBuffers.Add(new FWinD3D11ConstantBuffer(this, Size));
-		DSConstantBuffers.Add(new FWinD3D11ConstantBuffer(this, Size));
 		GSConstantBuffers.Add(new FWinD3D11ConstantBuffer(this, Size));
 		CSConstantBuffers.Add(new FWinD3D11ConstantBuffer(this, Size));
 	}
diff --git a/Engine/Source/Runtime/Windows/D3D11RHI/Public/D3D11Resources.h b/Engine/Source/Runtime/Windows/D3D11RHI/Public/D3D11Resources.h
--- a/Engine/Source/Runtime/Windows/D3D11RHI/Public/D3D11Resources.h
+++ b/Engine/Source/Runtime/Windows/D3D11RHI/Public/D3D11Resources.h
@@ -74,24 +74,6 @@ public:
 	TRefCountPtr<ID3D11GeometryShader> Resource;
 };
 
-class FD3D11HullShader : public FRHIHullShader, public FD3D11ShaderData
-{
-public:
-	enum { StaticFrequency = SF_Hull };
-
-	/** The shader resource. */
-	TRefCountPtr<ID3D11HullShader> Resource;
-};
-
-class FD3D11DomainShader : public FRHIDomainShader, public FD3D11ShaderData
-{
-public:
-	enum { StaticFrequency = SF_Domain };
-
-	/** The shader resource. */
-	TRefCountPtr<ID3D11DomainShader> Resource;
-};
-
 class FD3D11PixelShader : public FRHIPixelShader, public FD3D11ShaderData
 {
 public:
@@ -123,8 +105,6 @@ public:
 	TRefCountPtr<ID3D11InputLayout> InputLayout;
 	TRefCountPtr<ID3D11VertexShader> VertexShader;
 	TRefCountPtr<ID3D11PixelShader> PixelShader;
-	TRefCountPtr<ID3D11HullShader> HullShader;
-	TRefCountPtr<ID3D11DomainShader> DomainShader;
 	TRefCountPtr<ID3D11GeometryShader> GeometryShader;
 
 	bool bShaderNeedsGlobalConstantBuffer[SF_NumStandardFrequencies];
@@ -146,8 +126,6 @@ public:
 	 */
 	FORCEINLINE FD3D11VertexShader*   GetVertexShader() const   { return (FD3D11VertexShader*)CacheLink.GetVertexShader(); }
 	FORCEINLINE FD3D11PixelShader*    GetPixelShader() const    { return (FD3D11PixelShader*)CacheLink.GetPixelShader(); }
-	FORCEINLINE FD3D11HullShader*     GetHullShader() const     { return (FD3D11HullShader*)CacheLink.GetHullShader(); }
-	FORCEINLINE FD3D11DomainShader*   GetDomainShader() const   { return (FD3D11DomainShader*)CacheLink.GetDomainShader(); }
 	FORCEINLINE FD3D11GeometryShader* GetGeometryShader() const { return (FD3D11GeometryShader*)CacheLink.GetGeometryShader(); }
 };
 
@@ -751,16 +729,6 @@ struct TD3D11ResourceTraits<FRHIGeometryShader>
 	typedef FD3D11GeometryShader TConcreteType;
 };
 template<>
-struct TD3D11ResourceTraits<FRHIHullShader>
-{
-	typedef FD3D11HullShader TConcreteType;
-};
-template<>
-struct TD3D11ResourceTraits<FRHIDomainShader>
-{
-	typedef FD3D11DomainShader TConcreteType;
-};
-template<>
 struct TD3D11ResourceTraits<FRHIPixelShader>
 {
 	typedef FD3D11PixelShader TConcreteType;
