Fixed D11 RHICopyTexture bug, where source box size was incorrectly aligned to dest texture. For example this causes a 2x4 UINT4->BC7 copy to become a 4x4 copy.

diff --git a/Engine/Source/Runtime/Windows/D3D11RHI/Private/D3D11Texture.cpp b/Engine/Source/Runtime/Windows/D3D11RHI/Private/D3D11Texture.cpp
--- a/Engine/Source/Runtime/Windows/D3D11RHI/Private/D3D11Texture.cpp
+++ b/Engine/Source/Runtime/Windows/D3D11RHI/Private/D3D11Texture.cpp
@@ -2669,7 +2669,7 @@ void FD3D11DynamicRHI::RHICopyTexture(FRHITexture* SourceTextureRHI, FRHITexture
 
 	if (CopyInfo.Size != FIntVector::ZeroValue)
 	{
-		const FPixelFormatInfo& PixelFormatInfo = GPixelFormats[DestTextureRHI->GetFormat()];
+		const FPixelFormatInfo& PixelFormatInfo = GPixelFormats[SourceTextureRHI->GetFormat()];
 
 		D3D11_BOX SrcBox;
 		SrcBox.left = CopyInfo.SourcePosition.X;
