Removing 'Fast VRAM allocations' from D3D11 RHI.

diff --git a/Engine/Source/Runtime/Windows/D3D11RHI/Private/D3D11Buffer.cpp b/Engine/Source/Runtime/Windows/D3D11RHI/Private/D3D11Buffer.cpp
--- a/Engine/Source/Runtime/Windows/D3D11RHI/Private/D3D11Buffer.cpp
+++ b/Engine/Source/Runtime/Windows/D3D11RHI/Private/D3D11Buffer.cpp
@@ -79,14 +79,6 @@ FBufferRHIRef FD3D11DynamicRHI::RHICreateBuffer(uint32 Size, EBufferUsageFlags U
 		}
 	}
 
-	if (FPlatformMemory::SupportsFastVRAMMemory())
-	{
-		if (Usage & BUF_FastVRAM)
-		{
-			FFastVRAMAllocator::GetFastVRAMAllocator()->AllocUAVBuffer(Desc);
-		}
-	}
-
 	// If a resource array was provided for the resource, create the resource pre-populated
 	D3D11_SUBRESOURCE_DATA InitData;
 	D3D11_SUBRESOURCE_DATA* pInitData = NULL;
diff --git a/Engine/Source/Runtime/Windows/D3D11RHI/Private/D3D11RHI.cpp b/Engine/Source/Runtime/Windows/D3D11RHI/Private/D3D11RHI.cpp
--- a/Engine/Source/Runtime/Windows/D3D11RHI/Private/D3D11RHI.cpp
+++ b/Engine/Source/Runtime/Windows/D3D11RHI/Private/D3D11RHI.cpp
@@ -758,16 +758,3 @@ void UpdateBufferStats(TRefCountPtr<ID3D11Buffer> Buffer, bool bAllocating)
 #endif
 	}
 }
-
-#ifndef PLATFORM_IMPLEMENTS_FASTVRAMALLOCATOR
-	#define PLATFORM_IMPLEMENTS_FASTVRAMALLOCATOR		0
-#endif
-
-#if !PLATFORM_IMPLEMENTS_FASTVRAMALLOCATOR
-FFastVRAMAllocator* FFastVRAMAllocator::GetFastVRAMAllocator()
-{
-	static FFastVRAMAllocator FastVRAMAllocatorSingleton;
-	return &FastVRAMAllocatorSingleton;
-}
-#endif
-
diff --git a/Engine/Source/Runtime/Windows/D3D11RHI/Private/D3D11RHIPrivate.h b/Engine/Source/Runtime/Windows/D3D11RHI/Private/D3D11RHIPrivate.h
--- a/Engine/Source/Runtime/Windows/D3D11RHI/Private/D3D11RHIPrivate.h
+++ b/Engine/Source/Runtime/Windows/D3D11RHI/Private/D3D11RHIPrivate.h
@@ -1169,53 +1169,6 @@ inline bool HasStencilBits(DXGI_FORMAT InFormat)
 	return false;
 }
 
-/**
- *	Default 'Fast VRAM' allocator...
- */
-class FFastVRAMAllocator
-{
-public:
-	FFastVRAMAllocator() {}
-
-	virtual ~FFastVRAMAllocator() {}
-
-	/**
-	 *	IMPORTANT: This function CAN modify the TextureDesc!
-	 */
-	virtual FVRamAllocation AllocTexture2D(D3D11_TEXTURE2D_DESC& TextureDesc)
-	{
-		return FVRamAllocation();
-	}
-
-	/**
-	 *	IMPORTANT: This function CAN modify the TextureDesc!
-	 */
-	virtual FVRamAllocation AllocTexture3D(D3D11_TEXTURE3D_DESC& TextureDesc)
-	{
-		return FVRamAllocation();
-	}
-
-	/**
-	 *	IMPORTANT: This function CAN modify the BufferDesc!
-	 */
-	virtual FVRamAllocation AllocUAVBuffer(D3D11_BUFFER_DESC& BufferDesc)
-	{
-		return FVRamAllocation();
-	}
-
-	template< typename t_A, typename t_B >
-	static t_A RoundUpToNextMultiple( const t_A& a, const t_B& b )
-	{
-		return ( ( a - 1 ) / b + 1 ) * b;
-	}
-
-	static FFastVRAMAllocator* GetFastVRAMAllocator();
-};
-
-
-
-
-
 // 1d, 31 bit (uses the sign bit for internal use), O(n) where n is the amount of elements stored
 // does not enforce any alignment
 // unoccupied regions get compacted but occupied don't get compacted
@@ -1346,11 +1299,6 @@ public:
 		}
 	}
 
-
-	// All resources in ESRAM must be 64KiB aligned
-//	uint32 AlignedByteOffset = FFastVRAMAllocator::RoundUpToNextMultiple(ESRAMByteOffset, ESRAMMinumumAlignment );
-
-
 	// @param InSize >0
 	FRange AllocRange(uint32 InSize)//, uint32 Alignment)
 	{
diff --git a/Engine/Source/Runtime/Windows/D3D11RHI/Private/D3D11Texture.cpp b/Engine/Source/Runtime/Windows/D3D11RHI/Private/D3D11Texture.cpp
--- a/Engine/Source/Runtime/Windows/D3D11RHI/Private/D3D11Texture.cpp
+++ b/Engine/Source/Runtime/Windows/D3D11RHI/Private/D3D11Texture.cpp
@@ -766,16 +766,6 @@ TD3D11Texture2D<BaseResourceType>* FD3D11DynamicRHI::CreateD3D11Texture2D(uint32
 		bPooledTexture = false;
 	}
 
-	FVRamAllocation VRamAllocation;
-
-	if (FPlatformMemory::SupportsFastVRAMMemory())
-	{
-		if (Flags & TexCreate_FastVRAM)
-		{
-			VRamAllocation = FFastVRAMAllocator::GetFastVRAMAllocator()->AllocTexture2D(TextureDesc);
-		}
-	}
-
 	TRefCountPtr<ID3D11Texture2D> TextureResource;
 	TRefCountPtr<ID3D11ShaderResourceView> ShaderResourceView;
 	TArray<TRefCountPtr<ID3D11RenderTargetView> > RenderTargetViews;
@@ -1061,8 +1051,6 @@ TD3D11Texture2D<BaseResourceType>* FD3D11DynamicRHI::CreateD3D11Texture2D(uint32
 #endif
 		);
 
-	Texture2D->ResourceInfo.VRamAllocation = VRamAllocation;
-
 	D3D11TextureAllocated(*Texture2D);
 	
 #if !PLATFORM_HOLOLENS
@@ -1150,16 +1138,6 @@ FD3D11Texture3D* FD3D11DynamicRHI::CreateD3D11Texture3D(uint32 SizeX,uint32 Size
 		}
 	}
 
-	FVRamAllocation VRamAllocation;
-
-	if (FPlatformMemory::SupportsFastVRAMMemory())
-	{
-		if (Flags & TexCreate_FastVRAM)
-		{
-			VRamAllocation = FFastVRAMAllocator::GetFastVRAMAllocator()->AllocTexture3D(TextureDesc);
-		}
-	}
-
 	TRefCountPtr<ID3D11Texture3D> TextureResource;
 	const D3D11_SUBRESOURCE_DATA* SubResData = CreateInfo.BulkData != nullptr ? (const D3D11_SUBRESOURCE_DATA*)SubResourceData.GetData() : nullptr;
 	VERIFYD3D11CREATETEXTURERESULT(
@@ -1212,8 +1190,6 @@ FD3D11Texture3D* FD3D11DynamicRHI::CreateD3D11Texture3D(uint32 SizeX,uint32 Size
 	RenderTargetViews.Add(RenderTargetView);
 	FD3D11Texture3D* Texture3D = new FD3D11Texture3D(this,TextureResource,ShaderResourceView,RenderTargetViews,SizeX,SizeY,SizeZ,NumMips,(EPixelFormat)Format,Flags, CreateInfo.ClearValueBinding);
 
-	Texture3D->ResourceInfo.VRamAllocation = VRamAllocation;
-
 	D3D11TextureAllocated(*Texture3D);
 #if !PLATFORM_HOLOLENS
 	if (IsRHIDeviceNVIDIA() && (Flags & TexCreate_AFRManual))
