U5 - deletion of RHIResources is not happening on the RenderingThread anymore

diff --git a/Engine/Source/Runtime/Windows/D3D11RHI/Private/D3D11UniformBuffer.cpp b/Engine/Source/Runtime/Windows/D3D11RHI/Private/D3D11UniformBuffer.cpp
--- a/Engine/Source/Runtime/Windows/D3D11RHI/Private/D3D11UniformBuffer.cpp
+++ b/Engine/Source/Runtime/Windows/D3D11RHI/Private/D3D11UniformBuffer.cpp
@@ -352,7 +352,7 @@ FD3D11UniformBuffer::~FD3D11UniformBuffer()
 	// Do not return the allocation to the pool if it is in the dynamic constant buffer!
 	if (!RingAllocation.IsValid() && Resource != nullptr)
 	{
-		check(IsInRenderingThread());
+		check(IsInRHIThread() || IsInRenderingThread());
 		D3D11_BUFFER_DESC Desc;
 		Resource->GetDesc(&Desc);
 
