Add more info to crashes when no UB is bound.

diff --git a/Engine/Source/Runtime/Windows/D3D11RHI/Private/D3D11Commands.cpp b/Engine/Source/Runtime/Windows/D3D11RHI/Private/D3D11Commands.cpp
--- a/Engine/Source/Runtime/Windows/D3D11RHI/Private/D3D11Commands.cpp
+++ b/Engine/Source/Runtime/Windows/D3D11RHI/Private/D3D11Commands.cpp
@@ -1373,8 +1373,12 @@ void FD3D11DynamicRHI::SetResourcesFromTables(const ShaderType* RESTRICT Shader)
 
 		if (!Buffer)
 		{
-			UE_LOG(LogD3D11RHI, Fatal, TEXT("Shader expected a uniform buffer at slot %u but got null instead.  Rendering code needs to set a valid uniform buffer for this slot."), 
-				BufferIndex);
+			FString ShaderUB;
+			if (BufferIndex < Shader->UniformBuffers.Num())
+			{
+				Shader->UniformBuffers[BufferIndex].ToString(ShaderUB);
+			}
+			UE_LOG(LogD3D11RHI, Fatal, TEXT("Shader expected a uniform buffer at slot %u but got null instead (Shader='%s' UB='%s').  Rendering code needs to set a valid uniform buffer for this slot."), BufferIndex, Shader->GetShaderName(), *ShaderUB);
 		}
 
 #if !(UE_BUILD_SHIPPING || UE_BUILD_TEST)
@@ -1444,6 +1448,17 @@ int32 FD3D11DynamicRHI::SetUAVPSResourcesFromTables(const ShaderType* RESTRICT S
 		FD3D11UniformBuffer* Buffer = (FD3D11UniformBuffer*)BoundUniformBuffers[ShaderType::StaticFrequency][BufferIndex].GetReference();
 
 		check(BufferIndex < Shader->ShaderResourceTable.ResourceTableLayoutHashes.Num());
+
+		if (!Buffer)
+		{
+			FString ShaderUB;
+			if (BufferIndex < Shader->UniformBuffers.Num())
+			{
+				Shader->UniformBuffers[BufferIndex].ToString(ShaderUB);
+			}
+			UE_LOG(LogD3D11RHI, Fatal, TEXT("Shader expected a uniform buffer at slot %u but got null instead (Shader='%s' UB='%s').  Rendering code needs to set a valid uniform buffer for this slot."), BufferIndex, Shader->GetShaderName(), *ShaderUB);
+		}
+
 	#if !(UE_BUILD_SHIPPING || UE_BUILD_TEST)
 		const TCHAR* LayoutName = *Buffer->GetLayout().GetDebugName();
 	#else
