U5 - fix missing IsInRHIThread checks after mofing Resource deletes to the RHI thread.

diff --git a/Engine/Source/Runtime/Windows/D3D11RHI/Private/D3D11Viewport.cpp b/Engine/Source/Runtime/Windows/D3D11RHI/Private/D3D11Viewport.cpp
--- a/Engine/Source/Runtime/Windows/D3D11RHI/Private/D3D11Viewport.cpp
+++ b/Engine/Source/Runtime/Windows/D3D11RHI/Private/D3D11Viewport.cpp
@@ -209,7 +209,7 @@ FD3D11Texture2D* FD3D11Viewport::GetSwapChainSurface(FD3D11DynamicRHI* D3DRHI, E
 
 FD3D11Viewport::~FD3D11Viewport()
 {
-	check(IsInRenderingThread());
+	check(IsInRHIThread() || IsInRenderingThread());
 
 	// Turn off HDR display mode
 	D3DRHI->ShutdownHDR();
