Implementing RHIGetNativeCommandBuffer for D3D11RHI, returns the immediate context.

diff --git a/Engine/Source/Runtime/Windows/D3D11RHI/Private/D3D11Device.cpp b/Engine/Source/Runtime/Windows/D3D11RHI/Private/D3D11Device.cpp
--- a/Engine/Source/Runtime/Windows/D3D11RHI/Private/D3D11Device.cpp
+++ b/Engine/Source/Runtime/Windows/D3D11RHI/Private/D3D11Device.cpp
@@ -686,6 +686,11 @@ void* FD3D11DynamicRHI::RHIGetNativeInstance()
 	return nullptr;
 }
 
+void* FD3D11DynamicRHI::RHIGetNativeCommandBuffer()
+{
+	return (void*)Direct3DDeviceIMContext.GetReference();
+}
+
 static bool CanFormatBeDisplayed(const FD3D11DynamicRHI* InD3DRHI, EPixelFormat InPixelFormat)
 {
 	const DXGI_FORMAT DxgiFormat = FD3D11Viewport::GetRenderTargetFormat(InPixelFormat);
diff --git a/Engine/Source/Runtime/Windows/D3D11RHI/Private/D3D11RHIPrivate.h b/Engine/Source/Runtime/Windows/D3D11RHI/Private/D3D11RHIPrivate.h
--- a/Engine/Source/Runtime/Windows/D3D11RHI/Private/D3D11RHIPrivate.h
+++ b/Engine/Source/Runtime/Windows/D3D11RHI/Private/D3D11RHIPrivate.h
@@ -514,6 +514,7 @@ public:
 	virtual void RHIExecuteCommandList(FRHICommandList* CmdList) final override;
 	virtual void* RHIGetNativeDevice() final override;
 	virtual void* RHIGetNativeInstance() final override;
+	virtual void* RHIGetNativeCommandBuffer() final override;
 	virtual class IRHICommandContext* RHIGetDefaultContext() final override;
 	virtual class IRHICommandContextContainer* RHIGetCommandContextContainer(int32 Index, int32 Num) final override;
 
