non-unity fix

#rnx

diff --git a/Engine/Source/Runtime/Windows/D3D11RHI/Private/D3D11Texture.cpp b/Engine/Source/Runtime/Windows/D3D11RHI/Private/D3D11Texture.cpp
--- a/Engine/Source/Runtime/Windows/D3D11RHI/Private/D3D11Texture.cpp
+++ b/Engine/Source/Runtime/Windows/D3D11RHI/Private/D3D11Texture.cpp
@@ -6,7 +6,7 @@
 
 #include "D3D11RHIPrivate.h"
 
-#if PLATFORM_DESKTOP && !PLATFORM_HOLOLENS
+#if !PLATFORM_HOLOLENS
 // For Depth Bounds Test interface
 #include "Windows/AllowWindowsPlatformTypes.h"
 #include "nvapi.h"
