D3D11RHI fix for wrong adapter selection when multiple adapters are present
* Issue is that the adapters were referenced by index
* Adapters enumeration was sometime performed using EnumAdapters() vs EnumAdapterByGpuPreference()
* Using the same index for both calls would result in different adapters being returned...

On AWS g4ad.16xlarge instances, this issue would cause the engine to think it had selected a discrete adapter, while in fact it created a WARP (software) adapter...

diff --git a/Engine/Source/Runtime/Windows/D3D11RHI/Private/D3D11Device.cpp b/Engine/Source/Runtime/Windows/D3D11RHI/Private/D3D11Device.cpp
--- a/Engine/Source/Runtime/Windows/D3D11RHI/Private/D3D11Device.cpp
+++ b/Engine/Source/Runtime/Windows/D3D11RHI/Private/D3D11Device.cpp
@@ -52,7 +52,7 @@ TAutoConsoleVariable<int32> CVarD3D11ZeroBufferSizeInMB(
 	);
 
 
-FD3D11DynamicRHI::FD3D11DynamicRHI(IDXGIFactory1* InDXGIFactory1,D3D_FEATURE_LEVEL InFeatureLevel, int32 InChosenAdapter, const DXGI_ADAPTER_DESC& InChosenDescription, bool bInSoftwareAdapter) :
+FD3D11DynamicRHI::FD3D11DynamicRHI(IDXGIFactory1* InDXGIFactory1, D3D_FEATURE_LEVEL InFeatureLevel, const FD3D11Adapter& InAdapter) :
 	DXGIFactory1(InDXGIFactory1),
 #if NV_AFTERMATH
 	NVAftermathIMContextHandle(nullptr),
@@ -76,15 +76,13 @@ FD3D11DynamicRHI::FD3D11DynamicRHI(IDXGIFactory1* InDXGIFactory1,D3D_FEATURE_LEV
 	CurrentDSVAccessType(FExclusiveDepthStencil::DepthWrite_StencilWrite),
 	bDiscardSharedConstants(false),	
 	GPUProfilingData(this),
-	ChosenAdapter(InChosenAdapter),
-	bSoftwareAdapter(bInSoftwareAdapter),
-	ChosenDescription(InChosenDescription),
+	Adapter(InAdapter),
 	bAllowVendorDevice(!FParse::Param(FCommandLine::Get(), TEXT("novendordevice")))
 {
 	// This should be called once at the start 
-	check(ChosenAdapter >= 0);
-	check( IsInGameThread() );
-	check( !GIsThreadedRendering );
+	check(Adapter.IsValid());
+	check(IsInGameThread());
+	check(!GIsThreadedRendering);
 
 	// Allocate a buffer of zeroes. This is used when we need to pass D3D memory
 	// that we don't care about and will overwrite with valid data in the future.
@@ -292,27 +290,13 @@ void FD3D11DynamicRHI::RHIGetSupportedResolution( uint32 &Width, uint32 &Height
 
 	{
 		HRESULT HResult = S_OK;
-		TRefCountPtr<IDXGIAdapter> Adapter;
-		HResult = DXGIFactory1->EnumAdapters(ChosenAdapter,Adapter.GetInitReference());
-		if( DXGI_ERROR_NOT_FOUND == HResult )
-		{
-			return;
-		}
-		if( FAILED(HResult) )
-		{
-			return;
-		}
-
-		// get the description of the adapter
-		DXGI_ADAPTER_DESC AdapterDesc;
-		VERIFYD3D11RESULT(Adapter->GetDesc(&AdapterDesc));
 	  
 		// Enumerate outputs for this adapter
 		// TODO: Cap at 1 for default output
 		for(uint32 o = 0;o < 1; o++)
 		{
 			TRefCountPtr<IDXGIOutput> Output;
-			HResult = Adapter->EnumOutputs(o,Output.GetInitReference());
+			HResult = Adapter.DXGIAdapter->EnumOutputs(o, Output.GetInitReference());
 			if(DXGI_ERROR_NOT_FOUND == HResult)
 			{
 				break;
diff --git a/Engine/Source/Runtime/Windows/D3D11RHI/Private/D3D11RHIPrivate.h b/Engine/Source/Runtime/Windows/D3D11RHI/Private/D3D11RHIPrivate.h
--- a/Engine/Source/Runtime/Windows/D3D11RHI/Private/D3D11RHIPrivate.h
+++ b/Engine/Source/Runtime/Windows/D3D11RHI/Private/D3D11RHIPrivate.h
@@ -356,6 +356,45 @@ struct FD3D11TransitionData
 /** Forward declare the context for the AMD AGS utility library. */
 struct AGSContext;
 
+struct FD3D11Adapter
+{
+	/** Null if not supported or FindAdapter() wasn't called. */
+	TRefCountPtr<IDXGIAdapter> DXGIAdapter;
+
+	DXGI_ADAPTER_DESC DXGIAdapterDesc;
+
+	/** The maximum D3D11 feature level supported. 0 if not supported or FindAdapter() wasn't called */
+	D3D_FEATURE_LEVEL MaxSupportedFeatureLevel;
+
+	/** Whether this is a software adapter */
+	bool bSoftwareAdapter;
+
+	/** Whether the GPU is integrated or discrete. */
+	bool bIsIntegrated;
+
+	// constructors
+	FD3D11Adapter() 
+	{
+	}
+
+	FD3D11Adapter(TRefCountPtr<IDXGIAdapter> InDXGIAdapter, D3D_FEATURE_LEVEL InMaxSupportedFeatureLevel, bool bInSoftwareAdatper, bool InIsIntegrated)
+		: DXGIAdapter(InDXGIAdapter)
+		, MaxSupportedFeatureLevel(InMaxSupportedFeatureLevel)
+		, bSoftwareAdapter(bInSoftwareAdatper)
+		, bIsIntegrated(InIsIntegrated)
+	{
+		if (DXGIAdapter.IsValid())
+		{
+			VERIFYD3D11RESULT(DXGIAdapter->GetDesc(&DXGIAdapterDesc));
+		}
+	}
+
+	bool IsValid() const
+	{
+		return DXGIAdapter.IsValid();
+	}
+};
+
 PRAGMA_DISABLE_DEPRECATION_WARNINGS
 
 /** The interface which is implemented by the dynamically bound RHI. */
@@ -367,7 +406,7 @@ public:
 	friend class FWindowsMixedRealityViewport;
 
 	/** Initialization constructor. */
-	FD3D11DynamicRHI(IDXGIFactory1* InDXGIFactory1,D3D_FEATURE_LEVEL InFeatureLevel,int32 InChosenAdapter, const DXGI_ADAPTER_DESC& ChosenDescription, bool bInSoftwareAdapter);
+	FD3D11DynamicRHI(IDXGIFactory1* InDXGIFactory1, D3D_FEATURE_LEVEL InFeatureLevel, const FD3D11Adapter& InAdapter);
 
 	/** Destructor */
 	virtual ~FD3D11DynamicRHI();
@@ -913,12 +952,8 @@ public:
 
 protected:
 	FD3DGPUProfiler GPUProfilingData;
-	// >= 0, was computed before, unless hardware was changed during engine init it should be the same
-	int32 ChosenAdapter;
-	// if this is a software adapter
-	bool bSoftwareAdapter;
-	// we don't use AdapterDesc.Description as there is a bug with Optimus where it can report the wrong name
-	DXGI_ADAPTER_DESC ChosenDescription;
+
+	FD3D11Adapter Adapter;
 
 	// If this is false, disable any IHV optimization/libs
 	bool bAllowVendorDevice;
@@ -1042,32 +1077,6 @@ protected:
 
 PRAGMA_ENABLE_DEPRECATION_WARNINGS
 
-struct FD3D11Adapter
-{
-	/** -1 if not supported or FindAdpater() wasn't called. Ideally we would store a pointer to IDXGIAdapter but it's unlikely the adpaters change during engine init. */
-	int32 AdapterIndex;
-	/** The maximum D3D11 feature level supported. 0 if not supported or FindAdpater() wasn't called */
-	D3D_FEATURE_LEVEL MaxSupportedFeatureLevel;
-	/** Whether this is a software adapter */
-	bool bSoftwareAdapter;
-	/** Whether the GPU is integrated or discrete. */
-	bool bIsIntegrated;
-
-	// constructor
-	FD3D11Adapter(int32 InAdapterIndex = -1, D3D_FEATURE_LEVEL InMaxSupportedFeatureLevel = (D3D_FEATURE_LEVEL)0, bool bInSoftwareAdatper = false, bool InIsIntegrated = false)
-		: AdapterIndex(InAdapterIndex)
-		, MaxSupportedFeatureLevel(InMaxSupportedFeatureLevel)
-		, bSoftwareAdapter(bInSoftwareAdatper)
-		, bIsIntegrated(InIsIntegrated)
-	{
-	}
-
-	bool IsValid() const
-	{
-		return MaxSupportedFeatureLevel != (D3D_FEATURE_LEVEL)0 && AdapterIndex >= 0;
-	}
-};
-
 /** Implements the D3D11RHI module as a dynamic RHI providing module. */
 class FD3D11DynamicRHIModule : public IDynamicRHIModule
 {
@@ -1082,8 +1091,6 @@ public:
 
 private:
 	FD3D11Adapter ChosenAdapter;
-	// we don't use GetDesc().Description as there is a bug with Optimus where it can report the wrong name
-	DXGI_ADAPTER_DESC ChosenDescription;
 
 	// set MaxSupportedFeatureLevel and ChosenAdapter
 	void FindAdapter();
diff --git a/Engine/Source/Runtime/Windows/D3D11RHI/Private/D3D11Texture.cpp b/Engine/Source/Runtime/Windows/D3D11RHI/Private/D3D11Texture.cpp
--- a/Engine/Source/Runtime/Windows/D3D11RHI/Private/D3D11Texture.cpp
+++ b/Engine/Source/Runtime/Windows/D3D11RHI/Private/D3D11Texture.cpp
@@ -687,7 +687,7 @@ TD3D11Texture2D<BaseResourceType>* FD3D11DynamicRHI::CreateD3D11Texture2D(uint32
 	TextureDesc.CPUAccessFlags = CPUAccessFlags;
 	TextureDesc.MiscFlags = bCubeTexture ? D3D11_RESOURCE_MISC_TEXTURECUBE : 0;
 
-	ApplyBC7SoftwareAdapterWorkaround(bSoftwareAdapter, TextureDesc);
+	ApplyBC7SoftwareAdapterWorkaround(Adapter.bSoftwareAdapter, TextureDesc);
 
 	// NV12 doesn't support SRV in NV12 format so don't create SRV for it.
 	// Todo: add support for SRVs of underneath luminance & chrominance textures.
@@ -1287,7 +1287,7 @@ FTexture2DRHIRef FD3D11DynamicRHI::RHIAsyncCreateTexture2D(uint32 SizeX,uint32 S
 	TextureDesc.CPUAccessFlags = 0;
 	TextureDesc.MiscFlags = 0;
 
-	ApplyBC7SoftwareAdapterWorkaround(bSoftwareAdapter, TextureDesc);
+	ApplyBC7SoftwareAdapterWorkaround(Adapter.bSoftwareAdapter, TextureDesc);
 
 	for (uint32 MipIndex = 0; MipIndex < NumInitialMips; ++MipIndex)
 	{
diff --git a/Engine/Source/Runtime/Windows/D3D11RHI/Private/Windows/WindowsD3D11Device.cpp b/Engine/Source/Runtime/Windows/D3D11RHI/Private/Windows/WindowsD3D11Device.cpp
--- a/Engine/Source/Runtime/Windows/D3D11RHI/Private/Windows/WindowsD3D11Device.cpp
+++ b/Engine/Source/Runtime/Windows/D3D11RHI/Private/Windows/WindowsD3D11Device.cpp
@@ -837,6 +837,19 @@ static uint32 CountAdapterOutputs(TRefCountPtr<IDXGIAdapter>& Adapter)
 	return OutputCount;
 }
 
+static void LogDXGIAdapterDesc(const DXGI_ADAPTER_DESC& AdapterDesc)
+{
+	UE_LOG(LogD3D11RHI, Log, TEXT("    Description : %s"), AdapterDesc.Description);
+	UE_LOG(LogD3D11RHI, Log, TEXT("    VendorId    : %04x"), AdapterDesc.VendorId);
+	UE_LOG(LogD3D11RHI, Log, TEXT("    DeviceId    : %04x"), AdapterDesc.DeviceId);
+	UE_LOG(LogD3D11RHI, Log, TEXT("    SubSysId    : %04x"), AdapterDesc.SubSysId);
+	UE_LOG(LogD3D11RHI, Log, TEXT("    Revision    : %04x"), AdapterDesc.Revision);
+	UE_LOG(LogD3D11RHI, Log, TEXT("    DedicatedVideoMemory : %zu bytes"), AdapterDesc.DedicatedVideoMemory);
+	UE_LOG(LogD3D11RHI, Log, TEXT("    DedicatedSystemMemory : %zu bytes"), AdapterDesc.DedicatedSystemMemory);
+	UE_LOG(LogD3D11RHI, Log, TEXT("    SharedSystemMemory : %zu bytes"), AdapterDesc.SharedSystemMemory);
+	UE_LOG(LogD3D11RHI, Log, TEXT("    AdapterLuid : %lu %lu"), AdapterDesc.AdapterLuid.HighPart, AdapterDesc.AdapterLuid.LowPart);
+}
+
 void FD3D11DynamicRHIModule::FindAdapter()
 {
 	// Once we chosen one we don't need to do it again.
@@ -877,8 +890,6 @@ void FD3D11DynamicRHIModule::FindAdapter()
 
 	FD3D11Adapter FirstWithoutIntegratedAdapter;
 	FD3D11Adapter FirstAdapter;
-	// indexed by AdapterIndex, we store it instead of query it later from the created device to prevent some Optimus bug reporting the data/name of the wrong adapter
-	TArray<DXGI_ADAPTER_DESC> AdapterDescription;
 
 	bool bIsAnyAMD = false;
 	bool bIsAnyIntel = false;
@@ -915,9 +926,6 @@ void FD3D11DynamicRHIModule::FindAdapter()
 	// Enumerate the DXGIFactory's adapters.
 	for(uint32 AdapterIndex = 0; LocalEnumAdapters(AdapterIndex,TempAdapter.GetInitReference()) != DXGI_ERROR_NOT_FOUND; ++AdapterIndex)
 	{
-		// to make sure the array elements can be indexed with AdapterIndex
-		DXGI_ADAPTER_DESC& AdapterDesc = AdapterDescription[AdapterDescription.AddZeroed()];
-
 		// Check that if adapter supports D3D11.
 		if(TempAdapter)
 		{
@@ -925,6 +933,7 @@ void FD3D11DynamicRHIModule::FindAdapter()
 			if(SafeTestD3D11CreateDevice(TempAdapter,MinAllowedFeatureLevel,MaxAllowedFeatureLevel,&ActualFeatureLevel))
 			{
 				// Log some information about the available D3D11 adapters.
+				DXGI_ADAPTER_DESC AdapterDesc;
 				VERIFYD3D11RESULT(TempAdapter->GetDesc(&AdapterDesc));
 				uint32 OutputCount = CountAdapterOutputs(TempAdapter);
 
@@ -970,7 +979,7 @@ void FD3D11DynamicRHIModule::FindAdapter()
 				// PerfHUD is for performance profiling
 				const bool bIsPerfHUD = !FCString::Stricmp(AdapterDesc.Description,TEXT("NVIDIA PerfHUD"));
 
-				FD3D11Adapter CurrentAdapter(AdapterIndex, ActualFeatureLevel, bIsSoftware, bIsIntegrated);
+				FD3D11Adapter CurrentAdapter(TempAdapter, ActualFeatureLevel, bIsSoftware, bIsIntegrated);
 
 				// Add special check to support HMDs, which do not have associated outputs.
 				// To reject the software emulation, unless the cvar wants it.
@@ -1038,8 +1047,8 @@ void FD3D11DynamicRHIModule::FindAdapter()
 
 	if(ChosenAdapter.IsValid())
 	{
-		ChosenDescription = AdapterDescription[ChosenAdapter.AdapterIndex];
-		UE_LOG(LogD3D11RHI, Log, TEXT("Chosen D3D11 Adapter: %u"), ChosenAdapter.AdapterIndex);
+		UE_LOG(LogD3D11RHI, Log, TEXT("Chosen D3D11 Adapter:"));
+		LogDXGIAdapterDesc(ChosenAdapter.DXGIAdapterDesc);
 	}
 	else
 	{
@@ -1083,7 +1092,7 @@ FDynamicRHI* FD3D11DynamicRHIModule::CreateRHI(ERHIFeatureLevel::Type RequestedF
 	GMaxRHIShaderPlatform = GShaderPlatformForFeatureLevel[GMaxRHIFeatureLevel];
 	check(GMaxRHIShaderPlatform != SP_NumPlatforms);
 
-	GD3D11RHI = new FD3D11DynamicRHI(DXGIFactory1,ChosenAdapter.MaxSupportedFeatureLevel,ChosenAdapter.AdapterIndex,ChosenDescription,ChosenAdapter.bSoftwareAdapter);
+	GD3D11RHI = new FD3D11DynamicRHI(DXGIFactory1, ChosenAdapter.MaxSupportedFeatureLevel,ChosenAdapter);
 	FDynamicRHI* FinalRHI = GD3D11RHI;
 
 #if ENABLE_RHI_VALIDATION
@@ -1667,9 +1676,6 @@ void FD3D11DynamicRHI::InitD3DDevice()
 		// Clear shadowed shader resources.
 		ClearState();
 
-		// Determine the adapter and device type to use.
-		TRefCountPtr<IDXGIAdapter> Adapter;
-		
 		// In Direct3D 11, if you are trying to create a hardware or a software device, set pAdapter != NULL which constrains the other inputs to be:
 		//		DriverType must be D3D_DRIVER_TYPE_UNKNOWN 
 		//		Software must be NULL. 
@@ -1694,116 +1700,100 @@ void FD3D11DynamicRHI::InitD3DDevice()
 
 		GTexturePoolSize = 0;
 
-		TRefCountPtr<IDXGIAdapter> EnumAdapter;
-
-		if(DXGIFactory1->EnumAdapters(ChosenAdapter,EnumAdapter.GetInitReference()) != DXGI_ERROR_NOT_FOUND)
-		{
-			if (EnumAdapter)// && EnumAdapter->CheckInterfaceSupport(__uuidof(ID3D11Device),NULL) == S_OK)
-			{
-				// we don't use AdapterDesc.Description as there is a bug with Optimus where it can report the wrong name
-				DXGI_ADAPTER_DESC AdapterDesc = ChosenDescription;
-				Adapter = EnumAdapter;
-
-				GRHIAdapterName = AdapterDesc.Description;
-				GRHIVendorId = AdapterDesc.VendorId;
-				GRHIDeviceId = AdapterDesc.DeviceId;
-				GRHIDeviceRevision = AdapterDesc.Revision;
+		GRHIAdapterName = Adapter.DXGIAdapterDesc.Description;
+		GRHIVendorId = Adapter.DXGIAdapterDesc.VendorId;
+		GRHIDeviceId = Adapter.DXGIAdapterDesc.DeviceId;
+		GRHIDeviceRevision = Adapter.DXGIAdapterDesc.Revision;
 
-				UE_LOG(LogD3D11RHI, Log, TEXT("    GPU DeviceId: 0x%x (for the marketing name, search the web for \"GPU Device Id\")"), 
-					AdapterDesc.DeviceId);
+		UE_LOG(LogD3D11RHI, Log, TEXT("    GPU DeviceId: 0x%x (for the marketing name, search the web for \"GPU Device Id\")"), 
+			Adapter.DXGIAdapterDesc.DeviceId);
 
-				// get driver version (todo: share with other RHIs)
-				{
-					FGPUDriverInfo GPUDriverInfo = FPlatformMisc::GetGPUDriverInfo(GRHIAdapterName);
+		// get driver version (todo: share with other RHIs)
+		{
+			FGPUDriverInfo GPUDriverInfo = FPlatformMisc::GetGPUDriverInfo(GRHIAdapterName);
 
-					GRHIAdapterUserDriverVersion = GPUDriverInfo.UserDriverVersion;
-					GRHIAdapterInternalDriverVersion = GPUDriverInfo.InternalDriverVersion;
-					GRHIAdapterDriverDate = GPUDriverInfo.DriverDate;
+			GRHIAdapterUserDriverVersion = GPUDriverInfo.UserDriverVersion;
+			GRHIAdapterInternalDriverVersion = GPUDriverInfo.InternalDriverVersion;
+			GRHIAdapterDriverDate = GPUDriverInfo.DriverDate;
 
-					UE_LOG(LogD3D11RHI, Log, TEXT("    Adapter Name: %s"), *GRHIAdapterName);
-					UE_LOG(LogD3D11RHI, Log, TEXT("  Driver Version: %s (internal:%s, unified:%s)"), *GRHIAdapterUserDriverVersion, *GRHIAdapterInternalDriverVersion, *GPUDriverInfo.GetUnifiedDriverVersion());
-					UE_LOG(LogD3D11RHI, Log, TEXT("     Driver Date: %s"), *GRHIAdapterDriverDate);
-				}
+			UE_LOG(LogD3D11RHI, Log, TEXT("    Adapter Name: %s"), *GRHIAdapterName);
+			UE_LOG(LogD3D11RHI, Log, TEXT("  Driver Version: %s (internal:%s, unified:%s)"), *GRHIAdapterUserDriverVersion, *GRHIAdapterInternalDriverVersion, *GPUDriverInfo.GetUnifiedDriverVersion());
+			UE_LOG(LogD3D11RHI, Log, TEXT("     Driver Date: %s"), *GRHIAdapterDriverDate);
+		}
 
-				// Issue: 32bit windows doesn't report 64bit value, we take what we get.
-				FD3D11GlobalStats::GDedicatedVideoMemory = int64(AdapterDesc.DedicatedVideoMemory);
-				FD3D11GlobalStats::GDedicatedSystemMemory = int64(AdapterDesc.DedicatedSystemMemory);
-				FD3D11GlobalStats::GSharedSystemMemory = int64(AdapterDesc.SharedSystemMemory);
+		// Issue: 32bit windows doesn't report 64bit value, we take what we get.
+		FD3D11GlobalStats::GDedicatedVideoMemory = int64(Adapter.DXGIAdapterDesc.DedicatedVideoMemory);
+		FD3D11GlobalStats::GDedicatedSystemMemory = int64(Adapter.DXGIAdapterDesc.DedicatedSystemMemory);
+		FD3D11GlobalStats::GSharedSystemMemory = int64(Adapter.DXGIAdapterDesc.SharedSystemMemory);
 
-				// Total amount of system memory, clamped to 8 GB
-				int64 TotalPhysicalMemory = FMath::Min(int64(FPlatformMemory::GetConstants().TotalPhysicalGB), 8ll) * (1024ll * 1024ll * 1024ll);
+		// Total amount of system memory, clamped to 8 GB
+		int64 TotalPhysicalMemory = FMath::Min(int64(FPlatformMemory::GetConstants().TotalPhysicalGB), 8ll) * (1024ll * 1024ll * 1024ll);
 
-				// Consider 50% of the shared memory but max 25% of total system memory.
-				int64 ConsideredSharedSystemMemory = FMath::Min( FD3D11GlobalStats::GSharedSystemMemory / 2ll, TotalPhysicalMemory / 4ll );
+		// Consider 50% of the shared memory but max 25% of total system memory.
+		int64 ConsideredSharedSystemMemory = FMath::Min( FD3D11GlobalStats::GSharedSystemMemory / 2ll, TotalPhysicalMemory / 4ll );
 
-				TRefCountPtr<IDXGIAdapter3> DxgiAdapter3;
-				DXGI_QUERY_VIDEO_MEMORY_INFO LocalVideoMemoryInfo;
-				FD3D11GlobalStats::GTotalGraphicsMemory = 0;
-				if ( IsRHIDeviceIntel() )
-				{
-					// It's all system memory.
-					FD3D11GlobalStats::GTotalGraphicsMemory = FD3D11GlobalStats::GDedicatedVideoMemory;
-					FD3D11GlobalStats::GTotalGraphicsMemory += FD3D11GlobalStats::GDedicatedSystemMemory;
-					FD3D11GlobalStats::GTotalGraphicsMemory += ConsideredSharedSystemMemory;
-				}
-				else if (IsRHIDeviceAMD() && SUCCEEDED(Adapter->QueryInterface(_uuidof(IDXGIAdapter3), (void**)DxgiAdapter3.GetInitReference())) &&
-					DxgiAdapter3.IsValid() && SUCCEEDED(DxgiAdapter3->QueryVideoMemoryInfo(0, DXGI_MEMORY_SEGMENT_GROUP_LOCAL, &LocalVideoMemoryInfo)))
-				{
-					// use the entire budget for D3D11, in keeping with setting GTotalGraphicsMemory to all of AdapterDesc.DedicatedVideoMemory
-					// in the other method directly below
-					FD3D11GlobalStats::GTotalGraphicsMemory = LocalVideoMemoryInfo.Budget;
-				}
-				else if ( FD3D11GlobalStats::GDedicatedVideoMemory >= 200*1024*1024 )
-				{
-					// Use dedicated video memory, if it's more than 200 MB
-					FD3D11GlobalStats::GTotalGraphicsMemory = FD3D11GlobalStats::GDedicatedVideoMemory;
-				}
-				else if ( FD3D11GlobalStats::GDedicatedSystemMemory >= 200*1024*1024 )
-				{
-					// Use dedicated system memory, if it's more than 200 MB
-					FD3D11GlobalStats::GTotalGraphicsMemory = FD3D11GlobalStats::GDedicatedSystemMemory;
-				}
-				else if ( FD3D11GlobalStats::GSharedSystemMemory >= 400*1024*1024 )
-				{
-					// Use some shared system memory, if it's more than 400 MB
-					FD3D11GlobalStats::GTotalGraphicsMemory = ConsideredSharedSystemMemory;
-				}
-				else
-				{
-					// Otherwise consider 25% of total system memory for graphics.
-					FD3D11GlobalStats::GTotalGraphicsMemory = TotalPhysicalMemory / 4ll;
-				}
+		TRefCountPtr<IDXGIAdapter3> DxgiAdapter3;
+		DXGI_QUERY_VIDEO_MEMORY_INFO LocalVideoMemoryInfo;
+		FD3D11GlobalStats::GTotalGraphicsMemory = 0;
+		if ( IsRHIDeviceIntel() )
+		{
+			// It's all system memory.
+			FD3D11GlobalStats::GTotalGraphicsMemory = FD3D11GlobalStats::GDedicatedVideoMemory;
+			FD3D11GlobalStats::GTotalGraphicsMemory += FD3D11GlobalStats::GDedicatedSystemMemory;
+			FD3D11GlobalStats::GTotalGraphicsMemory += ConsideredSharedSystemMemory;
+		}
+		else if (IsRHIDeviceAMD() && SUCCEEDED(Adapter.DXGIAdapter->QueryInterface(_uuidof(IDXGIAdapter3), (void**)DxgiAdapter3.GetInitReference())) &&
+			DxgiAdapter3.IsValid() && SUCCEEDED(DxgiAdapter3->QueryVideoMemoryInfo(0, DXGI_MEMORY_SEGMENT_GROUP_LOCAL, &LocalVideoMemoryInfo)))
+		{
+			// use the entire budget for D3D11, in keeping with setting GTotalGraphicsMemory to all of AdapterDesc.DedicatedVideoMemory
+			// in the other method directly below
+			FD3D11GlobalStats::GTotalGraphicsMemory = LocalVideoMemoryInfo.Budget;
+		}
+		else if ( FD3D11GlobalStats::GDedicatedVideoMemory >= 200*1024*1024 )
+		{
+			// Use dedicated video memory, if it's more than 200 MB
+			FD3D11GlobalStats::GTotalGraphicsMemory = FD3D11GlobalStats::GDedicatedVideoMemory;
+		}
+		else if ( FD3D11GlobalStats::GDedicatedSystemMemory >= 200*1024*1024 )
+		{
+			// Use dedicated system memory, if it's more than 200 MB
+			FD3D11GlobalStats::GTotalGraphicsMemory = FD3D11GlobalStats::GDedicatedSystemMemory;
+		}
+		else if ( FD3D11GlobalStats::GSharedSystemMemory >= 400*1024*1024 )
+		{
+			// Use some shared system memory, if it's more than 400 MB
+			FD3D11GlobalStats::GTotalGraphicsMemory = ConsideredSharedSystemMemory;
+		}
+		else
+		{
+			// Otherwise consider 25% of total system memory for graphics.
+			FD3D11GlobalStats::GTotalGraphicsMemory = TotalPhysicalMemory / 4ll;
+		}
 
-				if ( sizeof(SIZE_T) < 8 )
-				{
-					// Clamp to 1 GB if we're less than 64-bit
-					FD3D11GlobalStats::GTotalGraphicsMemory = FMath::Min( FD3D11GlobalStats::GTotalGraphicsMemory, 1024ll * 1024ll * 1024ll );
-				}
+		if ( sizeof(SIZE_T) < 8 )
+		{
+			// Clamp to 1 GB if we're less than 64-bit
+			FD3D11GlobalStats::GTotalGraphicsMemory = FMath::Min( FD3D11GlobalStats::GTotalGraphicsMemory, 1024ll * 1024ll * 1024ll );
+		}
 
-				if ( GPoolSizeVRAMPercentage > 0 )
-				{
-					float PoolSize = float(GPoolSizeVRAMPercentage) * 0.01f * float(FD3D11GlobalStats::GTotalGraphicsMemory);
+		if ( GPoolSizeVRAMPercentage > 0 )
+		{
+			float PoolSize = float(GPoolSizeVRAMPercentage) * 0.01f * float(FD3D11GlobalStats::GTotalGraphicsMemory);
 
-					// Truncate GTexturePoolSize to MB (but still counted in bytes)
-					GTexturePoolSize = int64(FGenericPlatformMath::TruncToFloat(PoolSize / 1024.0f / 1024.0f)) * 1024 * 1024;
+			// Truncate GTexturePoolSize to MB (but still counted in bytes)
+			GTexturePoolSize = int64(FGenericPlatformMath::TruncToFloat(PoolSize / 1024.0f / 1024.0f)) * 1024 * 1024;
 
-					UE_LOG(LogRHI,Log,TEXT("Texture pool is %llu MB (%d%% of %llu MB)"),
-						GTexturePoolSize / 1024 / 1024,
-						GPoolSizeVRAMPercentage,
-						FD3D11GlobalStats::GTotalGraphicsMemory / 1024 / 1024);
-				}
+			UE_LOG(LogRHI,Log,TEXT("Texture pool is %llu MB (%d%% of %llu MB)"),
+				GTexturePoolSize / 1024 / 1024,
+				GPoolSizeVRAMPercentage,
+				FD3D11GlobalStats::GTotalGraphicsMemory / 1024 / 1024);
+		}
 
-				const bool bIsPerfHUD = !FCString::Stricmp(AdapterDesc.Description,TEXT("NVIDIA PerfHUD"));
+		const bool bIsPerfHUD = !FCString::Stricmp(Adapter.DXGIAdapterDesc.Description,TEXT("NVIDIA PerfHUD"));
 
-				if(bIsPerfHUD)
-				{
-					DriverType =  D3D_DRIVER_TYPE_REFERENCE;
-				}
-			}
-		}
-		else
+		if(bIsPerfHUD)
 		{
-			check(!"Internal error, EnumAdapters() failed but before it worked")
+			DriverType =  D3D_DRIVER_TYPE_REFERENCE;
 		}
 
 #ifdef AMD_AGS_API
@@ -1820,8 +1810,8 @@ void FD3D11DynamicRHI::InitD3DDevice()
 				for (int32 DeviceIndex = 0; DeviceIndex < AmdInfo.AmdGpuInfo.numDevices; DeviceIndex++)
 				{
 					const AGSDeviceInfo& DeviceInfo = AmdInfo.AmdGpuInfo.devices[DeviceIndex];
-					GRHIDeviceIsAMDPreGCNArchitecture |= (ChosenDescription.VendorId == DeviceInfo.vendorId) && (ChosenDescription.DeviceId == DeviceInfo.deviceId) && (DeviceInfo.asicFamily == AGSDeviceInfo::AsicFamily_PreGCN);
-					bFoundMatchingDevice |= (ChosenDescription.VendorId == DeviceInfo.vendorId) && (ChosenDescription.DeviceId == DeviceInfo.deviceId);
+					GRHIDeviceIsAMDPreGCNArchitecture |= (Adapter.DXGIAdapterDesc.VendorId == DeviceInfo.vendorId) && (Adapter.DXGIAdapterDesc.DeviceId == DeviceInfo.deviceId) && (DeviceInfo.asicFamily == AGSDeviceInfo::AsicFamily_PreGCN);
+					bFoundMatchingDevice |= (Adapter.DXGIAdapterDesc.VendorId == DeviceInfo.vendorId) && (Adapter.DXGIAdapterDesc.DeviceId == DeviceInfo.deviceId);
 				}
 				check(bFoundMatchingDevice);
 
@@ -1857,9 +1847,12 @@ void FD3D11DynamicRHI::InitD3DDevice()
 		{
 			uint32 AmdSupportedExtensionFlags = 0;
 
+			UE_LOG(LogD3D11RHI, Log, TEXT("Creating D3DDevice with AMD AGS, using adapter:"));
+			LogDXGIAdapterDesc(Adapter.DXGIAdapterDesc);
+
 			AGSDX11DeviceCreationParams DeviceCreationParams = 
 			{
-				Adapter,
+				Adapter.DXGIAdapter,
 				DriverType,
 				NULL,
 				DeviceFlags,
@@ -1942,9 +1935,12 @@ void FD3D11DynamicRHI::InitD3DDevice()
 
 		if (!bDeviceCreated)
 		{
+			UE_LOG(LogD3D11RHI, Log, TEXT("Creating D3DDevice using adapter:"));
+			LogDXGIAdapterDesc(Adapter.DXGIAdapterDesc);
+
 			// Creating the Direct3D device.
 			VERIFYD3D11RESULT(D3D11CreateDevice(
-				Adapter,
+				Adapter.DXGIAdapter,
 				DriverType,
 				NULL,
 				DeviceFlags,
@@ -2266,26 +2262,12 @@ bool FD3D11DynamicRHI::RHIGetAvailableResolutions(FScreenResolutionArray& Resolu
 	}
 
 	HRESULT HResult = S_OK;
-	TRefCountPtr<IDXGIAdapter> Adapter;
-	HResult = DXGIFactory1->EnumAdapters(ChosenAdapter,Adapter.GetInitReference());
-
-	if( DXGI_ERROR_NOT_FOUND == HResult )
-		return false;
-	if( FAILED(HResult) )
-		return false;
-
-	// get the description of the adapter
-	DXGI_ADAPTER_DESC AdapterDesc;
-	if( FAILED(Adapter->GetDesc(&AdapterDesc)) )
-	{
-		return false;
-	}
 
 	int32 CurrentOutput = 0;
 	do 
 	{
 		TRefCountPtr<IDXGIOutput> Output;
-		HResult = Adapter->EnumOutputs(CurrentOutput,Output.GetInitReference());
+		HResult = Adapter.DXGIAdapter->EnumOutputs(CurrentOutput,Output.GetInitReference());
 		if(DXGI_ERROR_NOT_FOUND == HResult)
 			break;
 		if(FAILED(HResult))
