U5 - remove unecessary abstraction and layering violation that contrained the EventQuery to be released on the Renderthread, which constrains it more than neceassry and triggered an assert when we moved destruction to the RHITHread.

diff --git a/Engine/Source/Runtime/Windows/D3D11RHI/Private/D3D11Query.cpp b/Engine/Source/Runtime/Windows/D3D11RHI/Private/D3D11Query.cpp
--- a/Engine/Source/Runtime/Windows/D3D11RHI/Private/D3D11Query.cpp
+++ b/Engine/Source/Runtime/Windows/D3D11RHI/Private/D3D11Query.cpp
@@ -385,19 +385,13 @@ void FD3D11EventQuery::WaitForCompletion()
 	{};
 }
 
-void FD3D11EventQuery::InitDynamicRHI()
+FD3D11EventQuery::FD3D11EventQuery(class FD3D11DynamicRHI* InD3DRHI):
+	D3DRHI(InD3DRHI)
 {
 	D3D11_QUERY_DESC QueryDesc;
 	QueryDesc.Query = D3D11_QUERY_EVENT;
 	QueryDesc.MiscFlags = 0;
 	VERIFYD3D11RESULT_EX(D3DRHI->GetDevice()->CreateQuery(&QueryDesc,Query.GetInitReference()), D3DRHI->GetDevice());
-
-	// Initialize the query by issuing an initial event.
-	IssueEvent();
-}
-void FD3D11EventQuery::ReleaseDynamicRHI()
-{
-	Query = NULL;
 }
 
 /*=============================================================================
diff --git a/Engine/Source/Runtime/Windows/D3D11RHI/Private/D3D11Viewport.cpp b/Engine/Source/Runtime/Windows/D3D11RHI/Private/D3D11Viewport.cpp
--- a/Engine/Source/Runtime/Windows/D3D11RHI/Private/D3D11Viewport.cpp
+++ b/Engine/Source/Runtime/Windows/D3D11RHI/Private/D3D11Viewport.cpp
@@ -223,8 +223,6 @@ FD3D11Viewport::~FD3D11Viewport()
 	}
 #endif
 
-	FrameSyncEvent.ReleaseResource();
-
 	D3DRHI->Viewports.Remove(this);
 }
 
diff --git a/Engine/Source/Runtime/Windows/D3D11RHI/Private/Windows/WindowsD3D11Viewport.cpp b/Engine/Source/Runtime/Windows/D3D11RHI/Private/Windows/WindowsD3D11Viewport.cpp
--- a/Engine/Source/Runtime/Windows/D3D11RHI/Private/Windows/WindowsD3D11Viewport.cpp
+++ b/Engine/Source/Runtime/Windows/D3D11RHI/Private/Windows/WindowsD3D11Viewport.cpp
@@ -256,7 +256,12 @@ FD3D11Viewport::FD3D11Viewport(FD3D11DynamicRHI* InD3DRHI,HWND InWindowHandle,ui
 	// @todo: For Slate viewports, it doesn't make sense to post WM_PAINT messages (we swallow those.)
 	::PostMessageW( WindowHandle, WM_PAINT, 0, 0 );
 
-	BeginInitResource(&FrameSyncEvent);
+	ENQUEUE_RENDER_COMMAND(FD3D11Viewport)(
+		[this](FRHICommandListImmediate& RHICmdList)
+	{
+		// Initialize the query by issuing an initial event.
+		FrameSyncEvent.IssueEvent();
+	});	
 }
 
 // When a window has moved or resized we need to check whether it is on a HDR monitor or not. Set the correct color space of the monitor
diff --git a/Engine/Source/Runtime/Windows/D3D11RHI/Public/D3D11Viewport.h b/Engine/Source/Runtime/Windows/D3D11RHI/Public/D3D11Viewport.h
--- a/Engine/Source/Runtime/Windows/D3D11RHI/Public/D3D11Viewport.h
+++ b/Engine/Source/Runtime/Windows/D3D11RHI/Public/D3D11Viewport.h
@@ -17,15 +17,12 @@
 #endif
 
 /** A D3D event query resource. */
-class D3D11RHI_API FD3D11EventQuery : public FRenderResource
+class D3D11RHI_API FD3D11EventQuery
 {
 public:
 
 	/** Initialization constructor. */
-	FD3D11EventQuery(class FD3D11DynamicRHI* InD3DRHI):
-		D3DRHI(InD3DRHI)
-	{
-	}
+	FD3D11EventQuery(class FD3D11DynamicRHI* InD3DRHI);
 
 	/** Issues an event for the query to poll. */
 	void IssueEvent();
@@ -33,10 +30,6 @@ public:
 	/** Waits for the event query to finish. */
 	void WaitForCompletion();
 
-	// FRenderResource interface.
-	virtual void InitDynamicRHI() override;
-	virtual void ReleaseDynamicRHI() override;
-
 private:
 	FD3D11DynamicRHI* D3DRHI;
 	TRefCountPtr<ID3D11Query> Query;
