Fix EnumAdapterByGpuPreference on DX11 PC

diff --git a/Engine/Source/Runtime/Windows/D3D11RHI/Private/Windows/WindowsD3D11Device.cpp b/Engine/Source/Runtime/Windows/D3D11RHI/Private/Windows/WindowsD3D11Device.cpp
--- a/Engine/Source/Runtime/Windows/D3D11RHI/Private/Windows/WindowsD3D11Device.cpp
+++ b/Engine/Source/Runtime/Windows/D3D11RHI/Private/Windows/WindowsD3D11Device.cpp
@@ -919,7 +919,7 @@ void FD3D11DynamicRHIModule::FindAdapter()
 		}
 		else
 		{
-			return DXGIFactory6->EnumAdapterByGpuPreference(AdapterIndex, DXGI_GPU_PREFERENCE_HIGH_PERFORMANCE, __uuidof(IDXGIAdapter), (void**)Adapter);
+			return DXGIFactory6->EnumAdapterByGpuPreference(AdapterIndex, GpuPreference, __uuidof(IDXGIAdapter), (void**)Adapter);
 		}
 	};
 
