UE5.0_Main:  Fix for incorrect strides for depth stencil formats on D3D11 and D3D12, in a utility function which reads render targets to a CPU buffer.

DXGI_FORMAT_R32G8X24_TYPELESS and related formats are 8 byte stride, not 5 byte, according to D3D documentation and Microsoft sample code that processes these formats.  The original code would never have worked, and is likely only used for debug purposes for perf reasons, so it seems safe to fix.

#rnx

diff --git a/Engine/Source/Runtime/D3D12RHI/Private/D3D12RenderTarget.cpp b/Engine/Source/Runtime/D3D12RHI/Private/D3D12RenderTarget.cpp
--- a/Engine/Source/Runtime/D3D12RHI/Private/D3D12RenderTarget.cpp
+++ b/Engine/Source/Runtime/D3D12RHI/Private/D3D12RenderTarget.cpp
@@ -561,7 +561,7 @@ static uint32 ComputeBytesPerPixel(DXGI_FORMAT Format)
 	case DXGI_FORMAT_R32G8X24_TYPELESS:
 	case DXGI_FORMAT_D32_FLOAT_S8X24_UINT:
 	case DXGI_FORMAT_R32_FLOAT_X8X24_TYPELESS:
-		BytesPerPixel = 5;
+		BytesPerPixel = 8;
 		break;
 	case DXGI_FORMAT_R8_TYPELESS:
 	case DXGI_FORMAT_R8_UNORM:
diff --git a/Engine/Source/Runtime/Windows/D3D11RHI/Private/D3D11RenderTarget.cpp b/Engine/Source/Runtime/Windows/D3D11RHI/Private/D3D11RenderTarget.cpp
--- a/Engine/Source/Runtime/Windows/D3D11RHI/Private/D3D11RenderTarget.cpp
+++ b/Engine/Source/Runtime/Windows/D3D11RHI/Private/D3D11RenderTarget.cpp
@@ -403,7 +403,7 @@ static uint32 ComputeBytesPerPixel(DXGI_FORMAT Format)
 		case DXGI_FORMAT_R32G8X24_TYPELESS:
 		case DXGI_FORMAT_D32_FLOAT_S8X24_UINT:
 		case DXGI_FORMAT_R32_FLOAT_X8X24_TYPELESS:
-			BytesPerPixel = 5;
+			BytesPerPixel = 8;
 			break;
 		case DXGI_FORMAT_R8_TYPELESS:
 		case DXGI_FORMAT_R8_UNORM:
