D3D11 cleanup. Removing long dead code paths.

diff --git a/Engine/Source/Runtime/Windows/D3D11RHI/Private/D3D11Commands.cpp b/Engine/Source/Runtime/Windows/D3D11RHI/Private/D3D11Commands.cpp
--- a/Engine/Source/Runtime/Windows/D3D11RHI/Private/D3D11Commands.cpp
+++ b/Engine/Source/Runtime/Windows/D3D11RHI/Private/D3D11Commands.cpp
@@ -353,19 +353,6 @@ void FD3D11DynamicRHI::RHISetBoundShaderState(FRHIBoundShaderState* BoundShaderS
 	}
 }
 
-template <EShaderFrequency ShaderFrequency>
-FORCEINLINE void FD3D11DynamicRHI::SetShaderTexture(FD3D11TextureBase* NewTexture, ID3D11ShaderResourceView* ShaderResourceView, uint32 TextureIndex, FRHITexture* NewTextureRHI)
-{
-	if ((NewTexture == nullptr) || (NewTexture->GetRenderTargetView(0, 0) != NULL) || (NewTexture->HasDepthStencilView()))
-	{
-		SetShaderResourceView<ShaderFrequency>(NewTexture, ShaderResourceView, TextureIndex, NewTextureRHI ? NewTextureRHI->GetName() : NAME_None, FD3D11StateCache::SRV_Dynamic);
-	}
-	else
-	{
-		SetShaderResourceView<ShaderFrequency>(NewTexture, ShaderResourceView, TextureIndex, NewTextureRHI->GetName(), FD3D11StateCache::SRV_Static);
-	}
-}
-
 void FD3D11DynamicRHI::RHISetShaderTexture(FRHIGraphicsShader* ShaderRHI,uint32 TextureIndex, FRHITexture* NewTextureRHI)
 {
 	FD3D11TextureBase* NewTexture = GetD3D11TextureFromRHITexture(NewTextureRHI);
@@ -377,21 +364,21 @@ void FD3D11DynamicRHI::RHISetShaderTexture(FRHIGraphicsShader* ShaderRHI,uint32
 	{
 		FD3D11VertexShader* VertexShader = static_cast<FD3D11VertexShader*>(ShaderRHI);
 		VALIDATE_BOUND_SHADER(VertexShader);
-		SetShaderTexture<SF_Vertex>(NewTexture, ShaderResourceView, TextureIndex, NewTextureRHI);
+		SetShaderResourceView<SF_Vertex>(NewTexture, ShaderResourceView, TextureIndex);
 	}
 	break;
 	case SF_Geometry:
 	{
 		FD3D11GeometryShader* GeometryShader = static_cast<FD3D11GeometryShader*>(ShaderRHI);
 		VALIDATE_BOUND_SHADER(GeometryShader);
-		SetShaderTexture<SF_Geometry>(NewTexture, ShaderResourceView, TextureIndex, NewTextureRHI);
+		SetShaderResourceView<SF_Geometry>(NewTexture, ShaderResourceView, TextureIndex);
 	}
 	break;
 	case SF_Pixel:
 	{
 		FD3D11PixelShader* PixelShader = static_cast<FD3D11PixelShader*>(ShaderRHI);
 		VALIDATE_BOUND_SHADER(PixelShader);
-		SetShaderTexture<SF_Pixel>(NewTexture, ShaderResourceView, TextureIndex, NewTextureRHI);
+		SetShaderResourceView<SF_Pixel>(NewTexture, ShaderResourceView, TextureIndex);
 	}
 	break;
 	default:
@@ -405,7 +392,7 @@ void FD3D11DynamicRHI::RHISetShaderTexture(FRHIComputeShader* ComputeShaderRHI,u
 
 	FD3D11TextureBase* NewTexture = GetD3D11TextureFromRHITexture(NewTextureRHI);
 	ID3D11ShaderResourceView* ShaderResourceView = NewTexture ? NewTexture->GetShaderResourceView() : nullptr;
-	SetShaderTexture<SF_Compute>(NewTexture, ShaderResourceView, TextureIndex, NewTextureRHI);
+	SetShaderResourceView<SF_Compute>(NewTexture, ShaderResourceView, TextureIndex);
 }
 
 
@@ -480,21 +467,21 @@ void FD3D11DynamicRHI::RHISetShaderResourceViewParameter(FRHIGraphicsShader* Sha
 	{
 		FD3D11VertexShader* VertexShader = static_cast<FD3D11VertexShader*>(ShaderRHI);
 		VALIDATE_BOUND_SHADER(VertexShader);
-		SetShaderResourceView<SF_Vertex>(Resource, D3D11SRV, TextureIndex, NAME_None);
+		SetShaderResourceView<SF_Vertex>(Resource, D3D11SRV, TextureIndex);
 	}
 	break;
 	case SF_Geometry:
 	{
 		FD3D11GeometryShader* GeometryShader = static_cast<FD3D11GeometryShader*>(ShaderRHI);
 		VALIDATE_BOUND_SHADER(GeometryShader);
-		SetShaderResourceView<SF_Geometry>(Resource, D3D11SRV, TextureIndex, NAME_None);
+		SetShaderResourceView<SF_Geometry>(Resource, D3D11SRV, TextureIndex);
 	}
 	break;
 	case SF_Pixel:
 	{
 		FD3D11PixelShader* PixelShader = static_cast<FD3D11PixelShader*>(ShaderRHI);
 		VALIDATE_BOUND_SHADER(PixelShader);
-		SetShaderResourceView<SF_Pixel>(Resource, D3D11SRV, TextureIndex, NAME_None);
+		SetShaderResourceView<SF_Pixel>(Resource, D3D11SRV, TextureIndex);
 	}
 	break;
 	default:
@@ -517,7 +504,7 @@ void FD3D11DynamicRHI::RHISetShaderResourceViewParameter(FRHIComputeShader* Comp
 		D3D11SRV = SRV->View;
 	}
 
-	SetShaderResourceView<SF_Compute>(Resource, D3D11SRV, TextureIndex, NAME_None);
+	SetShaderResourceView<SF_Compute>(Resource, D3D11SRV, TextureIndex);
 }
 
 void FD3D11DynamicRHI::RHISetShaderSampler(FRHIGraphicsShader* ShaderRHI,uint32 SamplerIndex, FRHISamplerState* NewStateRHI)
@@ -1170,11 +1157,11 @@ void FD3D11DynamicRHI::CommitComputeShaderConstants()
 }
 
 template <EShaderFrequency Frequency>
-FORCEINLINE void SetResource(FD3D11DynamicRHI* RESTRICT D3D11RHI, FD3D11StateCache* RESTRICT StateCache, uint32 BindIndex, FD3D11BaseShaderResource* RESTRICT ShaderResource, ID3D11ShaderResourceView* RESTRICT SRV, FName ResourceName = FName())
+FORCEINLINE void SetResource(FD3D11DynamicRHI* RESTRICT D3D11RHI, FD3D11StateCache* RESTRICT StateCache, uint32 BindIndex, FD3D11BaseShaderResource* RESTRICT ShaderResource, ID3D11ShaderResourceView* RESTRICT SRV)
 {
 	// We set the resource through the RHI to track state for the purposes of unbinding SRVs when a UAV or RTV is bound.
 	// todo: need to support SRV_Static for faster calls when possible
-	D3D11RHI->SetShaderResourceView<Frequency>(ShaderResource, SRV, BindIndex, ResourceName, FD3D11StateCache::SRV_Unknown);
+	D3D11RHI->SetShaderResourceView<Frequency>(ShaderResource, SRV, BindIndex);
 }
 
 template <EShaderFrequency Frequency>
@@ -1229,7 +1216,7 @@ inline int32 SetShaderResourcesFromBuffer_Surface(FD3D11DynamicRHI* RESTRICT D3D
 #endif
 
 			// todo: could coalesce adjacent bound resources.
-			SetResource<ShaderFrequency>(D3D11RHI, StateCache, BindIndex, ShaderResource, D3D11Resource, TextureRHI->GetName());
+			SetResource<ShaderFrequency>(D3D11RHI, StateCache, BindIndex, ShaderResource, D3D11Resource);
 			NumSetCalls++;
 			ResourceInfo = *ResourceInfos++;
 		} while (FRHIResourceTableEntry::GetUniformBufferIndex(ResourceInfo) == BufferIndex);
diff --git a/Engine/Source/Runtime/Windows/D3D11RHI/Private/D3D11RHI.cpp b/Engine/Source/Runtime/Windows/D3D11RHI/Private/D3D11RHI.cpp
--- a/Engine/Source/Runtime/Windows/D3D11RHI/Private/D3D11RHI.cpp
+++ b/Engine/Source/Runtime/Windows/D3D11RHI/Private/D3D11RHI.cpp
@@ -137,7 +137,7 @@ void GetMipAndSliceInfoFromSRV(ID3D11ShaderResourceView* SRV, int32& MipLevel, i
 }
 
 template <EShaderFrequency ShaderFrequency>
-void FD3D11DynamicRHI::InternalSetShaderResourceView(FD3D11BaseShaderResource* Resource, ID3D11ShaderResourceView* SRV, int32 ResourceIndex, FName SRVName, FD3D11StateCache::ESRV_Type SrvType)
+void FD3D11DynamicRHI::InternalSetShaderResourceView(FD3D11BaseShaderResource* Resource, ID3D11ShaderResourceView* SRV, int32 ResourceIndex)
 {
 	// Check either both are set, or both are null.
 	check((Resource && SRV) || (!Resource && !SRV));
@@ -178,13 +178,13 @@ void FD3D11DynamicRHI::InternalSetShaderResourceView(FD3D11BaseShaderResource* R
 	}
 
 	// Set the SRV we have been given (or null).
-	StateCache.SetShaderResourceView<ShaderFrequency>(SRV, ResourceIndex, SrvType);
+	StateCache.SetShaderResourceView<ShaderFrequency>(SRV, ResourceIndex);
 }
 
-template void FD3D11DynamicRHI::InternalSetShaderResourceView<SF_Vertex>(FD3D11BaseShaderResource* Resource, ID3D11ShaderResourceView* SRV, int32 ResourceIndex, FName SRVName, FD3D11StateCache::ESRV_Type SrvType);
-template void FD3D11DynamicRHI::InternalSetShaderResourceView<SF_Pixel>(FD3D11BaseShaderResource* Resource, ID3D11ShaderResourceView* SRV, int32 ResourceIndex, FName SRVName, FD3D11StateCache::ESRV_Type SrvType);
-template void FD3D11DynamicRHI::InternalSetShaderResourceView<SF_Geometry>(FD3D11BaseShaderResource* Resource, ID3D11ShaderResourceView* SRV, int32 ResourceIndex, FName SRVName, FD3D11StateCache::ESRV_Type SrvType);
-template void FD3D11DynamicRHI::InternalSetShaderResourceView<SF_Compute>(FD3D11BaseShaderResource* Resource, ID3D11ShaderResourceView* SRV, int32 ResourceIndex, FName SRVName, FD3D11StateCache::ESRV_Type SrvType);
+template void FD3D11DynamicRHI::InternalSetShaderResourceView<SF_Vertex>(FD3D11BaseShaderResource* Resource, ID3D11ShaderResourceView* SRV, int32 ResourceIndex);
+template void FD3D11DynamicRHI::InternalSetShaderResourceView<SF_Pixel>(FD3D11BaseShaderResource* Resource, ID3D11ShaderResourceView* SRV, int32 ResourceIndex);
+template void FD3D11DynamicRHI::InternalSetShaderResourceView<SF_Geometry>(FD3D11BaseShaderResource* Resource, ID3D11ShaderResourceView* SRV, int32 ResourceIndex);
+template void FD3D11DynamicRHI::InternalSetShaderResourceView<SF_Compute>(FD3D11BaseShaderResource* Resource, ID3D11ShaderResourceView* SRV, int32 ResourceIndex);
 
 void FD3D11DynamicRHI::TrackResourceBoundAsVB(FD3D11BaseShaderResource* Resource, int32 StreamIndex)
 {
@@ -228,7 +228,7 @@ void FD3D11DynamicRHI::ClearShaderResourceViews(FD3D11BaseShaderResource* Resour
 		if (CurrentResourcesBoundAsSRVs[ShaderFrequency][ResourceIndex] == Resource)
 		{
 			// Unset the SRV from the device context
-			InternalSetShaderResourceView<ShaderFrequency>(nullptr, nullptr, ResourceIndex, NAME_None);
+			SetShaderResourceView<ShaderFrequency>(nullptr, nullptr, ResourceIndex);
 		}
 	}
 }
@@ -271,7 +271,7 @@ void FD3D11DynamicRHI::ClearAllShaderResourcesForFrequency()
 		if (CurrentResourcesBoundAsSRVs[ShaderFrequency][ResourceIndex] != nullptr)
 		{
 			// Unset the SRV from the device context
-			InternalSetShaderResourceView<ShaderFrequency>(nullptr, nullptr, ResourceIndex, NAME_None);
+			SetShaderResourceView<ShaderFrequency>(nullptr, nullptr, ResourceIndex);
 		}
 	}
 	StateCache.ClearConstantBuffers<ShaderFrequency>();
diff --git a/Engine/Source/Runtime/Windows/D3D11RHI/Private/D3D11RHIPrivate.h b/Engine/Source/Runtime/Windows/D3D11RHI/Private/D3D11RHIPrivate.h
--- a/Engine/Source/Runtime/Windows/D3D11RHI/Private/D3D11RHIPrivate.h
+++ b/Engine/Source/Runtime/Windows/D3D11RHI/Private/D3D11RHIPrivate.h
@@ -740,9 +740,6 @@ private:
 
 	virtual void RHIClearMRTImpl(bool bClearColor, int32 NumClearColors, const FLinearColor* ColorArray, bool bClearDepth, float Depth, bool bClearStencil, uint32 Stencil);
 
-	template <EShaderFrequency ShaderFrequency>
-	inline void SetShaderTexture(FD3D11TextureBase* NewTexture, ID3D11ShaderResourceView* ShaderResourceView, uint32 TextureIndex, FRHITexture* NewTextureRHI);
-
 	template <EShaderFrequency ShaderFrequency>
 	void ClearShaderResourceViews(FD3D11BaseShaderResource* Resource);
 
@@ -750,7 +747,7 @@ private:
 	void ClearAllShaderResourcesForFrequency();
 
 	template <EShaderFrequency ShaderFrequency>
-	void InternalSetShaderResourceView(FD3D11BaseShaderResource* Resource, ID3D11ShaderResourceView* SRV, int32 ResourceIndex, FName SRVName, FD3D11StateCache::ESRV_Type SrvType = FD3D11StateCache::SRV_Unknown);
+	void InternalSetShaderResourceView(FD3D11BaseShaderResource* Resource, ID3D11ShaderResourceView* SRV, int32 ResourceIndex);
 
 	void TrackResourceBoundAsVB(FD3D11BaseShaderResource* Resource, int32 StreamIndex);
 	void TrackResourceBoundAsIB(FD3D11BaseShaderResource* Resource);
@@ -768,9 +765,9 @@ private:
 public:
 
 	template <EShaderFrequency ShaderFrequency>
-	void SetShaderResourceView(FD3D11BaseShaderResource* Resource, ID3D11ShaderResourceView* SRV, int32 ResourceIndex, FName SRVName, FD3D11StateCache::ESRV_Type SrvType = FD3D11StateCache::SRV_Unknown)
+	void SetShaderResourceView(FD3D11BaseShaderResource* Resource, ID3D11ShaderResourceView* SRV, int32 ResourceIndex)
 	{
-		InternalSetShaderResourceView<ShaderFrequency>(Resource, SRV, ResourceIndex, SRVName, SrvType);
+		InternalSetShaderResourceView<ShaderFrequency>(Resource, SRV, ResourceIndex);
 	}
 
 	void ClearState();
diff --git a/Engine/Source/Runtime/Windows/D3D11RHI/Private/D3D11RenderTarget.cpp b/Engine/Source/Runtime/Windows/D3D11RHI/Private/D3D11RenderTarget.cpp
--- a/Engine/Source/Runtime/Windows/D3D11RHI/Private/D3D11RenderTarget.cpp
+++ b/Engine/Source/Runtime/Windows/D3D11RHI/Private/D3D11RenderTarget.cpp
@@ -149,7 +149,7 @@ void FD3D11DynamicRHI::ResolveTextureUsingShader(
 
 	if (SourceTexture)
 	{
-		SetShaderResourceView<SF_Pixel>(SourceTexture, SourceTexture->GetShaderResourceView(), TextureIndex, SourceTexture->GetName());
+		SetShaderResourceView<SF_Pixel>(SourceTexture, SourceTexture->GetShaderResourceView(), TextureIndex);
 	}
 
 	RHICmdList.DrawPrimitive(0, 2, 1);
diff --git a/Engine/Source/Runtime/Windows/D3D11RHI/Private/D3D11StateCachePrivate.h b/Engine/Source/Runtime/Windows/D3D11RHI/Private/D3D11StateCachePrivate.h
--- a/Engine/Source/Runtime/Windows/D3D11RHI/Private/D3D11StateCachePrivate.h
+++ b/Engine/Source/Runtime/Windows/D3D11RHI/Private/D3D11StateCachePrivate.h
@@ -55,13 +55,6 @@ static const bool GD3D11SkipStateCaching = false;
 class FD3D11StateCacheBase
 {
 public:
-	enum ESRV_Type
-	{
-		SRV_Unknown,
-		SRV_Dynamic,
-		SRV_Static,
-	};
-
 	bool bDepthBoundsEnabled = false;
 	float DepthBoundsMin = 0.0f;
 	float DepthBoundsMax = 1.0f;
@@ -172,8 +165,7 @@ protected:
 		}
 	}
 
-	typedef void (*TSetIndexBufferAlternate)(FD3D11StateCacheBase* StateCache, ID3D11Buffer* IndexBuffer, DXGI_FORMAT Format, uint32 Offset);
-	D3D11_STATE_CACHE_INLINE void InternalSetIndexBuffer(ID3D11Buffer* IndexBuffer, DXGI_FORMAT Format, uint32 Offset, TSetIndexBufferAlternate AlternatePathFunction)
+	D3D11_STATE_CACHE_INLINE void InternalSetIndexBuffer(ID3D11Buffer* IndexBuffer, DXGI_FORMAT Format, uint32 Offset)
 	{
 #if D3D11_ALLOW_STATE_CACHE
 		D3D11_STATE_CACHE_VERIFY_PRE();
@@ -183,14 +175,7 @@ protected:
 			CurrentIndexBuffer = IndexBuffer;
 			CurrentIndexFormat = Format;
 			CurrentIndexOffset = Offset;
-			if (AlternatePathFunction != nullptr)
-			{
-				(*AlternatePathFunction)(this, IndexBuffer, Format, Offset);
-			}
-			else
-			{
-				Direct3DDeviceIMContext->IASetIndexBuffer(IndexBuffer, Format, Offset);
-			}
+			Direct3DDeviceIMContext->IASetIndexBuffer(IndexBuffer, Format, Offset);
 		}
 		D3D11_STATE_CACHE_VERIFY_POST();
 #else
@@ -198,9 +183,8 @@ protected:
 #endif
 	}
 
-	typedef void (*TSetSRVAlternate)(FD3D11StateCacheBase* StateCache, ID3D11ShaderResourceView* SRV, uint32 ResourceIndex, ESRV_Type SrvType);
 	template <EShaderFrequency ShaderFrequency>
-	D3D11_STATE_CACHE_INLINE void InternalSetShaderResourceView(ID3D11ShaderResourceView*& SRV, uint32 ResourceIndex, ESRV_Type SrvType, TSetSRVAlternate AlternatePathFunction)
+	D3D11_STATE_CACHE_INLINE void InternalSetShaderResourceView(ID3D11ShaderResourceView*& SRV, uint32 ResourceIndex)
 	{
 #if D3D11_ALLOW_STATE_CACHE
 		D3D11_STATE_CACHE_VERIFY_PRE();
@@ -216,14 +200,7 @@ protected:
 				CurrentShaderResourceViews[ShaderFrequency][ResourceIndex]->Release();
 			}
 			CurrentShaderResourceViews[ShaderFrequency][ResourceIndex] = SRV;
-			if (AlternatePathFunction != nullptr)
-			{
-				(*AlternatePathFunction)(this, SRV, ResourceIndex, SrvType);
-			}
-			else
-			{
-				InternalSetShaderResourceView<ShaderFrequency>(ResourceIndex, SRV);
-			}
+			InternalSetShaderResourceView<ShaderFrequency>(ResourceIndex, SRV);
 		}
 		D3D11_STATE_CACHE_VERIFY_POST();
 #else	// !D3D11_ALLOW_STATE_CACHE
@@ -231,8 +208,7 @@ protected:
 #endif
 	}
 
-	typedef void (*TSetStreamSourceAlternate)(FD3D11StateCacheBase* StateCache, ID3D11Buffer* VertexBuffer, uint32 StreamIndex, uint32 Stride, uint32 Offset);
-	D3D11_STATE_CACHE_INLINE void InternalSetStreamSource(ID3D11Buffer* VertexBuffer, uint32 StreamIndex, uint32 Stride, uint32 Offset, TSetStreamSourceAlternate AlternatePathFunction)
+	D3D11_STATE_CACHE_INLINE void InternalSetStreamSource(ID3D11Buffer* VertexBuffer, uint32 StreamIndex, uint32 Stride, uint32 Offset)
 	{
 #if D3D11_ALLOW_STATE_CACHE
 		D3D11_STATE_CACHE_VERIFY_PRE();
@@ -243,14 +219,7 @@ protected:
 			Slot.VertexBuffer = VertexBuffer;
 			Slot.Offset = Offset;
 			Slot.Stride = Stride;
-			if (AlternatePathFunction != nullptr)
-			{
-				(*AlternatePathFunction)(this, VertexBuffer, StreamIndex, Stride, Offset);
-			}
-			else
-			{
-				Direct3DDeviceIMContext->IASetVertexBuffers(StreamIndex, 1, &VertexBuffer, &Stride, &Offset);
-			}
+			Direct3DDeviceIMContext->IASetVertexBuffers(StreamIndex, 1, &VertexBuffer, &Stride, &Offset);
 		}
 		D3D11_STATE_CACHE_VERIFY_POST();
 #else
@@ -258,9 +227,8 @@ protected:
 #endif
 	}
 
-	typedef void (*TSetSamplerStateAlternate)(FD3D11StateCacheBase* StateCache, ID3D11SamplerState* SamplerState, uint32 SamplerIndex);
 	template <EShaderFrequency ShaderFrequency>
-	D3D11_STATE_CACHE_INLINE void InternalSetSamplerState(ID3D11SamplerState* SamplerState, uint32 SamplerIndex, TSetSamplerStateAlternate AlternatePathFunction)
+	D3D11_STATE_CACHE_INLINE void InternalSetSamplerState(ID3D11SamplerState* SamplerState, uint32 SamplerIndex)
 	{
 #if D3D11_ALLOW_STATE_CACHE
 		D3D11_STATE_CACHE_VERIFY_PRE();
@@ -268,14 +236,7 @@ protected:
 		if ((CurrentSamplerStates[ShaderFrequency][SamplerIndex] != SamplerState) || GD3D11SkipStateCaching)
 		{
 			CurrentSamplerStates[ShaderFrequency][SamplerIndex] = SamplerState;
-			if (AlternatePathFunction != nullptr)
-			{
-				(*AlternatePathFunction)(this, SamplerState, SamplerIndex);
-			}
-			else
-			{
-				InternalSetSamplerState<ShaderFrequency>(SamplerIndex, SamplerState);
-			}
+			InternalSetSamplerState<ShaderFrequency>(SamplerIndex, SamplerState);
 		}
 		D3D11_STATE_CACHE_VERIFY_POST();
 #else
@@ -302,9 +263,9 @@ public:
 	}
 
 	template <EShaderFrequency ShaderFrequency>
-	D3D11_STATE_CACHE_INLINE void SetShaderResourceView(ID3D11ShaderResourceView* SRV, uint32 ResourceIndex, ESRV_Type SrvType = SRV_Unknown)
+	D3D11_STATE_CACHE_INLINE void SetShaderResourceView(ID3D11ShaderResourceView* SRV, uint32 ResourceIndex)
 	{
-		InternalSetShaderResourceView<ShaderFrequency>(SRV, ResourceIndex, SrvType, nullptr);
+		InternalSetShaderResourceView<ShaderFrequency>(SRV, ResourceIndex);
 	}
 
 	template <EShaderFrequency ShaderFrequency>
@@ -416,7 +377,7 @@ public:
 	template <EShaderFrequency ShaderFrequency>
 	D3D11_STATE_CACHE_INLINE void SetSamplerState(ID3D11SamplerState* SamplerState, uint32 SamplerIndex)
 	{
-		InternalSetSamplerState<ShaderFrequency>(SamplerState, SamplerIndex, nullptr);
+		InternalSetSamplerState<ShaderFrequency>(SamplerState, SamplerIndex);
 	}
 
 	template <EShaderFrequency ShaderFrequency>
@@ -771,12 +732,12 @@ template <EShaderFrequency ShaderFrequency>
 	D3D11_STATE_CACHE_INLINE void SetStreamSource(ID3D11Buffer* VertexBuffer, uint32 StreamIndex, uint32 Stride, uint32 Offset)
 	{
 		ensure(Stride == StreamStrides[StreamIndex]);
-		InternalSetStreamSource(VertexBuffer, StreamIndex, Stride, Offset, nullptr);
+		InternalSetStreamSource(VertexBuffer, StreamIndex, Stride, Offset);
 	}
 
 	D3D11_STATE_CACHE_INLINE void SetStreamSource(ID3D11Buffer* VertexBuffer, uint32 StreamIndex, uint32 Offset)
 	{
-		InternalSetStreamSource(VertexBuffer, StreamIndex, StreamStrides[StreamIndex], Offset, nullptr);
+		InternalSetStreamSource(VertexBuffer, StreamIndex, StreamStrides[StreamIndex], Offset);
 	}
 
 	D3D11_STATE_CACHE_INLINE void GetStreamSources(uint32 StartStreamIndex, uint32 NumStreams, ID3D11Buffer** VertexBuffers, uint32* Strides, uint32* Offsets)
@@ -803,7 +764,7 @@ public:
 
 	D3D11_STATE_CACHE_INLINE void SetIndexBuffer(ID3D11Buffer* IndexBuffer, DXGI_FORMAT Format, uint32 Offset)
 	{
-		InternalSetIndexBuffer(IndexBuffer, Format, Offset, nullptr);
+		InternalSetIndexBuffer(IndexBuffer, Format, Offset);
 	}
 
 	D3D11_STATE_CACHE_INLINE void GetIndexBuffer(ID3D11Buffer** IndexBuffer, DXGI_FORMAT* Format, uint32* Offset)
diff --git a/Engine/Source/Runtime/Windows/D3D11RHI/Public/D3D11Resources.h b/Engine/Source/Runtime/Windows/D3D11RHI/Public/D3D11Resources.h
--- a/Engine/Source/Runtime/Windows/D3D11RHI/Public/D3D11Resources.h
+++ b/Engine/Source/Runtime/Windows/D3D11RHI/Public/D3D11Resources.h
@@ -159,7 +159,6 @@ public:
 	, RenderTargetViews(InRenderTargetViews)
 	, bCreatedRTVsPerSlice(bInCreatedRTVsPerSlice)
 	, RTVArraySize(InRTVArraySize)
-	, NumDepthStencilViews(0)	
 	{
 		// Set the DSVs for all the access type combinations
 		if ( InDepthStencilViews != nullptr )
@@ -167,10 +166,6 @@ public:
 			for (uint32 Index = 0; Index < FExclusiveDepthStencil::MaxIndex; Index++)
 			{
 				DepthStencilViews[Index] = InDepthStencilViews[Index];
-				// New Monolithic Graphics drivers have optional "fast calls" replacing various D3d functions
-				// You can't use fast version of XXSetShaderResources (called XXSetFastShaderResource) on dynamic or d/s targets
-				if ( DepthStencilViews[Index] != NULL )
-					NumDepthStencilViews++;
 			}
 		}
 	}
@@ -232,34 +227,19 @@ public:
 		return DepthStencilViews[AccessType.GetIndex()]; 
 	}
 
-	// New Monolithic Graphics drivers have optional "fast calls" replacing various D3d functions
-	// You can't use fast version of XXSetShaderResources (called XXSetFastShaderResource) on dynamic or d/s targets
-	bool HasDepthStencilView()
-	{
-		return ( NumDepthStencilViews > 0 );
-	}	
-
 	void AliasResources(FD3D11TextureBase* Texture)
 	{
 		check(MemorySize == Texture->MemorySize);
 		check(bCreatedRTVsPerSlice == Texture->bCreatedRTVsPerSlice);
 		check(RTVArraySize == Texture->RTVArraySize);
 
-		// If we're creating an aliased texture, make sure we handle this case correctly.
-		if (Texture->NumDepthStencilViews && !NumDepthStencilViews)
-		{
-			NumDepthStencilViews = Texture->NumDepthStencilViews;
-		}
-
-		check(NumDepthStencilViews == Texture->NumDepthStencilViews);
-
 		// Do not copy the BaseShaderResource from the source texture (this is initialized correctly here, and is used for
 		// state caching logic).
 		Resource = Texture->Resource;
 		ShaderResourceView = Texture->ShaderResourceView;
 		RenderTargetViews = Texture->RenderTargetViews;
 
-		for (uint32 Index = 0; Index < NumDepthStencilViews; Index++)
+		for (uint32 Index = 0; Index < FExclusiveDepthStencil::MaxIndex; Index++)
 		{
 			DepthStencilViews[Index] = Texture->DepthStencilViews[Index];
 		}
@@ -294,9 +274,6 @@ protected:
 
 	/** A depth-stencil targetable view of the texture. */
 	TRefCountPtr<ID3D11DepthStencilView> DepthStencilViews[FExclusiveDepthStencil::MaxIndex];
-
-	/** Number of Depth Stencil Views - used for fast call tracking. */
-	uint32	NumDepthStencilViews;	
 };
 
 /** 2D texture (vanilla, cubemap or 2D array) */
