Fix for crash in PIE audio device removal/add logic

diff --git a/Engine/Source/Runtime/Windows/AudioMixerXAudio2/Private/AudioMixerPlatformWindows.cpp b/Engine/Source/Runtime/Windows/AudioMixerXAudio2/Private/AudioMixerPlatformWindows.cpp
--- a/Engine/Source/Runtime/Windows/AudioMixerXAudio2/Private/AudioMixerPlatformWindows.cpp
+++ b/Engine/Source/Runtime/Windows/AudioMixerXAudio2/Private/AudioMixerPlatformWindows.cpp
@@ -1236,7 +1236,6 @@ namespace Audio
 			}
 			
 			WindowsNotificationClient->UnRegisterDeviceDeviceChangedListener(this);
-			WindowsNotificationClient.Reset();
 		}
 	}
 
