Reduce logging for AudioDeviceSwap. Put some existing logs behind flag and bump verbosity.

diff --git a/Engine/Source/Runtime/Windows/AudioMixerXAudio2/Private/AudioMixerPlatformWindows.cpp b/Engine/Source/Runtime/Windows/AudioMixerXAudio2/Private/AudioMixerPlatformWindows.cpp
--- a/Engine/Source/Runtime/Windows/AudioMixerXAudio2/Private/AudioMixerPlatformWindows.cpp
+++ b/Engine/Source/Runtime/Windows/AudioMixerXAudio2/Private/AudioMixerPlatformWindows.cpp
@@ -204,11 +204,8 @@ public:
 	{
 		FScopeLock ScopeLock(&ListenerArrayMutationLock);
 
-		if (Audio::IAudioMixer::ShouldLogDeviceSwaps())
-		{
-			UE_LOG(LogAudioMixer, Warning, TEXT("OnDeviceAdded: %s"), *GetFriendlyName(pwstrDeviceId));
-		}
-		
+		UE_CLOG(Audio::IAudioMixer::ShouldLogDeviceSwaps(), LogAudioMixer, Verbose, TEXT("OnDeviceAdded: %s"), *GetFriendlyName(pwstrDeviceId));
+			
 		if (Audio::IAudioMixer::ShouldIgnoreDeviceSwaps())
 		{
 			return S_OK;
@@ -227,10 +224,7 @@ public:
 	{
 		FScopeLock ScopeLock(&ListenerArrayMutationLock);
 
-		if (Audio::IAudioMixer::ShouldLogDeviceSwaps())
-		{
-			UE_LOG(LogAudioMixer, Warning, TEXT("OnDeviceRemoved: %s"), *GetFriendlyName(pwstrDeviceId));
-		}
+		UE_CLOG(Audio::IAudioMixer::ShouldLogDeviceSwaps(),LogAudioMixer, Verbose, TEXT("OnDeviceRemoved: %s"), *GetFriendlyName(pwstrDeviceId));
 
 		if (Audio::IAudioMixer::ShouldIgnoreDeviceSwaps())
 		{
@@ -249,10 +243,7 @@ public:
 	{
 		FScopeLock ScopeLock(&ListenerArrayMutationLock);
 
-		if (Audio::IAudioMixer::ShouldLogDeviceSwaps())
-		{
-			UE_LOG(LogAudioMixer, Warning, TEXT("OnDeviceStateChanged: %s, %d"), *GetFriendlyName(pwstrDeviceId), dwNewState);
-		}
+		UE_CLOG(Audio::IAudioMixer::ShouldLogDeviceSwaps(), LogAudioMixer, Verbose, TEXT("OnDeviceStateChanged: %s, %d"), *GetFriendlyName(pwstrDeviceId), dwNewState);
 
 		if (Audio::IAudioMixer::ShouldIgnoreDeviceSwaps())
 		{
@@ -529,7 +520,7 @@ public:
 	HRESULT STDMETHODCALLTYPE OnSessionDisconnected(
 		AudioSessionDisconnectReason InDisconnectReason)
 	{
-		UE_LOG(LogAudioMixer, Warning, TEXT("Session Disconnect: %s"), GetDisconnectReasonString(InDisconnectReason));
+		UE_CLOG(Audio::IAudioMixer::ShouldLogDeviceSwaps(), LogAudioMixer, Log, TEXT("Session Disconnect: %s"), GetDisconnectReasonString(InDisconnectReason));
 
 		Audio::IAudioMixerDeviceChangedListener::EDisconnectReason Reason = AudioSessionDisconnectToEDisconnectReason(InDisconnectReason);
 		FScopeLock ScopeLock(&ListenerArrayMutationLock);
@@ -731,7 +722,7 @@ namespace Audio
 			// We didn't match channel masks for all channels, revert to a default ordering
 			if (ChanCount < (uint32)OutInfo.NumChannels)
 			{
-				UE_LOG(LogAudioMixer, Warning, TEXT("Did not find the channel type flags for audio device '%s'. Reverting to a default channel ordering."), *OutInfo.FriendlyName);
+				UE_CLOG(Audio::IAudioMixer::ShouldLogDeviceSwaps(), LogAudioMixer, Warning, TEXT("Did not find the channel type flags for audio device '%s'. Reverting to a default channel ordering."), *OutInfo.FriendlyName);
 
 				OutInfo.OutputChannels.Reset();
 
@@ -879,7 +870,8 @@ namespace Audio
 				}
 				else
 				{
-					UE_LOG(LogAudioMixer, Warning, TEXT("Failed to get Format for '%s'"), *OutInfo.FriendlyName);
+					// Log a warning if this device is active as we failed to ask for a format
+					UE_CLOG(DeviceState == DEVICE_STATE_ACTIVE, LogAudioMixer, Warning, TEXT("Failed to get Format for active device '%s'"), *OutInfo.FriendlyName);
 				}
 			}
 	
@@ -915,7 +907,7 @@ namespace Audio
 								// Enumerate props into our info object.
 								EnumerateDeviceProps(Device, Info);
 
-								UE_LOG(LogAudioMixer, Display, TEXT("%s Device '%s' ID='%s'"),
+								UE_LOG(LogAudioMixer, Verbose, TEXT("%s Device '%s' ID='%s'"),
 									Info.Type == FCacheEntry::EEndpointType::Capture ? TEXT("Capture") :
 									Info.Type == FCacheEntry::EEndpointType::Render ? TEXT("Render") :
 									TEXT("UNKNOWN!"),
@@ -1032,7 +1024,7 @@ namespace Audio
 		{
 			TOptional<FCacheEntry> Info = BuildCacheEntry(DeviceId);
 			
-			UE_LOG(LogAudioMixer, Display, TEXT("Device '%s' - '%s' state changed."), Info ? *Info->FriendlyName : TEXT("Unknown"), *DeviceId);
+			UE_CLOG(Audio::IAudioMixer::ShouldLogDeviceSwaps(), LogAudioMixer, Verbose, TEXT("Device '%s' - '%s' state changed."), Info ? *Info->FriendlyName : TEXT("Unknown"), *DeviceId);
 			
 			FReadScopeLock ReadLock(CacheMutationLock);
 			FName DeviceIdName = *DeviceId;
@@ -1057,21 +1049,21 @@ namespace Audio
 					if (Found->NumChannels != InFormat.NumChannels)
 					{
 						FWriteScopeLock WriteLock(Found->MutationLock);
-						UE_LOG(LogAudioMixer, Display, TEXT("Device '%s' changed default format from %d channels to %d."), *InDeviceId, Found->NumChannels, InFormat.NumChannels);
+						UE_CLOG(Audio::IAudioMixer::ShouldLogDeviceSwaps(),LogAudioMixer, Verbose, TEXT("Device '%s' changed default format from %d channels to %d."), *InDeviceId, Found->NumChannels, InFormat.NumChannels);
 						Found->NumChannels = InFormat.NumChannels;
 						bDirty = true;
 					}
 					if (Found->SampleRate != InFormat.SampleRate)
 					{
 						FWriteScopeLock WriteLock(Found->MutationLock);
-						UE_LOG(LogAudioMixer, Display, TEXT("Device '%s' changed default format from %dhz to %dhz."), *InDeviceId, Found->SampleRate, InFormat.SampleRate);
+						UE_CLOG(Audio::IAudioMixer::ShouldLogDeviceSwaps(),LogAudioMixer, Verbose, TEXT("Device '%s' changed default format from %dhz to %dhz."), *InDeviceId, Found->SampleRate, InFormat.SampleRate);
 						Found->SampleRate = InFormat.SampleRate;
 						bDirty = true;
 					}
 					if (Found->SpeakerConfig != InFormat.ChannelConfig)
 					{
 						FWriteScopeLock WriteLock(Found->MutationLock);
-						UE_LOG(LogAudioMixer, Display, TEXT("Device '%s' changed default format from 0x%x to 0x%x"), *InDeviceId, Found->SpeakerConfig, InFormat.ChannelConfig);
+						UE_CLOG(Audio::IAudioMixer::ShouldLogDeviceSwaps(),LogAudioMixer, Verbose, TEXT("Device '%s' changed default format from 0x%x to 0x%x"), *InDeviceId, Found->SpeakerConfig, InFormat.ChannelConfig);
 						Found->SpeakerConfig = InFormat.ChannelConfig;
 						bDirty = true;
 					}
@@ -1091,7 +1083,7 @@ namespace Audio
 				{
 					FWriteScopeLock WriteLock(Found->MutationLock);
 					
-					UE_LOG(LogAudioMixer, Display, TEXT("Device '%s' changed default format from 0x%x to 0x%x"), 
+					UE_CLOG(Audio::IAudioMixer::ShouldLogDeviceSwaps(),LogAudioMixer, Verbose, TEXT("Device '%s' changed default format from 0x%x to 0x%x"),
 						*InDeviceId, Found->SpeakerConfig, InSpeakerBitmask);
 					Found->SpeakerConfig = InSpeakerBitmask;
 
