Fix for crash in Editor device swap

diff --git a/Engine/Source/Runtime/AudioMixerCore/Private/AudioMixer.cpp b/Engine/Source/Runtime/AudioMixerCore/Private/AudioMixer.cpp
--- a/Engine/Source/Runtime/AudioMixerCore/Private/AudioMixer.cpp
+++ b/Engine/Source/Runtime/AudioMixerCore/Private/AudioMixer.cpp
@@ -60,7 +60,7 @@ FAutoConsoleVariableRef CVarDisableDeviceSwap(
 	TEXT("0: Not Enabled, 1: Enabled"),
 	ECVF_Default);
 
-static int32 bUseThreadedDeviceSwapCVar = 0;
+static int32 bUseThreadedDeviceSwapCVar = 1;
 FAutoConsoleVariableRef CVarUseThreadedDeviceSwap(
 	TEXT("au.UseThreadedDeviceSwap"),
 	bUseThreadedDeviceSwapCVar,
diff --git a/Engine/Source/Runtime/Windows/AudioMixerXAudio2/Private/AudioMixerPlatformWindows.cpp b/Engine/Source/Runtime/Windows/AudioMixerXAudio2/Private/AudioMixerPlatformWindows.cpp
--- a/Engine/Source/Runtime/Windows/AudioMixerXAudio2/Private/AudioMixerPlatformWindows.cpp
+++ b/Engine/Source/Runtime/Windows/AudioMixerXAudio2/Private/AudioMixerPlatformWindows.cpp
@@ -71,6 +71,16 @@ public:
 
 	bool RegisterForSessionNotifications(const Audio::TScopeComObject<IMMDevice>& InDevice)
 	{
+		FScopeLock ScopeLock(&MutationCs);
+
+		// If we're already listening to this device, we can early out.
+		if (DeviceListeningToSessionEvents == InDevice)
+		{
+			return true;
+		}
+
+		DeviceListeningToSessionEvents = InDevice;
+				
 		// Unregister for any device we're already listening to.
 		if (SessionControls)
 		{
@@ -127,7 +137,7 @@ public:
 
 	HRESULT STDMETHODCALLTYPE OnDefaultDeviceChanged(EDataFlow InFlow, ERole InRole, LPCWSTR pwstrDeviceId) override
 	{
-		FScopeLock ScopeLock(&ListenerArrayMutationLock);
+		FScopeLock ScopeLock(&MutationCs);
 
 		if (Audio::IAudioMixer::ShouldLogDeviceSwaps())
 		{
@@ -202,7 +212,7 @@ public:
 
 	HRESULT STDMETHODCALLTYPE OnDeviceAdded(LPCWSTR pwstrDeviceId) override
 	{
-		FScopeLock ScopeLock(&ListenerArrayMutationLock);
+		FScopeLock ScopeLock(&MutationCs);
 
 		UE_CLOG(Audio::IAudioMixer::ShouldLogDeviceSwaps(), LogAudioMixer, Verbose, TEXT("OnDeviceAdded: %s"), *GetFriendlyName(pwstrDeviceId));
 			
@@ -222,7 +232,7 @@ public:
 
 	HRESULT STDMETHODCALLTYPE OnDeviceRemoved(LPCWSTR pwstrDeviceId) override
 	{
-		FScopeLock ScopeLock(&ListenerArrayMutationLock);
+		FScopeLock ScopeLock(&MutationCs);
 
 		UE_CLOG(Audio::IAudioMixer::ShouldLogDeviceSwaps(),LogAudioMixer, Verbose, TEXT("OnDeviceRemoved: %s"), *GetFriendlyName(pwstrDeviceId));
 
@@ -241,7 +251,7 @@ public:
 
 	HRESULT STDMETHODCALLTYPE OnDeviceStateChanged(LPCWSTR pwstrDeviceId, DWORD dwNewState) override
 	{
-		FScopeLock ScopeLock(&ListenerArrayMutationLock);
+		FScopeLock ScopeLock(&MutationCs);
 
 		UE_CLOG(Audio::IAudioMixer::ShouldLogDeviceSwaps(), LogAudioMixer, Verbose, TEXT("OnDeviceStateChanged: %s, %d"), *GetFriendlyName(pwstrDeviceId), dwNewState);
 
@@ -382,7 +392,7 @@ public:
 							FormatChanged.ChannelConfig = WaveFormatEx->wFormatTag == WAVE_FORMAT_EXTENSIBLE ?
 								((const WAVEFORMATEXTENSIBLE*)WaveFormatEx)->dwChannelMask : 0;
 
-							FScopeLock ScopeLock(&ListenerArrayMutationLock);
+							FScopeLock ScopeLock(&MutationCs);
 							for (Audio::IAudioMixerDeviceChangedListener* i : Listeners)
 							{
 								i->OnFormatChanged(DeviceId, FormatChanged);
@@ -393,7 +403,7 @@ public:
 					{
 						if (SUCCEEDED(PropertyStore->GetValue(PKEY_AudioEndpoint_PhysicalSpeakers, &Prop)))
 						{
-							FScopeLock ScopeLock(&ListenerArrayMutationLock);
+							FScopeLock ScopeLock(&MutationCs);
 							for (Audio::IAudioMixerDeviceChangedListener* i : Listeners)
 			{
 								i->OnSpeakerConfigChanged(DeviceId, Prop.uintVal);
@@ -462,13 +472,13 @@ public:
 
 	void RegisterDeviceChangedListener(Audio::IAudioMixerDeviceChangedListener* DeviceChangedListener)
 	{
-		FScopeLock ScopeLock(&ListenerArrayMutationLock);
+		FScopeLock ScopeLock(&MutationCs);
 		Listeners.Add(DeviceChangedListener);
 	}
 
 	void UnRegisterDeviceDeviceChangedListener(Audio::IAudioMixerDeviceChangedListener* DeviceChangedListener)
 	{
-		FScopeLock ScopeLock(&ListenerArrayMutationLock);
+		FScopeLock ScopeLock(&MutationCs);
 		Listeners.Remove(DeviceChangedListener);
 	}
 
@@ -523,7 +533,7 @@ public:
 		UE_CLOG(Audio::IAudioMixer::ShouldLogDeviceSwaps(), LogAudioMixer, Log, TEXT("Session Disconnect: %s"), GetDisconnectReasonString(InDisconnectReason));
 
 		Audio::IAudioMixerDeviceChangedListener::EDisconnectReason Reason = AudioSessionDisconnectToEDisconnectReason(InDisconnectReason);
-		FScopeLock ScopeLock(&ListenerArrayMutationLock);
+		FScopeLock ScopeLock(&MutationCs);
 		for (Audio::IAudioMixerDeviceChangedListener* i : Listeners)
 		{
 			i->OnSessionDisconnect(Reason);
@@ -573,10 +583,11 @@ public:
 private:
 	LONG Ref;
 	TSet<Audio::IAudioMixerDeviceChangedListener*> Listeners;
-	FCriticalSection ListenerArrayMutationLock;
+	FCriticalSection MutationCs;
 	Audio::TScopeComObject<IMMDeviceEnumerator> DeviceEnumerator;
 	Audio::TScopeComObject<IAudioSessionManager> SessionManager;
 	Audio::TScopeComObject<IAudioSessionControl> SessionControls;
+	Audio::TScopeComObject<IMMDevice> DeviceListeningToSessionEvents;
 	bool bComInitialized;
 };
 
diff --git a/Engine/Source/Runtime/Windows/AudioMixerXAudio2/Private/ScopedCom.h b/Engine/Source/Runtime/Windows/AudioMixerXAudio2/Private/ScopedCom.h
--- a/Engine/Source/Runtime/Windows/AudioMixerXAudio2/Private/ScopedCom.h
+++ b/Engine/Source/Runtime/Windows/AudioMixerXAudio2/Private/ScopedCom.h
@@ -69,6 +69,13 @@ namespace Audio
 			InObj.Obj = nullptr;
 			return *this;
 		}
+
+		// Equality operator
+		bool operator==(const ThisType& InObj) const
+		{
+			return InObj.Obj == Obj;
+		}
+
 	private:
 		void SafeAddRef()
 		{
