OnDisconnected double trigger fix

diff --git a/Engine/Source/Runtime/Windows/AudioMixerXAudio2/Private/AudioMixerPlatformWindows.cpp b/Engine/Source/Runtime/Windows/AudioMixerXAudio2/Private/AudioMixerPlatformWindows.cpp
--- a/Engine/Source/Runtime/Windows/AudioMixerXAudio2/Private/AudioMixerPlatformWindows.cpp
+++ b/Engine/Source/Runtime/Windows/AudioMixerXAudio2/Private/AudioMixerPlatformWindows.cpp
@@ -52,6 +52,7 @@ public:
 	FWindowsMMNotificationClient()
 		: Ref(1)
 		, DeviceEnumerator(nullptr)
+		, bHasDisconnectSessionHappened(false)
 	{
 		bComInitialized = FWindowsPlatformMisc::CoInitialize();
 		HRESULT Result = CoCreateInstance(__uuidof(MMDeviceEnumerator), nullptr, CLSCTX_INPROC_SERVER, IID_PPV_ARGS(&DeviceEnumerator.Obj));
@@ -82,6 +83,7 @@ public:
 		}
 
 		DeviceListeningToSessionEvents = InDevice;
+		bHasDisconnectSessionHappened = false;
 				
 		// Unregister for any device we're already listening to.
 		if (SessionControls)
@@ -534,11 +536,17 @@ public:
 	{
 		UE_CLOG(Audio::IAudioMixer::ShouldLogDeviceSwaps(), LogAudioMixer, Log, TEXT("Session Disconnect: %s"), GetDisconnectReasonString(InDisconnectReason));
 
-		Audio::IAudioMixerDeviceChangedListener::EDisconnectReason Reason = AudioSessionDisconnectToEDisconnectReason(InDisconnectReason);
 		FScopeLock ScopeLock(&MutationCs);
-		for (Audio::IAudioMixerDeviceChangedListener* i : Listeners)
+		if (!bHasDisconnectSessionHappened)
 		{
-			i->OnSessionDisconnect(Reason);
+			Audio::IAudioMixerDeviceChangedListener::EDisconnectReason Reason = AudioSessionDisconnectToEDisconnectReason(InDisconnectReason);
+			for (Audio::IAudioMixerDeviceChangedListener* i : Listeners)
+			{
+				i->OnSessionDisconnect(Reason);
+			}
+			
+			// Mark this true.
+			bHasDisconnectSessionHappened = true;
 		}
 		return S_OK;
 	}
@@ -591,6 +599,7 @@ private:
 	Audio::TScopeComObject<IAudioSessionControl> SessionControls;
 	Audio::TScopeComObject<IMMDevice> DeviceListeningToSessionEvents;
 	bool bComInitialized;
+	bool bHasDisconnectSessionHappened;
 };
 
 namespace Audio
diff --git a/Engine/Source/Runtime/Windows/AudioMixerXAudio2/Private/AudioMixerPlatformXAudio2.cpp b/Engine/Source/Runtime/Windows/AudioMixerXAudio2/Private/AudioMixerPlatformXAudio2.cpp
--- a/Engine/Source/Runtime/Windows/AudioMixerXAudio2/Private/AudioMixerPlatformXAudio2.cpp
+++ b/Engine/Source/Runtime/Windows/AudioMixerXAudio2/Private/AudioMixerPlatformXAudio2.cpp
@@ -2075,9 +2075,14 @@ namespace Audio
 			// OnFormatChanged, retry again same device.
 			RequestDeviceSwap(GetDeviceId());
 		}
+		else if (InReason == IAudioMixerDeviceChangedListener::EDisconnectReason::DeviceRemoval)
+		{
+			// Ignore Device Removal, as this is handle by the Device Removal logic in the Notification Client.
+		}
 		else
 		{
-			// Device removed entirely, so attempt to switch to default.
+			// ServerShutdown, SessionLogoff, SessionDisconnected, ExclusiveModeOverride
+			// Attempt a default swap, will likely fail, but then we'll switch to a null device.
 			RequestDeviceSwap(TEXT(""));
 		}
 	}
