Safety refactor, Convert Audio::TScopedComPtr to TComPtr

diff --git a/Engine/Source/Runtime/Windows/AudioMixerXAudio2/Private/AudioMixerPlatformWindows.cpp b/Engine/Source/Runtime/Windows/AudioMixerXAudio2/Private/AudioMixerPlatformWindows.cpp
--- a/Engine/Source/Runtime/Windows/AudioMixerXAudio2/Private/AudioMixerPlatformWindows.cpp
+++ b/Engine/Source/Runtime/Windows/AudioMixerXAudio2/Private/AudioMixerPlatformWindows.cpp
@@ -12,6 +12,7 @@
 
 #if PLATFORM_WINDOWS
 
+#include "Microsoft/COMPointer.h"
 #include "ScopedCom.h"
 
 #include "Windows/AllowWindowsPlatformTypes.h"
@@ -55,7 +56,7 @@ public:
 		, bHasDisconnectSessionHappened(false)
 	{
 		bComInitialized = FWindowsPlatformMisc::CoInitialize();
-		HRESULT Result = CoCreateInstance(__uuidof(MMDeviceEnumerator), nullptr, CLSCTX_INPROC_SERVER, IID_PPV_ARGS(&DeviceEnumerator.Obj));
+		HRESULT Result = CoCreateInstance(__uuidof(MMDeviceEnumerator), nullptr, CLSCTX_INPROC_SERVER, IID_PPV_ARGS(&DeviceEnumerator));
 		if (Result == S_OK)
 		{
 			DeviceEnumerator->RegisterEndpointNotificationCallback(this);
@@ -64,15 +65,15 @@ public:
 		// Register for session events from default endpoint.
 		if (DeviceEnumerator)
 		{
-			Audio::TScopeComObject<IMMDevice> DefaultDevice;
-			if (SUCCEEDED(DeviceEnumerator->GetDefaultAudioEndpoint(eRender, eMultimedia, &DefaultDevice.Obj)))
+			TComPtr<IMMDevice> DefaultDevice;
+			if (SUCCEEDED(DeviceEnumerator->GetDefaultAudioEndpoint(eRender, eMultimedia, &DefaultDevice)))
 			{
 				RegisterForSessionNotifications(DefaultDevice);
 			}
 		}
 	}
 
-	bool RegisterForSessionNotifications(const Audio::TScopeComObject<IMMDevice>& InDevice)
+	bool RegisterForSessionNotifications(const TComPtr<IMMDevice>& InDevice)
 	{
 		FScopeLock ScopeLock(&MutationCs);
 
@@ -98,9 +99,9 @@ public:
 		
 		if(InDevice) 
 		{
-			if (SUCCEEDED(InDevice->Activate(__uuidof(IAudioSessionManager), CLSCTX_INPROC_SERVER, NULL, (void**)&SessionManager.Obj)))
+			if (SUCCEEDED(InDevice->Activate(__uuidof(IAudioSessionManager), CLSCTX_INPROC_SERVER, NULL, (void**)&SessionManager)))
 			{
-				if (SUCCEEDED(SessionManager->GetAudioSessionControl(NULL, 0, &SessionControls.Obj)))
+				if (SUCCEEDED(SessionManager->GetAudioSessionControl(NULL, 0, &SessionControls)))
 				{
 					if (SUCCEEDED(SessionControls->RegisterAudioSessionNotification(this)))
 					{
@@ -114,7 +115,7 @@ public:
 
 	bool RegisterForSessionNotifications(const FString& InDeviceId)
 	{
-		if (Audio::TScopeComObject<IMMDevice> Device = GetDevice(InDeviceId))
+		if (TComPtr<IMMDevice> Device = GetDevice(InDeviceId))
 		{
 			return RegisterForSessionNotifications(Device);
 		}
@@ -199,10 +200,10 @@ public:
 	bool IsRenderDevice(const FString& InDeviceId) const
 	{
 		bool bIsRender = true;
-		if (Audio::TScopeComObject<IMMDevice> Device = GetDevice(InDeviceId))
+		if (TComPtr<IMMDevice> Device = GetDevice(InDeviceId))
 		{
-			Audio::TScopeComObject<IMMEndpoint> Endpoint;
-			if (SUCCEEDED(Device->QueryInterface(IID_PPV_ARGS(&Endpoint.Obj))))
+			TComPtr<IMMEndpoint> Endpoint;
+			if (SUCCEEDED(Device->QueryInterface(IID_PPV_ARGS(&Endpoint))))
 			{
 				EDataFlow DataFlow = eRender;
 				if (SUCCEEDED(Endpoint->GetDataFlow(&DataFlow)))
@@ -315,11 +316,11 @@ public:
 		FString FriendlyName = TEXT("[No Friendly Name for Device]");
 			
 		// Get device.
-		if (Audio::TScopeComObject<IMMDevice> Device = GetDevice(*InDeviceID))
+		if (TComPtr<IMMDevice> Device = GetDevice(*InDeviceID))
 		{
 			// Get property store.
-			Audio::TScopeComObject<IPropertyStore> PropStore;
-			HRESULT Hr = Device->OpenPropertyStore(STGM_READ, &PropStore.Obj);
+			TComPtr<IPropertyStore> PropStore;
+			HRESULT Hr = Device->OpenPropertyStore(STGM_READ, &PropStore);
 
 			// Get friendly name.
 			if (SUCCEEDED(Hr) && PropStore)
@@ -344,11 +345,11 @@ public:
 		return FriendlyName;
 		}
 
-	Audio::TScopeComObject<IMMDevice> GetDevice(const FString InDeviceID) const
+	TComPtr<IMMDevice> GetDevice(const FString InDeviceID) const
 	{
 		// Get device.
-		Audio::TScopeComObject<IMMDevice> Device;
-		HRESULT Hr = DeviceEnumerator->GetDevice(*InDeviceID, &Device.Obj);
+		TComPtr<IMMDevice> Device;
+		HRESULT Hr = DeviceEnumerator->GetDevice(*InDeviceID, &Device);
 		if (SUCCEEDED(Hr))
 		{
 			return Device;
@@ -368,14 +369,14 @@ public:
 		{
 			// Get device.
 			FString DeviceId = pwstrDeviceId;
-			Audio::TScopeComObject<IMMDevice> Device;			
-			HRESULT Hr = DeviceEnumerator->GetDevice(*DeviceId, &Device.Obj);
+			TComPtr<IMMDevice> Device;			
+			HRESULT Hr = DeviceEnumerator->GetDevice(*DeviceId, &Device);
 
 			// Get property store.
-			Audio::TScopeComObject<IPropertyStore> PropertyStore;
+			TComPtr<IPropertyStore> PropertyStore;
 			if (SUCCEEDED(Hr) && Device)
 			{
-				Hr = Device->OpenPropertyStore(STGM_READ, &PropertyStore.Obj);
+				Hr = Device->OpenPropertyStore(STGM_READ, &PropertyStore);
 				if (SUCCEEDED(Hr) && PropertyStore)
 				{
 					// Device Format
@@ -594,10 +595,10 @@ private:
 	LONG Ref;
 	TSet<Audio::IAudioMixerDeviceChangedListener*> Listeners;
 	FCriticalSection MutationCs;
-	Audio::TScopeComObject<IMMDeviceEnumerator> DeviceEnumerator;
-	Audio::TScopeComObject<IAudioSessionManager> SessionManager;
-	Audio::TScopeComObject<IAudioSessionControl> SessionControls;
-	Audio::TScopeComObject<IMMDevice> DeviceListeningToSessionEvents;
+	TComPtr<IMMDeviceEnumerator> DeviceEnumerator;
+	TComPtr<IAudioSessionManager> SessionManager;
+	TComPtr<IAudioSessionControl> SessionControls;
+	TComPtr<IMMDevice> DeviceListeningToSessionEvents;
 	bool bComInitialized;
 	bool bHasDisconnectSessionHappened;
 };
@@ -675,9 +676,9 @@ namespace Audio
 				: DeviceId{ InDeviceId }
 			{}
 		};
-		
-		TScopeComObject<IMMDeviceEnumerator> DeviceEnumerator;
-		
+
+		TComPtr<IMMDeviceEnumerator> DeviceEnumerator;
+
 		mutable FRWLock CacheMutationLock;							// R/W lock protects map and default arrays.
 		TMap<FName, FCacheEntry> Cache;								// DeviceID GUID -> Info.
 		FName DefaultCaptureId[(int32)EAudioDeviceRole::COUNT];		// Role -> DeviceID GUID
@@ -685,7 +686,7 @@ namespace Audio
 
 		FWindowsMMDeviceCache()
 		{
-			ensure(SUCCEEDED(CoCreateInstance(__uuidof(MMDeviceEnumerator), NULL, CLSCTX_INPROC_SERVER, IID_PPV_ARGS(&DeviceEnumerator.Obj))) && DeviceEnumerator);
+			ensure(SUCCEEDED(CoCreateInstance(__uuidof(MMDeviceEnumerator), NULL, CLSCTX_INPROC_SERVER, IID_PPV_ARGS(&DeviceEnumerator))) && DeviceEnumerator);
 
 			EnumerateEndpoints();
 			EnumerateDefaults();
@@ -795,7 +796,7 @@ namespace Audio
 			}
 			return true;
 		}
-		
+
 		bool EnumerateSpeakers(const WAVEFORMATEX* InFormat, FCacheEntry& OutInfo)
 		{
 			OutInfo.OutputChannels.Empty();
@@ -821,19 +822,19 @@ namespace Audio
 			return true;
 		}
 
-		FCacheEntry::EEndpointType QueryDeviceDataFlow(const Audio::TScopeComObject<IMMDevice>& InDevice) const
+		FCacheEntry::EEndpointType QueryDeviceDataFlow(const TComPtr<IMMDevice>& InDevice) const
 		{
-			Audio::TScopeComObject<IMMEndpoint> Endpoint;
-			if (SUCCEEDED(InDevice->QueryInterface(IID_PPV_ARGS(&Endpoint.Obj))))
+			TComPtr<IMMEndpoint> Endpoint;
+			if (SUCCEEDED(InDevice->QueryInterface(IID_PPV_ARGS(&Endpoint))))
 			{
 				EDataFlow DataFlow = eRender;
 				if (SUCCEEDED(Endpoint->GetDataFlow(&DataFlow)))
 				{
 					switch (DataFlow)
 					{
-					case eRender: 
+					case eRender:
 						return FCacheEntry::EEndpointType::Render;
-					case eCapture: 
+					case eCapture:
 						return FCacheEntry::EEndpointType::Capture;
 					default:
 						break;
@@ -843,7 +844,7 @@ namespace Audio
 			return FCacheEntry::EEndpointType::Unknown;
 		}
 
-		bool EnumerateDeviceProps(const Audio::TScopeComObject<IMMDevice>& InDevice, FCacheEntry& OutInfo)
+		bool EnumerateDeviceProps(const TComPtr<IMMDevice>& InDevice, FCacheEntry& OutInfo)
 		{
 			// Mark if this is a Render Device or Capture or Unknown.
 			OutInfo.Type = QueryDeviceDataFlow(InDevice);
@@ -855,8 +856,8 @@ namespace Audio
 				OutInfo.State = ConvertWordToDeviceState(DeviceState);
 			}
 
-			Audio::TScopeComObject<IPropertyStore> PropertyStore;
-			if (SUCCEEDED(InDevice->OpenPropertyStore(STGM_READ, &PropertyStore.Obj)))
+			TComPtr<IPropertyStore> PropertyStore;
+			if (SUCCEEDED(InDevice->OpenPropertyStore(STGM_READ, &PropertyStore)))
 			{
 				// Friendly Name
 				PROPVARIANT FriendlyName;
@@ -867,12 +868,12 @@ namespace Audio
 					PropVariantClear(&FriendlyName);
 				}
 
-				auto EnumDeviceFormat = [this](const Audio::TScopeComObject<IPropertyStore>& InPropStore, REFPROPERTYKEY InKey, FCacheEntry& OutInfo) -> bool
+				auto EnumDeviceFormat = [this](const TComPtr<IPropertyStore>& InPropStore, REFPROPERTYKEY InKey, FCacheEntry& OutInfo) -> bool
 				{
 					// Device Format
 					PROPVARIANT DeviceFormat;
 					PropVariantInit(&DeviceFormat);
-					
+
 					if (SUCCEEDED(InPropStore->GetValue(InKey, &DeviceFormat)) && DeviceFormat.blob.pBlobData)
 					{
 						const WAVEFORMATEX* WaveFormatEx = (const WAVEFORMATEX*)(DeviceFormat.blob.pBlobData);
@@ -910,15 +911,15 @@ namespace Audio
 			if (DeviceEnumerator)
 			{
 				// Get Render Device Collection. (note we ask for ALL states, which include disabled/unplugged devices.).
-				Audio::TScopeComObject<IMMDeviceCollection> DeviceCollection;
+				TComPtr<IMMDeviceCollection> DeviceCollection;
 				uint32 DeviceCount = 0;
-				if (SUCCEEDED(DeviceEnumerator->EnumAudioEndpoints(eAll, DEVICE_STATEMASK_ALL, &DeviceCollection.Obj)) && DeviceCollection &&
+				if (SUCCEEDED(DeviceEnumerator->EnumAudioEndpoints(eAll, DEVICE_STATEMASK_ALL, &DeviceCollection)) && DeviceCollection &&
 					SUCCEEDED(DeviceCollection->GetCount(&DeviceCount)))
 				{
 					for (uint32 i = 0; i < DeviceCount; ++i)
 					{
-						Audio::TScopeComObject<IMMDevice> Device;
-						if (SUCCEEDED(DeviceCollection->Item(i, &Device.Obj)) && Device)
+						TComPtr<IMMDevice> Device;
+						if (SUCCEEDED(DeviceCollection->Item(i, &Device)) && Device)
 						{
 							// Get the device id string (guid)
 							Audio::FScopeComString DeviceIdString;
@@ -958,8 +959,8 @@ namespace Audio
 			{
 				// Mark default device.
 				bool bSuccess = false;
-				TScopeComObject<IMMDevice> DefaultDevice;
-				if (SUCCEEDED(DeviceEnumerator->GetDefaultAudioEndpoint(InDataFlow, InRole, &DefaultDevice.Obj)))
+				TComPtr<IMMDevice> DefaultDevice;
+				if (SUCCEEDED(DeviceEnumerator->GetDefaultAudioEndpoint(InDataFlow, InRole, &DefaultDevice)))
 				{
 					Audio::FScopeComString DeviceIdString;
 					if (SUCCEEDED(DefaultDevice->GetId(&DeviceIdString.StringPtr)) && DeviceIdString)
@@ -1006,8 +1007,8 @@ namespace Audio
 		{
 			if (ensure(DeviceEnumerator))
 			{
-				Audio::TScopeComObject<IMMDevice> Device;
-				if(SUCCEEDED(DeviceEnumerator->GetDevice(*DeviceId, &Device.Obj)))
+				TComPtr<IMMDevice> Device;
+				if (SUCCEEDED(DeviceEnumerator->GetDevice(*DeviceId, &Device)))
 				{
 					FCacheEntry Info{ *DeviceId };
 					if (EnumerateDeviceProps(Device, Info))
@@ -1030,8 +1031,8 @@ namespace Audio
 
 		TOptional<FCacheEntry> BuildCacheEntry(const FString& DeviceId)
 		{
-			Audio::TScopeComObject<IMMDevice> Device;
-			if (SUCCEEDED(DeviceEnumerator->GetDevice(*DeviceId, &Device.Obj)))
+			TComPtr<IMMDevice> Device;
+			if (SUCCEEDED(DeviceEnumerator->GetDevice(*DeviceId, &Device)))
 			{
 				FCacheEntry Info{ *DeviceId };
 				if (EnumerateDeviceProps(Device, Info))
diff --git a/Engine/Source/Runtime/Windows/AudioMixerXAudio2/Private/ScopedCom.h b/Engine/Source/Runtime/Windows/AudioMixerXAudio2/Private/ScopedCom.h
--- a/Engine/Source/Runtime/Windows/AudioMixerXAudio2/Private/ScopedCom.h
+++ b/Engine/Source/Runtime/Windows/AudioMixerXAudio2/Private/ScopedCom.h
@@ -4,96 +4,6 @@
 
 namespace Audio
 {
-	// COM RAII helper
-	template<typename T>
-	struct TScopeComObject final
-	{
-		using ThisType = TScopeComObject<T>;
-
-		T* Obj = nullptr;
-
-		T* Get() const
-		{
-			return static_cast<T*>(Obj);
-		}
-
-		T* operator->() const
-		{
-			check(Get());
-			return Get();
-		}
-
-		explicit operator bool() const
-		{
-			return Get() != nullptr;
-		}
-
-		void Reset()
-		{
-			*this = ThisType();
-		}
-
-		// Returned from QueryInterface call, will already have AddRef called.
-		TScopeComObject(T* InObj = nullptr)
-			: Obj(InObj)
-		{}
-		TScopeComObject(const ThisType& InObj)
-		{
-			*this = InObj;
-		}
-		TScopeComObject(ThisType&& InObj)
-		{
-			*this = MoveTemp(InObj);
-		}
-		~TScopeComObject()
-		{
-			SafeRelease();
-		}
-
-		// Copy Assignment.
-		ThisType& operator=(const ThisType& In)
-		{
-			// Remove any ref we currently hold
-			SafeRelease();
-			// Assign new state and add ref as we are copy
-			Obj = In.Obj;
-			SafeAddRef();
-			return *this;
-		}
-
-		// Move assignment.
-		ThisType& operator=(ThisType&& InObj)
-		{
-			SafeRelease();
-			Obj = InObj.Obj;
-			InObj.Obj = nullptr;
-			return *this;
-		}
-
-		// Equality operator
-		bool operator==(const ThisType& InObj) const
-		{
-			return InObj.Obj == Obj;
-		}
-
-	private:
-		void SafeAddRef()
-		{
-			if (Obj)
-			{
-				Obj->AddRef();
-			}
-		}
-		void SafeRelease()
-		{
-			if (Obj)
-			{
-				Obj->Release();
-				Obj = nullptr;
-			}
-		}
-	};
-
 	struct FScopeComString final
 	{
 		LPTSTR StringPtr = nullptr;
