New logging around device swap to track down device swap problems.

diff --git a/Engine/Source/Runtime/AudioMixer/Private/AudioMixerBlueprintLibrary.cpp b/Engine/Source/Runtime/AudioMixer/Private/AudioMixerBlueprintLibrary.cpp
--- a/Engine/Source/Runtime/AudioMixer/Private/AudioMixerBlueprintLibrary.cpp
+++ b/Engine/Source/Runtime/AudioMixer/Private/AudioMixerBlueprintLibrary.cpp
@@ -807,7 +807,7 @@ void UAudioMixerBlueprintLibrary::SwapAudioOutputDevice(const UObject* WorldCont
 		//Send message to swap device
 		if (MixerPlatform)
 		{
-			bool result = MixerPlatform->RequestDeviceSwap(NewDeviceId);
+			bool result = MixerPlatform->RequestDeviceSwap(NewDeviceId, /*force*/ false, TEXT("UAudioMixerBlueprintLibrary::SwapAudioOutputDevice"));
 			FAudioOutputDeviceInfo CurrentOutputDevice = MixerPlatform->GetPlatformDeviceInfo();
 
 			//Send data through delegate on game thread
diff --git a/Engine/Source/Runtime/AudioMixerCore/Private/AudioMixer.cpp b/Engine/Source/Runtime/AudioMixerCore/Private/AudioMixer.cpp
--- a/Engine/Source/Runtime/AudioMixerCore/Private/AudioMixer.cpp
+++ b/Engine/Source/Runtime/AudioMixerCore/Private/AudioMixer.cpp
@@ -307,6 +307,9 @@ namespace Audio
 	 * IAudioMixerPlatformInterface
 	 */
 
+	// Static linkage.
+	FThreadSafeCounter IAudioMixerPlatformInterface::NextInstanceID;
+
 	IAudioMixerPlatformInterface::IAudioMixerPlatformInterface()
 		: bWarnedBufferUnderrun(false)
 		, AudioRenderEvent(nullptr)
@@ -321,6 +324,7 @@ namespace Audio
 		, bMoveAudioStreamToNewAudioDevice(false)
 		, bIsUsingNullDevice(false)
 		, bIsGeneratingAudio(false)
+		, InstanceID(NextInstanceID.Increment())
 		, NullDeviceCallback(nullptr)
 	{
 		FadeParam.SetValue(0.0f);
@@ -449,7 +453,7 @@ namespace Audio
 
 	void IAudioMixerPlatformInterface::StartRunningNullDevice()
 	{
-		UE_LOG(LogAudioMixer, Verbose, TEXT("StartRunningNullDevice() called"));
+		UE_LOG(LogAudioMixer, Verbose, TEXT("StartRunningNullDevice() called, InstanceID=%d"), InstanceID);
 		SCOPED_NAMED_EVENT(FMixerPlatformXAudio2_StartRunningNullDevice, FColor::Blue);
 		
 		auto ThrowAwayBuffer = [this]() { this->ReadNextBuffer(); };
@@ -474,7 +478,7 @@ namespace Audio
 
 	void IAudioMixerPlatformInterface::StopRunningNullDevice()
 	{
-		UE_LOG(LogAudioMixer, Verbose, TEXT("StopRunningNullDevice() called"));
+		UE_LOG(LogAudioMixer, Verbose, TEXT("StopRunningNullDevice() called, InstanceID=%d"), InstanceID);
 		SCOPED_NAMED_EVENT(FMixerPlatformXAudio2_StopRunningNullDevice, FColor::Blue);
 
 		if (NullDeviceCallback.IsValid())
@@ -563,7 +567,7 @@ namespace Audio
 					// Underrun/Starvation:
 					// Things to try: Increase # output buffers, ensure audio-render thread has time to run (affinity and priority), debug your mix and reduce # sounds playing.
 
-					UE_LOG(LogAudioMixer, Display, TEXT("Audio Buffer Underrun (starvation) detected."));
+					UE_LOG(LogAudioMixer, Display, TEXT("Audio Buffer Underrun (starvation) detected. InstanceID=%d"), InstanceID);
 					bWarnedBufferUnderrun = true;
 					TimeLastWarningCycles = FPlatformTime::Cycles64();
 				}
@@ -574,7 +578,7 @@ namespace Audio
 			// As soon as a valid buffer goes through, allow more warning
 			if (bWarnedBufferUnderrun)
 			{
-				UE_LOG(LogAudioMixerDebug, Log, TEXT("Audio had %d underruns [Total: %d]."), CurrentUnderrunCount, UnderrunCount);
+				UE_LOG(LogAudioMixerDebug, Log, TEXT("Audio had %d underruns [Total: %d], InstanceID=%d"), CurrentUnderrunCount, UnderrunCount, InstanceID);
 			}
 
 			CurrentUnderrunCount = 0;
@@ -608,7 +612,7 @@ namespace Audio
 
 		// Set the number of buffers to be one more than the number to queue.
 		NumOutputBuffers = FMath::Max(OpenStreamParams.NumBuffers, 2);
-		UE_LOG(LogAudioMixer, Display, TEXT("Output buffers initialized: Frames=%i, Channels=%i, Samples=%i"), NumOutputFrames, NumOutputChannels, NumOutputSamples);
+		UE_LOG(LogAudioMixer, Display, TEXT("Output buffers initialized: Frames=%i, Channels=%i, Samples=%i, InstanceID=%d"), NumOutputFrames, NumOutputChannels, NumOutputSamples, InstanceID);
 
 
 		OutputBuffer.Init(AudioStreamInfo.AudioMixer, NumOutputSamples, NumOutputBuffers, AudioStreamInfo.DeviceInfo.Format);
@@ -702,7 +706,7 @@ namespace Audio
 
 	uint32 IAudioMixerPlatformInterface::RunInternal()
 	{
-		UE_LOG(LogAudioMixer, Display, TEXT("Starting AudioMixerPlatformInterface::RunInternal()"));
+		UE_LOG(LogAudioMixer, Display, TEXT("Starting AudioMixerPlatformInterface::RunInternal(), InstanceID=%d"), InstanceID);
 
 		// Lets prime and submit the first buffer (which is going to be the buffer underrun buffer)
 		int32 NumSamplesPopped;
@@ -723,13 +727,15 @@ namespace Audio
 			OverrunTimeoutCVar = FMath::Clamp(OverrunTimeoutCVar, 500, 5000);
 
 			// Now wait for a buffer to be consumed, which will bump up the read index.
+			double WaitStartTime = FPlatformTime::Seconds();
 			if (AudioRenderEvent && !AudioRenderEvent->Wait(static_cast<uint32>(OverrunTimeoutCVar)))
 			{
 				// if we reached this block, we timed out, and should attempt to
 				// bail on our current device.
-				bMoveAudioStreamToNewAudioDevice = true;
+				RequestDeviceSwap(TEXT(""), /* force */true, TEXT("AudioMixerPlatformInterface. Timeout waiting for h/w."));
 
-				UE_LOG(LogAudioMixer, Warning, TEXT("AudioMixerPlatformInterface. Timeout waiting for h/w. (Trying new device)."));
+				float TimeWaited = FPlatformTime::Seconds() - WaitStartTime;
+				UE_LOG(LogAudioMixer, Warning, TEXT("AudioMixerPlatformInterface Timeout [%dms] waiting for h/w. InstanceID=%d"), (int32)(TimeWaited * 1000.f),InstanceID);
 			}
 		}
 
diff --git a/Engine/Source/Runtime/AudioMixerCore/Private/AudioMixerNullDevice.cpp b/Engine/Source/Runtime/AudioMixerCore/Private/AudioMixerNullDevice.cpp
--- a/Engine/Source/Runtime/AudioMixerCore/Private/AudioMixerNullDevice.cpp
+++ b/Engine/Source/Runtime/AudioMixerCore/Private/AudioMixerNullDevice.cpp
@@ -33,7 +33,7 @@ namespace Audio
 				WakeupEvent->Wait(MAX_uint32);
 				WakeupEvent->Reset();
 				
-				UE_CLOG(!bShouldShutdown && !bShouldRecyle, LogAudioMixer, Display, TEXT("FMixerNullCallback: Simulating a h/w device callback at [%dms]"), (int32)(CallbackTime * 1000.f));
+				UE_CLOG(!bShouldShutdown && !bShouldRecyle, LogAudioMixer, Display, TEXT("FMixerNullCallback: Simulating a h/w device callback at [%dms], ThreadID=%u"),  (int32)(CallbackTime * 1000.f), CallbackThread->GetThreadID() );
 
 				// Reset our time differential.
 				AudioClock = FPlatformTime::Seconds();
@@ -64,12 +64,11 @@ namespace Audio
 
 #if !NO_LOGGING
 				// Warn if there's any crazy deltas (limit to every 30s).
-				static double LastLog = 0.f;
 				if (RealClock - LastLog > 30.f)
 				{
 					if (FMath::Abs(SleepTime) > 0.2f)
 					{
-						UE_LOG(LogAudioMixer, Warning, TEXT("FMixerNullCallback: Large time delta between simulated audio clock and realtime [%dms]"), (int32)(SleepTime * 1000.f));
+						UE_LOG(LogAudioMixer, Warning, TEXT("FMixerNullCallback: Large time delta between simulated audio clock and realtime [%dms], ThreadID=%u"), (int32)(SleepTime * 1000.f), CallbackThread->GetThreadID());
 						LastLog = RealClock;
 					}
 				}
diff --git a/Engine/Source/Runtime/AudioMixerCore/Public/AudioMixer.h b/Engine/Source/Runtime/AudioMixerCore/Public/AudioMixer.h
--- a/Engine/Source/Runtime/AudioMixerCore/Public/AudioMixer.h
+++ b/Engine/Source/Runtime/AudioMixerCore/Public/AudioMixer.h
@@ -475,7 +475,7 @@ namespace Audio
 		virtual bool MoveAudioStreamToNewAudioDevice(const FString& InNewDeviceId) { return true;  }
 
 		/** Sends a command to swap which output device is being used */
-		virtual bool RequestDeviceSwap(const FString& DeviceID) { return false; }
+		virtual bool RequestDeviceSwap(const FString& DeviceID, bool bInForce, const TCHAR* InReason = nullptr) { return false; }
 
 		/** Returns the platform device info of the currently open audio stream. */
 		virtual FAudioPlatformDeviceInfo GetPlatformDeviceInfo() const = 0;
@@ -707,6 +707,12 @@ namespace Audio
 		FThreadSafeBool bIsUsingNullDevice;
 		FThreadSafeBool bIsGeneratingAudio;
 
+		/** A Counter to provide the next unique id. */
+		static FThreadSafeCounter NextInstanceID;
+
+		/** A Unique ID Identifying this instance. Mostly used for logging. */ 
+		const int32 InstanceID{ -1 };
+
 	private:
 		TUniquePtr<FMixerNullCallback> NullDeviceCallback;
 	};
diff --git a/Engine/Source/Runtime/AudioMixerCore/Public/AudioMixerNullDevice.h b/Engine/Source/Runtime/AudioMixerCore/Public/AudioMixerNullDevice.h
--- a/Engine/Source/Runtime/AudioMixerCore/Public/AudioMixerNullDevice.h
+++ b/Engine/Source/Runtime/AudioMixerCore/Public/AudioMixerNullDevice.h
@@ -58,6 +58,7 @@ namespace Audio
 		FEvent* SleepEvent = nullptr;
 		FEvent* WakeupEvent = nullptr;	
 		TUniquePtr<FRunnableThread> CallbackThread;
+		double LastLog = 0.f;
 	};
 }
 
diff --git a/Engine/Source/Runtime/Windows/AudioMixerXAudio2/Private/AudioMixerPlatformWindows.cpp b/Engine/Source/Runtime/Windows/AudioMixerXAudio2/Private/AudioMixerPlatformWindows.cpp
--- a/Engine/Source/Runtime/Windows/AudioMixerXAudio2/Private/AudioMixerPlatformWindows.cpp
+++ b/Engine/Source/Runtime/Windows/AudioMixerXAudio2/Private/AudioMixerPlatformWindows.cpp
@@ -82,20 +82,11 @@ public:
 		{
 			return true;
 		}
+		
+		UnregisterForSessionNotifications();
 
 		DeviceListeningToSessionEvents = InDevice;
-		bHasDisconnectSessionHappened = false;
-				
-		// Unregister for any device we're already listening to.
-		if (SessionControls)
-		{
-			SessionControls->UnregisterAudioSessionNotification(this);
-			SessionControls.Reset();
-		}
-		if (SessionManager)
-		{
-			SessionManager.Reset();
-		}
+
 		
 		if(InDevice) 
 		{
@@ -105,6 +96,7 @@ public:
 				{
 					if (SUCCEEDED(SessionControls->RegisterAudioSessionNotification(this)))
 					{
+						UE_LOG(LogAudioMixer, Verbose, TEXT("FWindowsMMNotificationClient: Registering for sessions events for '%s'"), *GetFriendlyName(DeviceListeningToSessionEvents.Get()));
 						return true;
 					}
 				}
@@ -121,13 +113,32 @@ public:
 		}
 		return false;
 	}
-	
-	virtual ~FWindowsMMNotificationClient()
+
+	void UnregisterForSessionNotifications()
 	{
+		FScopeLock ScopeLock(&MutationCs);
+		
+		// Unregister for any device we're already listening to.
 		if (SessionControls)
 		{
+			UE_LOG(LogAudioMixer, Verbose, TEXT("FWindowsMMNotificationClient: Unregistering for sessions events for device '%s'"), DeviceListeningToSessionEvents ? *GetFriendlyName(DeviceListeningToSessionEvents.Get()) : TEXT("None"));
 			SessionControls->UnregisterAudioSessionNotification(this);
+			SessionControls.Reset();
+		}
+		if (SessionManager)
+		{
+			SessionManager.Reset();
 		}
+		
+		DeviceListeningToSessionEvents.Reset();
+
+		// Reset this flag.
+		bHasDisconnectSessionHappened = false;
+	}
+	
+	virtual ~FWindowsMMNotificationClient()
+	{
+		UnregisterForSessionNotifications();
 
 		if (DeviceEnumerator)
 		{
@@ -140,14 +151,63 @@ public:
 		}
 	}
 
+	#define CASE_TO_STRING(X) case X: return TEXT(#X)
+	static const TCHAR* ToString(EDataFlow InFlow)
+	{
+#if !NO_LOGGING
+		switch (InFlow)
+		{
+			CASE_TO_STRING(eRender);
+			CASE_TO_STRING(eCapture);
+			CASE_TO_STRING(eAll);
+			default:
+				checkNoEntry();
+				break;
+		}
+#endif //!NO_LOGGING
+		return TEXT("Unknown");
+	}
+	static const TCHAR* ToString(ERole InRole)
+	{
+#if !NO_LOGGING
+		switch (InRole)
+		{
+			CASE_TO_STRING(eConsole);
+			CASE_TO_STRING(eMultimedia);
+			CASE_TO_STRING(eCommunications);
+			default:
+				checkNoEntry();
+				break;
+		}
+#endif //!NO_LOGGING
+		return TEXT("Unknown");
+	}	
+	static const TCHAR* ToString(AudioSessionDisconnectReason InDisconnectReason)
+	{
+#if !NO_LOGGING
+		switch (InDisconnectReason)
+		{
+			CASE_TO_STRING(DisconnectReasonDeviceRemoval);
+			CASE_TO_STRING(DisconnectReasonServerShutdown);
+			CASE_TO_STRING(DisconnectReasonFormatChanged);
+			CASE_TO_STRING(DisconnectReasonSessionLogoff);
+			CASE_TO_STRING(DisconnectReasonSessionDisconnected);
+			CASE_TO_STRING(DisconnectReasonExclusiveModeOverride);
+		default:
+			checkNoEntry();
+			break;
+		}
+#endif //!NO_LOGGING
+		return TEXT("Unknown");
+	}
+	#undef CASE_TO_STRING
+
 	HRESULT STDMETHODCALLTYPE OnDefaultDeviceChanged(EDataFlow InFlow, ERole InRole, LPCWSTR pwstrDeviceId) override
 	{
 		FScopeLock ScopeLock(&MutationCs);
 
-		if (Audio::IAudioMixer::ShouldLogDeviceSwaps())
-		{
-			UE_LOG(LogAudioMixer, Warning, TEXT("OnDefaultDeviceChanged: %d, %d, %s"), InFlow, InRole, pwstrDeviceId);
-		}
+		UE_CLOG(Audio::IAudioMixer::ShouldLogDeviceSwaps(),LogAudioMixer, Warning, 
+			TEXT("FWindowsMMNotificationClient: OnDefaultDeviceChanged: %s, %s, %s - %s"), ToString(InFlow), ToString(InRole), pwstrDeviceId, *GetFriendlyName(pwstrDeviceId));
 
 		Audio::EAudioDeviceRole AudioDeviceRole;
 
@@ -169,26 +229,27 @@ public:
 			AudioDeviceRole = Audio::EAudioDeviceRole::Communications;
 		}
 
+		FString DeviceString(pwstrDeviceId);
 		if (InFlow == eRender)
 		{
 			for (Audio::IAudioMixerDeviceChangedListener* Listener : Listeners)
 			{
-				Listener->OnDefaultRenderDeviceChanged(AudioDeviceRole, FString(pwstrDeviceId));
+				Listener->OnDefaultRenderDeviceChanged(AudioDeviceRole, DeviceString);
 			}
 		}
 		else if (InFlow == eCapture)
 		{
 			for (Audio::IAudioMixerDeviceChangedListener* Listener : Listeners)
 			{
-				Listener->OnDefaultCaptureDeviceChanged(AudioDeviceRole, FString(pwstrDeviceId));
+				Listener->OnDefaultCaptureDeviceChanged(AudioDeviceRole, DeviceString);
 			}
 		}
 		else
 		{
 			for (Audio::IAudioMixerDeviceChangedListener* Listener : Listeners)
 			{
-				Listener->OnDefaultCaptureDeviceChanged(AudioDeviceRole, FString(pwstrDeviceId));
-				Listener->OnDefaultRenderDeviceChanged(AudioDeviceRole, FString(pwstrDeviceId));
+				Listener->OnDefaultCaptureDeviceChanged(AudioDeviceRole, DeviceString);
+				Listener->OnDefaultRenderDeviceChanged(AudioDeviceRole, DeviceString);
 			}
 		}
 
@@ -219,7 +280,7 @@ public:
 	{
 		FScopeLock ScopeLock(&MutationCs);
 
-		UE_CLOG(Audio::IAudioMixer::ShouldLogDeviceSwaps(), LogAudioMixer, Verbose, TEXT("OnDeviceAdded: %s"), *GetFriendlyName(pwstrDeviceId));
+		UE_CLOG(Audio::IAudioMixer::ShouldLogDeviceSwaps(), LogAudioMixer, Display, TEXT("FWindowsMMNotificationClient: OnDeviceAdded: %s"), *GetFriendlyName(pwstrDeviceId));
 			
 		if (Audio::IAudioMixer::ShouldIgnoreDeviceSwaps())
 		{
@@ -239,7 +300,7 @@ public:
 	{
 		FScopeLock ScopeLock(&MutationCs);
 
-		UE_CLOG(Audio::IAudioMixer::ShouldLogDeviceSwaps(),LogAudioMixer, Verbose, TEXT("OnDeviceRemoved: %s"), *GetFriendlyName(pwstrDeviceId));
+		UE_CLOG(Audio::IAudioMixer::ShouldLogDeviceSwaps(),LogAudioMixer, Display, TEXT("FWindowsMMNotificationClient: OnDeviceRemoved: %s"), *GetFriendlyName(pwstrDeviceId));
 
 		if (Audio::IAudioMixer::ShouldIgnoreDeviceSwaps())
 		{
@@ -258,7 +319,7 @@ public:
 	{
 		FScopeLock ScopeLock(&MutationCs);
 
-		UE_CLOG(Audio::IAudioMixer::ShouldLogDeviceSwaps(), LogAudioMixer, Verbose, TEXT("OnDeviceStateChanged: %s, %d"), *GetFriendlyName(pwstrDeviceId), dwNewState);
+		UE_CLOG(Audio::IAudioMixer::ShouldLogDeviceSwaps(), LogAudioMixer, Display, TEXT("FWindowsMMNotificationClient: OnDeviceStateChanged: %s, %d"), *GetFriendlyName(pwstrDeviceId), dwNewState);
 
 		if (Audio::IAudioMixer::ShouldIgnoreDeviceSwaps())
 		{
@@ -317,10 +378,20 @@ public:
 			
 		// Get device.
 		if (TComPtr<IMMDevice> Device = GetDevice(*InDeviceID))
+		{
+			return GetFriendlyName(Device);
+		}
+		return FriendlyName;
+	}
+	FString GetFriendlyName(const TComPtr<IMMDevice>& InDevice)
+	{
+		FString FriendlyName = TEXT("[No Friendly Name for Device]");
+
+		if (InDevice)
 		{
 			// Get property store.
 			TComPtr<IPropertyStore> PropStore;
-			HRESULT Hr = Device->OpenPropertyStore(STGM_READ, &PropStore);
+			HRESULT Hr = InDevice->OpenPropertyStore(STGM_READ, &PropStore);
 
 			// Get friendly name.
 			if (SUCCEEDED(Hr) && PropStore)
@@ -534,8 +605,9 @@ public:
 
 	HRESULT STDMETHODCALLTYPE OnSessionDisconnected(
 		AudioSessionDisconnectReason InDisconnectReason)
-	{
-		UE_CLOG(Audio::IAudioMixer::ShouldLogDeviceSwaps(), LogAudioMixer, Log, TEXT("Session Disconnect: %s"), GetDisconnectReasonString(InDisconnectReason));
+	{	
+		UE_CLOG(Audio::IAudioMixer::ShouldLogDeviceSwaps(), LogAudioMixer, Verbose, TEXT("Session Disconnect: Reason=%s, DeviceBound=%s, HasDisconnectSessionHappened=%d"), 
+			ToString(InDisconnectReason), *GetFriendlyName(DeviceListeningToSessionEvents), (int32)bHasDisconnectSessionHappened);
 
 		FScopeLock ScopeLock(&MutationCs);
 		if (!bHasDisconnectSessionHappened)
@@ -549,28 +621,10 @@ public:
 			// Mark this true.
 			bHasDisconnectSessionHappened = true;
 		}
+
 		return S_OK;
 	}
 	// End IAudioSessionEvents
-
-	static const TCHAR* GetDisconnectReasonString(AudioSessionDisconnectReason InDisconnectReason)
-	{
-		#define CASE_TO_STRING(X) case X: return TEXT(#X)
-		switch (InDisconnectReason)
-		{
-			CASE_TO_STRING(DisconnectReasonDeviceRemoval);
-			CASE_TO_STRING(DisconnectReasonServerShutdown);
-			CASE_TO_STRING(DisconnectReasonFormatChanged);
-			CASE_TO_STRING(DisconnectReasonSessionLogoff);
-			CASE_TO_STRING(DisconnectReasonSessionDisconnected);
-			CASE_TO_STRING(DisconnectReasonExclusiveModeOverride);
-			default:
-				checkNoEntry();
-				break;
-		}
-		return TEXT("Unknown");
-		#undef CASE_TO_STRING
-	}
 	
 	static Audio::IAudioMixerDeviceChangedListener::EDisconnectReason 
 	AudioSessionDisconnectToEDisconnectReason(AudioSessionDisconnectReason InDisconnectReason)
@@ -614,6 +668,13 @@ namespace Audio
 			WindowsNotificationClient->RegisterForSessionNotifications(InDeviceId);
 		}
 	}
+	void UnregisterForSessionEvents()
+	{
+		if (WindowsNotificationClient)
+		{
+			WindowsNotificationClient->UnregisterForSessionNotifications();
+		}
+	}
 
 	struct FWindowsMMDeviceCache : IAudioMixerDeviceChangedListener, IAudioPlatformDeviceInfoCache
 	{
@@ -745,7 +806,7 @@ namespace Audio
 			// We didn't match channel masks for all channels, revert to a default ordering
 			if (ChanCount < (uint32)OutInfo.NumChannels)
 			{
-				UE_CLOG(Audio::IAudioMixer::ShouldLogDeviceSwaps(), LogAudioMixer, Warning, TEXT("Did not find the channel type flags for audio device '%s'. Reverting to a default channel ordering."), *OutInfo.FriendlyName);
+				UE_CLOG(Audio::IAudioMixer::ShouldLogDeviceSwaps(), LogAudioMixer, Warning, TEXT("FWindowsMMDeviceCache: Did not find the channel type flags for audio device '%s'. Reverting to a default channel ordering."), *OutInfo.FriendlyName);
 
 				OutInfo.OutputChannels.Reset();
 
@@ -894,7 +955,7 @@ namespace Audio
 				else
 				{
 					// Log a warning if this device is active as we failed to ask for a format
-					UE_CLOG(DeviceState == DEVICE_STATE_ACTIVE, LogAudioMixer, Warning, TEXT("Failed to get Format for active device '%s'"), *OutInfo.FriendlyName);
+					UE_CLOG(DeviceState == DEVICE_STATE_ACTIVE, LogAudioMixer, Warning, TEXT("FWindowsMMDeviceCache: Failed to get Format for active device '%s'"), *OutInfo.FriendlyName);
 				}
 			}
 	
@@ -930,7 +991,7 @@ namespace Audio
 								// Enumerate props into our info object.
 								EnumerateDeviceProps(Device, Info);
 
-								UE_LOG(LogAudioMixer, Verbose, TEXT("%s Device '%s' ID='%s'"),
+								UE_LOG(LogAudioMixer, Verbose, TEXT("FWindowsMMDeviceCache: %s Device '%s' ID='%s'"),
 									Info.Type == FCacheEntry::EEndpointType::Capture ? TEXT("Capture") :
 									Info.Type == FCacheEntry::EEndpointType::Render ? TEXT("Render") :
 									TEXT("UNKNOWN!"),
@@ -980,10 +1041,12 @@ namespace Audio
 				FName DeviceIdName;
 				if (GetDefaultDeviceID(eRender, static_cast<ERole>(i), DeviceIdName))
 				{
+					UE_CLOG(!DeviceIdName.IsNone(), LogAudioMixer, Verbose, TEXT("FWindowsMMDeviceCache: Default Render Role='%s', Device='%s'"), ToString((EAudioDeviceRole)i), *GetFriendlyName(DeviceIdName));
 					DefaultRenderId[i] = DeviceIdName;
 				}
 				if (GetDefaultDeviceID(eCapture, static_cast<ERole>(i), DeviceIdName))
 				{
+					UE_CLOG(!DeviceIdName.IsNone(), LogAudioMixer, Verbose, TEXT("FWindowsMMDeviceCache: Default Capture Role='%s', Device='%s'"), ToString((EAudioDeviceRole)i), *GetFriendlyName(DeviceIdName));
 					DefaultCaptureId[i] = DeviceIdName;
 				}
 			}
@@ -1043,11 +1106,46 @@ namespace Audio
 			return {};
 		}
 
+		FString GetFriendlyName(FName InDeviceId) const
+		{
+			if (const FCacheEntry* Entry = Cache.Find(InDeviceId))
+			{
+				return Entry->FriendlyName;
+			}
+			return TEXT("Unknown");
+		}
+
+		static const TCHAR* ToString(EAudioDeviceState InState)
+		{
+			#define CASE_TO_STRING(X) case EAudioDeviceState::X: return TEXT(#X)
+			switch (InState)
+			{
+				CASE_TO_STRING(Active);
+				CASE_TO_STRING(Disabled);
+				CASE_TO_STRING(NotPresent);
+				CASE_TO_STRING(Unplugged);
+			default:
+				return TEXT("Unknown");
+			}
+			#undef CASE_TO_STRING
+		}
+		static const TCHAR* ToString(EAudioDeviceRole InRole)
+		{
+			#define CASE_TO_STRING(X) case EAudioDeviceRole::X: return TEXT(#X)
+			switch (InRole)
+			{
+				CASE_TO_STRING(Console);
+				CASE_TO_STRING(Multimedia);
+				CASE_TO_STRING(Communications);
+			default:
+				return TEXT("Unknown");
+			}
+			#undef CASE_TO_STRING
+		}
+
 		void OnDeviceStateChanged(const FString& DeviceId, const EAudioDeviceState InState, bool ) override
 		{
-			TOptional<FCacheEntry> Info = BuildCacheEntry(DeviceId);
-			
-			UE_CLOG(Audio::IAudioMixer::ShouldLogDeviceSwaps(), LogAudioMixer, Verbose, TEXT("Device '%s' - '%s' state changed."), Info ? *Info->FriendlyName : TEXT("Unknown"), *DeviceId);
+			TOptional<FCacheEntry> Info = BuildCacheEntry(DeviceId);					
 			
 			FReadScopeLock ReadLock(CacheMutationLock);
 			FName DeviceIdName = *DeviceId;
@@ -1056,6 +1154,10 @@ namespace Audio
 			if (FCacheEntry* Entry = Cache.Find(DeviceIdName))
 			{
 				FWriteScopeLock WriteLock(Entry->MutationLock);
+
+				UE_CLOG(Audio::IAudioMixer::ShouldLogDeviceSwaps(), LogAudioMixer, Verbose, TEXT("FWindowsMMDeviceCache: DeviceName='%s' - DeviceID='%s' state changed from '%s' to '%s'."),
+					Info ? *Info->FriendlyName : TEXT("Unknown"), *DeviceId, ToString(Entry->State), ToString(InState));
+
 				Entry->State = InState;
 			}
 		}
@@ -1072,21 +1174,21 @@ namespace Audio
 					if (Found->NumChannels != InFormat.NumChannels)
 					{
 						FWriteScopeLock WriteLock(Found->MutationLock);
-						UE_CLOG(Audio::IAudioMixer::ShouldLogDeviceSwaps(),LogAudioMixer, Verbose, TEXT("Device '%s' changed default format from %d channels to %d."), *InDeviceId, Found->NumChannels, InFormat.NumChannels);
+						UE_CLOG(Audio::IAudioMixer::ShouldLogDeviceSwaps(),LogAudioMixer, Verbose, TEXT("FWindowsMMDeviceCache: Device '%s' changed default format from %d channels to %d."), *InDeviceId, Found->NumChannels, InFormat.NumChannels);
 						Found->NumChannels = InFormat.NumChannels;
 						bDirty = true;
 					}
 					if (Found->SampleRate != InFormat.SampleRate)
 					{
 						FWriteScopeLock WriteLock(Found->MutationLock);
-						UE_CLOG(Audio::IAudioMixer::ShouldLogDeviceSwaps(),LogAudioMixer, Verbose, TEXT("Device '%s' changed default format from %dhz to %dhz."), *InDeviceId, Found->SampleRate, InFormat.SampleRate);
+						UE_CLOG(Audio::IAudioMixer::ShouldLogDeviceSwaps(),LogAudioMixer, Verbose, TEXT("FWindowsMMDeviceCache: Device '%s' changed default format from %dhz to %dhz."), *InDeviceId, Found->SampleRate, InFormat.SampleRate);
 						Found->SampleRate = InFormat.SampleRate;
 						bDirty = true;
 					}
 					if (Found->SpeakerConfig != InFormat.ChannelConfig)
 					{
 						FWriteScopeLock WriteLock(Found->MutationLock);
-						UE_CLOG(Audio::IAudioMixer::ShouldLogDeviceSwaps(),LogAudioMixer, Verbose, TEXT("Device '%s' changed default format from 0x%x to 0x%x"), *InDeviceId, Found->SpeakerConfig, InFormat.ChannelConfig);
+						UE_CLOG(Audio::IAudioMixer::ShouldLogDeviceSwaps(),LogAudioMixer, Verbose, TEXT("FWindowsMMDeviceCache: Device '%s' changed default format from 0x%x to 0x%x"), *InDeviceId, Found->SpeakerConfig, InFormat.ChannelConfig);
 						Found->SpeakerConfig = InFormat.ChannelConfig;
 						bDirty = true;
 					}
@@ -1263,24 +1365,14 @@ namespace Audio
 	}
 
 	void FMixerPlatformXAudio2::OnDefaultRenderDeviceChanged(const EAudioDeviceRole InAudioDeviceRole, const FString& DeviceId)
-	{
-		if (!AllowDeviceSwap())
-		{
-			return;
-		}
-
+	{		
 		// Ignore changes that are for a change in default communication device.
 		if (InAudioDeviceRole != EAudioDeviceRole::Communications)
 		{
-		if (AudioDeviceSwapCriticalSection.TryLock())
-		{
-				UE_LOG(LogAudioMixer, Warning, TEXT("Changing default audio render device to new device: %s."), *WindowsNotificationClient->GetFriendlyName(DeviceId));
+			UE_LOG(LogAudioMixer, Warning, TEXT("FMixerPlatformXAudio2: Changing default audio render device to new device: Role=%s, DeviceName=%s, InstanceID=%d"), 
+				FWindowsMMDeviceCache::ToString(InAudioDeviceRole), *WindowsNotificationClient->GetFriendlyName(DeviceId), InstanceID);
 
-				NewAudioDeviceId = DeviceId;
-			bMoveAudioStreamToNewAudioDevice = true;
-
-			AudioDeviceSwapCriticalSection.Unlock();
-		}
+			RequestDeviceSwap(DeviceId, /* force */true, TEXT("FMixerPlatformXAudio2::OnDefaultRenderDeviceChanged"));
 		}
 
 		if (UAudioDeviceNotificationSubsystem* AudioDeviceNotifSubsystem = UAudioDeviceNotificationSubsystem::Get())
@@ -1303,10 +1395,10 @@ namespace Audio
 			// move our audio stream to this newly added device.
 			if (AudioStreamInfo.DeviceInfo.DeviceId != OriginalAudioDeviceId && DeviceId == OriginalAudioDeviceId)
 			{
-				UE_LOG(LogAudioMixer, Warning, TEXT("Original audio device re-added. Moving audio back to original audio device %s."), *WindowsNotificationClient->GetFriendlyName(*OriginalAudioDeviceId));
+				UE_LOG(LogAudioMixer, Warning, TEXT("FMixerPlatformXAudio2: Original audio device re-added. Moving audio back to original audio device: DeviceName=%s, bRenderDevice=%d, InstanceID=%d"), 
+					*WindowsNotificationClient->GetFriendlyName(*OriginalAudioDeviceId), (int32)bIsRenderDevice, InstanceID);
 
-				NewAudioDeviceId = OriginalAudioDeviceId;
-				bMoveAudioStreamToNewAudioDevice = true;
+				RequestDeviceSwap(OriginalAudioDeviceId, /*force */ true, TEXT("FMixerPlatformXAudio2::OnDeviceAdded"));
 			}
 
 			AudioDeviceSwapCriticalSection.Unlock();
@@ -1331,10 +1423,10 @@ namespace Audio
 			// If the device we're currently using was removed... then switch to the new default audio device.
 			if (AudioStreamInfo.DeviceInfo.DeviceId == DeviceId)
 			{
-				UE_LOG(LogAudioMixer, Warning, TEXT("Audio device removed [%s], falling back to other windows default device."), *WindowsNotificationClient->GetFriendlyName(DeviceId) );
+				UE_LOG(LogAudioMixer, Warning, TEXT("FMixerPlatformXAudio2: Audio device removed [%s], falling back to other windows default device. bIsRenderDevice=%d, InstanceID=%d"), 
+					*WindowsNotificationClient->GetFriendlyName(DeviceId), (int32)bIsRenderDevice, InstanceID);
 
-				NewAudioDeviceId = "";
-				bMoveAudioStreamToNewAudioDevice = true;
+				RequestDeviceSwap(TEXT(""), /* force */ true, TEXT("FMixerPlatformXAudio2::OnDeviceRemoved"));
 			}
 			AudioDeviceSwapCriticalSection.Unlock();
 		}
diff --git a/Engine/Source/Runtime/Windows/AudioMixerXAudio2/Private/AudioMixerPlatformXAudio2.cpp b/Engine/Source/Runtime/Windows/AudioMixerXAudio2/Private/AudioMixerPlatformXAudio2.cpp
--- a/Engine/Source/Runtime/Windows/AudioMixerXAudio2/Private/AudioMixerPlatformXAudio2.cpp
+++ b/Engine/Source/Runtime/Windows/AudioMixerXAudio2/Private/AudioMixerPlatformXAudio2.cpp
@@ -286,6 +286,8 @@ namespace Audio
 			{
 				SCOPED_NAMED_EVENT(FMixerPlatformXAudio2_CheckThreadedDeviceSwap_StartAsyncSwap, FColor::Blue);
 
+				UE_LOG(LogAudioMixer, Verbose, TEXT("FMixerPlatformXAudio2::CheckThreadedDeviceSwap() Starting swap to [%s]"), NewAudioDeviceId.IsEmpty() ? TEXT("[System Default]") : *NewAudioDeviceId);
+
 				// Toggle move to audio stream now we're spawning a new job.
 				bMoveAudioStreamToNewAudioDevice = false;
 
@@ -305,7 +307,7 @@ namespace Audio
 					}
 					else
 					{
-						UE_LOG(LogAudioMixer, Error, TEXT("Ignoring attempt to switch to device with unsupported params: Channels=%u, SampleRate=%u, Id=%s, Name=%s"),
+						UE_LOG(LogAudioMixer, Warning, TEXT("Ignoring attempt to switch to device with unsupported params: Channels=%u, SampleRate=%u, Id=%s, Name=%s"),
 							(uint32)DeviceInfo->NumChannels, (uint32)DeviceInfo->SampleRate, *DeviceInfo->DeviceId, *DeviceInfo->Name);
 
 						NewDevice.Reset();
@@ -328,6 +330,7 @@ namespace Audio
 						bSwitchingToValidDevice ]() mutable->FXAudio2AsyncCreateResult
 					{
 						SCOPED_NAMED_EVENT(FMixerPlatformXAudio2_AsyncDeleteCreate, FColor::Blue);
+						UE_LOG(LogAudioMixer, Verbose, TEXT("FMixerPlatformXAudio2::CheckThreadedDeviceSwap() - AsyncTask Start"));
 
 						// New thread might not have COM setup.
 						FPlatformMisc::CoInitialize();
@@ -527,8 +530,8 @@ namespace Audio
 							AudioStreamInfo.DeviceInfo = ActiveDeviceSwap.Get().DeviceInfo;
 
 							// Display our new XAudio2 Mastering voice details.
-							UE_LOG(LogAudioMixer, Display, TEXT("Successful Swap new Device is (NumChannels=%u, SampleRate=%u, DeviceID=%s, Name=%s)"),
-								(uint32)AudioStreamInfo.DeviceInfo.NumChannels, (uint32)AudioStreamInfo.DeviceInfo.SampleRate, *AudioStreamInfo.DeviceInfo.DeviceId, *AudioStreamInfo.DeviceInfo.Name);
+							UE_LOG(LogAudioMixer, Display, TEXT("Successful Swap new Device is (NumChannels=%u, SampleRate=%u, DeviceID=%s, Name=%s), InstanceID=%d"),
+								(uint32)AudioStreamInfo.DeviceInfo.NumChannels, (uint32)AudioStreamInfo.DeviceInfo.SampleRate, *AudioStreamInfo.DeviceInfo.DeviceId, *AudioStreamInfo.DeviceInfo.Name, InstanceID);
 
 							// Reinitialize the output circular buffer to match the buffer math of the new audio device.
 							const int32 NumOutputSamples = AudioStreamInfo.NumOutputFrames * AudioStreamInfo.DeviceInfo.NumChannels;
@@ -1524,7 +1527,7 @@ namespace Audio
 
 	bool FMixerPlatformXAudio2::StartAudioStream()
 	{
-		UE_LOG(LogAudioMixer, Log, TEXT("FMixerPlatformXAudio2::StartAudioStream() called"));
+		UE_LOG(LogAudioMixer, Log, TEXT("FMixerPlatformXAudio2::StartAudioStream() called. InstanceID=%d"), InstanceID);
 		// Start generating audio with our output source voice
 		BeginGeneratingAudio();
 
@@ -1553,7 +1556,7 @@ namespace Audio
 			return false;
 		}
 
-		UE_LOG(LogAudioMixer, Log, TEXT("FMixerPlatformXAudio2::StopAudioStream() called"));
+		UE_LOG(LogAudioMixer, Log, TEXT("FMixerPlatformXAudio2::StopAudioStream() called. InstanceID=%d"), InstanceID);
 
 		if (AudioStreamInfo.StreamState != EAudioOutputStreamState::Stopped && AudioStreamInfo.StreamState != EAudioOutputStreamState::Closed)
 		{
@@ -1599,17 +1602,26 @@ namespace Audio
 		return false;
 	}
 
-	bool FMixerPlatformXAudio2::RequestDeviceSwap(const FString& DeviceID)
+	bool FMixerPlatformXAudio2::RequestDeviceSwap(const FString& DeviceID, bool bInForce, const TCHAR* InReason)
 	{
-		if (!AllowDeviceSwap())
+		if (!AllowDeviceSwap() && !bInForce)
 		{
+			UE_LOG(LogAudioMixer, Verbose, TEXT("NOT-ALLOWING attempt to swap audio render device to new device: '%s', because: '%s', force=%d, InstanceID=%d"),
+				!DeviceID.IsEmpty() ? *DeviceID : TEXT("[System Default]"),
+				InReason ? InReason : TEXT("None specified"),
+				(int32)bInForce,
+				InstanceID
+			);
 			return false;
 		}
 
 		if (AudioDeviceSwapCriticalSection.TryLock())
 		{
-			UE_LOG(LogAudioMixer, Display, TEXT("Attempt to swap audio render device to new device: %s."),
-				!DeviceID.IsEmpty() ? *DeviceID : TEXT("[System Default]")
+			UE_LOG(LogAudioMixer, Verbose, TEXT("Attempt to swap audio render device to new device: '%s', because: '%s', force=%d, InstanceID=%d"),
+				!DeviceID.IsEmpty() ? *DeviceID : TEXT("[System Default]"),
+				InReason ? InReason : TEXT("None specified"),
+				(int32)bInForce,
+				InstanceID
 			);
 
 			NewAudioDeviceId = DeviceID;
@@ -1620,6 +1632,8 @@ namespace Audio
 			return true;
 		}
 
+		UE_CLOG(bInForce,LogAudioMixer, Warning, TEXT("FMixerPlatformXAudio2::RequestDeviceSwap(), Failed to acquire CS, device swap to '%s', will be ignored. InstanceID=%d"),
+			!DeviceID.IsEmpty() ? *DeviceID : TEXT("[System Default]"), InstanceID );
 		return false;
 	}
 
@@ -1878,7 +1892,7 @@ namespace Audio
 
 			if(!FirstBufferSubmitted)
 			{
-				UE_LOG(LogAudioMixer, Display, TEXT("FMixerPlatformXAudio2::SubmitBuffer() called for the first time"));
+				UE_LOG(LogAudioMixer, Display, TEXT("FMixerPlatformXAudio2::SubmitBuffer() called for the first time. InstanceID=%d"), InstanceID);
 				FirstBufferSubmitted = true;
 			}
 		}
@@ -2073,7 +2087,7 @@ namespace Audio
 		if (InReason == IAudioMixerDeviceChangedListener::EDisconnectReason::FormatChanged)
 		{
 			// OnFormatChanged, retry again same device.
-			RequestDeviceSwap(GetDeviceId());
+			RequestDeviceSwap(GetDeviceId(), /*force*/ true, TEXT("FMixerPlatformXAudio2::OnSessionDisconnect() - FormatChanged"));		
 		}
 		else if (InReason == IAudioMixerDeviceChangedListener::EDisconnectReason::DeviceRemoval)
 		{
@@ -2083,7 +2097,7 @@ namespace Audio
 		{
 			// ServerShutdown, SessionLogoff, SessionDisconnected, ExclusiveModeOverride
 			// Attempt a default swap, will likely fail, but then we'll switch to a null device.
-			RequestDeviceSwap(TEXT(""));
+			RequestDeviceSwap(TEXT(""), /*force*/ true, TEXT("FMixerPlatformXAudio2::OnSessionDisconnect() - Other"));
 		}
 	}
 
diff --git a/Engine/Source/Runtime/Windows/AudioMixerXAudio2/Private/AudioMixerPlatformXAudio2.h b/Engine/Source/Runtime/Windows/AudioMixerXAudio2/Private/AudioMixerPlatformXAudio2.h
--- a/Engine/Source/Runtime/Windows/AudioMixerXAudio2/Private/AudioMixerPlatformXAudio2.h
+++ b/Engine/Source/Runtime/Windows/AudioMixerXAudio2/Private/AudioMixerPlatformXAudio2.h
@@ -77,7 +77,7 @@ namespace Audio
 		virtual bool StartAudioStream() override;
 		virtual bool StopAudioStream() override;
 		virtual bool MoveAudioStreamToNewAudioDevice(const FString& InNewDeviceId) override;
-		virtual bool RequestDeviceSwap(const FString& DeviceID) override;
+		virtual bool RequestDeviceSwap(const FString& DeviceID, bool bInForce, const TCHAR* InReason) override;
 		virtual void ResumePlaybackOnNewDevice() override;
 		virtual FAudioPlatformDeviceInfo GetPlatformDeviceInfo() const override;
 		virtual void SubmitBuffer(const uint8* Buffer) override;
