Log reason of swap in thread device results, also lock around lamba capture.

diff --git a/Engine/Source/Runtime/Windows/AudioMixerXAudio2/Private/AudioMixerPlatformXAudio2.cpp b/Engine/Source/Runtime/Windows/AudioMixerXAudio2/Private/AudioMixerPlatformXAudio2.cpp
--- a/Engine/Source/Runtime/Windows/AudioMixerXAudio2/Private/AudioMixerPlatformXAudio2.cpp
+++ b/Engine/Source/Runtime/Windows/AudioMixerXAudio2/Private/AudioMixerPlatformXAudio2.cpp
@@ -318,149 +318,159 @@ namespace Audio
 					// Failed to find any device/default or otherwise. Reset our current device and do a switch to a null device.
 					NewDevice.Reset();
 				}
-			
-				// Build Delete/Create Async lambda
-				auto AsyncDeleteCreate = [
-					System = XAudio2System,
+				
+				TFunction<FXAudio2AsyncCreateResult()> AsyncDeleteCreate;
+
+				{
+					// Lock for creation of lambda as we are copying state.
+					FScopeLock Lock(&AudioDeviceSwapCriticalSection);
+
+					// Build Delete/Create Async lambda
+					AsyncDeleteCreate = [
+						System = XAudio2System,
 						SourceVoice = OutputAudioStreamSourceVoice,
 						MasterVoice = OutputAudioStreamMasteringVoice,
 						NewDevice,
 						&Callbacks = OutputVoiceCallback,
 						RenderingSampleRate = OpenStreamParams.SampleRate,
-						bSwitchingToValidDevice ]() mutable->FXAudio2AsyncCreateResult
-					{
-						SCOPED_NAMED_EVENT(FMixerPlatformXAudio2_AsyncDeleteCreate, FColor::Blue);
-						UE_LOG(LogAudioMixer, Verbose, TEXT("FMixerPlatformXAudio2::CheckThreadedDeviceSwap() - AsyncTask Start"));
+						bSwitchingToValidDevice,
+						SwapReason = DeviceSwapReason
+					]() mutable->FXAudio2AsyncCreateResult
+						{
+							SCOPED_NAMED_EVENT(FMixerPlatformXAudio2_AsyncDeleteCreate, FColor::Blue);
+							UE_LOG(LogAudioMixer, Verbose, TEXT("FMixerPlatformXAudio2::CheckThreadedDeviceSwap() - AsyncTask Start. Because=%s"), *SwapReason);
 
-						// New thread might not have COM setup.
-						FPlatformMisc::CoInitialize();
+							// New thread might not have COM setup.
+							FPlatformMisc::CoInitialize();
 
-						{
-							// Stop old engine running.
-							if (System)
 							{
-								SCOPED_NAMED_EVENT(FMixerPlatformXAudio2_AsyncDeleteCreate_StopEngine, FColor::Blue);
-								System->StopEngine();
-							}
+								// Stop old engine running.
+								if (System)
+								{
+									SCOPED_NAMED_EVENT(FMixerPlatformXAudio2_AsyncDeleteCreate_StopEngine, FColor::Blue);
+									System->StopEngine();
+								}
 
-							// Kill source voice.
-							if (SourceVoice)
-							{
+								// Kill source voice.
+								if (SourceVoice)
 								{
-									SCOPED_NAMED_EVENT(FMixerPlatformXAudio2_AsyncDeleteCreate_FlushSourceBuffers, FColor::Blue);
-									HRESULT Result = SourceVoice->FlushSourceBuffers();
-									if (FAILED(Result))
 									{
-										XAUDIO2_LOG_RESULT("SourceVoice->FlushSourceBuffers", Result);
+										SCOPED_NAMED_EVENT(FMixerPlatformXAudio2_AsyncDeleteCreate_FlushSourceBuffers, FColor::Blue);
+										HRESULT Result = SourceVoice->FlushSourceBuffers();
+										if (FAILED(Result))
+										{
+											XAUDIO2_LOG_RESULT("SourceVoice->FlushSourceBuffers", Result);
+										}
 									}
+									SCOPED_NAMED_EVENT(FMixerPlatformXAudio2_AsyncDeleteCreate_DestroySourceVoice, FColor::Blue);
+									SourceVoice->DestroyVoice();
+									SourceVoice = nullptr;
 								}
-								SCOPED_NAMED_EVENT(FMixerPlatformXAudio2_AsyncDeleteCreate_DestroySourceVoice, FColor::Blue);
-								SourceVoice->DestroyVoice();
-								SourceVoice = nullptr;
-							}
 
-							// Now destroy the mastering voice
-							if (MasterVoice)
-							{
-								SCOPED_NAMED_EVENT(FMixerPlatformXAudio2_AsyncDeleteCreate_DestroyMasterVoice, FColor::Blue);
-								MasterVoice->DestroyVoice();
-								MasterVoice = nullptr;
-							}
+								// Now destroy the mastering voice
+								if (MasterVoice)
+								{
+									SCOPED_NAMED_EVENT(FMixerPlatformXAudio2_AsyncDeleteCreate_DestroyMasterVoice, FColor::Blue);
+									MasterVoice->DestroyVoice();
+									MasterVoice = nullptr;
+								}
 
-							// Destroy System
-							{
-								SCOPED_NAMED_EVENT(FMixerPlatformXAudio2_AsyncDeleteCreate_DestroySystem, FColor::Blue);
-								SAFE_RELEASE(System);
+								// Destroy System
+								{
+									SCOPED_NAMED_EVENT(FMixerPlatformXAudio2_AsyncDeleteCreate_DestroySystem, FColor::Blue);
+									SAFE_RELEASE(System);
+								}
 							}
-						}
 
-						FXAudio2AsyncCreateResult Results;
+							FXAudio2AsyncCreateResult Results;
+							Results.SwapReason = SwapReason;
 
-						// Don't attempt to create a new setup if there's no devices available.
-						if (!bSwitchingToValidDevice)
-						{
-							return {};
-						}
+							// Don't attempt to create a new setup if there's no devices available.
+							if (!bSwitchingToValidDevice)
+							{
+								return {};
+							}
 
-						// Create System.
-						{
-							SCOPED_NAMED_EVENT(FMixerPlatformXAudio2_AsyncDeleteCreate_CreateSystem, FColor::Blue);
-							HRESULT Result = XAudio2Create(&Results.XAudio2System, 0, GetXAudio2ProcessorsToUse());
-							if (FAILED(Result))
+							// Create System.
 							{
-								XAUDIO2_LOG_RESULT("XAudio2Create", Result);
-								return {};// FAIL.
+								SCOPED_NAMED_EVENT(FMixerPlatformXAudio2_AsyncDeleteCreate_CreateSystem, FColor::Blue);
+								HRESULT Result = XAudio2Create(&Results.XAudio2System, 0, GetXAudio2ProcessorsToUse());
+								if (FAILED(Result))
+								{
+									XAUDIO2_LOG_RESULT("XAudio2Create", Result);
+									return {};// FAIL.
+								}
 							}
-						}
 
-						// Create Master
-						{
-							check(NewDevice.NumChannels <= XAUDIO2_MAX_AUDIO_CHANNELS);
-							check(NewDevice.SampleRate >= XAUDIO2_MIN_SAMPLE_RATE);
-							check(NewDevice.SampleRate <= XAUDIO2_MAX_SAMPLE_RATE);					
-
-							SCOPED_NAMED_EVENT(FMixerPlatformXAudio2_AsyncDeleteCreate_CreateMasterVoice, FColor::Blue);
-							HRESULT Result = Results.XAudio2System->CreateMasteringVoice(
-								&Results.OutputAudioStreamMasteringVoice,
-								(uint32)NewDevice.NumChannels,
-								(uint32)NewDevice.SampleRate,
-								0,
-								*NewDevice.DeviceId,
-								nullptr,
-								AudioCategory_GameEffects
-							);
-							if (FAILED(Result))
+							// Create Master
 							{
-								UE_LOG(LogAudioMixer, Error, TEXT("XAudio2System->CreateMasteringVoice() -> 0x%X: %s (line: %d) with Args (NumChannels=%u, SampleRate=%u, DeviceID=%s Name=%s)"),
-									Result, *GetErrorString(Result), __LINE__, (uint32)NewDevice.NumChannels, (uint32)NewDevice.SampleRate, *NewDevice.DeviceId, *NewDevice.Name);
-								SAFE_RELEASE(Results.XAudio2System);
-								return {}; // FAIL.
+								check(NewDevice.NumChannels <= XAUDIO2_MAX_AUDIO_CHANNELS);
+								check(NewDevice.SampleRate >= XAUDIO2_MIN_SAMPLE_RATE);
+								check(NewDevice.SampleRate <= XAUDIO2_MAX_SAMPLE_RATE);
+
+								SCOPED_NAMED_EVENT(FMixerPlatformXAudio2_AsyncDeleteCreate_CreateMasterVoice, FColor::Blue);
+								HRESULT Result = Results.XAudio2System->CreateMasteringVoice(
+									&Results.OutputAudioStreamMasteringVoice,
+									(uint32)NewDevice.NumChannels,
+									(uint32)NewDevice.SampleRate,
+									0,
+									*NewDevice.DeviceId,
+									nullptr,
+									AudioCategory_GameEffects
+								);
+								if (FAILED(Result))
+								{
+									UE_LOG(LogAudioMixer, Error, TEXT("XAudio2System->CreateMasteringVoice() -> 0x%X: %s (line: %d) with Args (NumChannels=%u, SampleRate=%u, DeviceID=%s Name=%s)"),
+										Result, *GetErrorString(Result), __LINE__, (uint32)NewDevice.NumChannels, (uint32)NewDevice.SampleRate, *NewDevice.DeviceId, *NewDevice.Name);
+									SAFE_RELEASE(Results.XAudio2System);
+									return {}; // FAIL.
+								}
 							}
-						}
 
-						// Create Source Voice.
-						{
-							SCOPED_NAMED_EVENT(FMixerPlatformXAudio2_AsyncDeleteCreate_CreateSourceVoice, FColor::Blue);
-
-							// Setup the format of the output source voice
-							WAVEFORMATEX Format = { 0 };
-							Format.nChannels = NewDevice.NumChannels;
-							Format.nSamplesPerSec = RenderingSampleRate;		// NOTE: We use the Rendering sample rate here.
-							Format.wFormatTag = WAVE_FORMAT_IEEE_FLOAT;
-							Format.nAvgBytesPerSec = Format.nSamplesPerSec * sizeof(float) * Format.nChannels;
-							Format.nBlockAlign = sizeof(float) * Format.nChannels;
-							Format.wBitsPerSample = sizeof(float) * 8;
-						
-							// Create the output source voice
-							HRESULT Result = Results.XAudio2System->CreateSourceVoice(
-								&Results.OutputAudioStreamSourceVoice,
-								&Format,
-								XAUDIO2_VOICE_NOPITCH,
-								XAUDIO2_DEFAULT_FREQ_RATIO,
-								&Callbacks,
-								nullptr,
-								nullptr
-							);
-
-							if (FAILED(Result))
+							// Create Source Voice.
 							{
-								XAUDIO2_LOG_RESULT("XAudio2System->CreateSourceVoice", Result);
-								Results.OutputAudioStreamMasteringVoice->DestroyVoice();
-								Results.OutputAudioStreamMasteringVoice = nullptr;
-								SAFE_RELEASE(Results.XAudio2System);
-								return {};
+								SCOPED_NAMED_EVENT(FMixerPlatformXAudio2_AsyncDeleteCreate_CreateSourceVoice, FColor::Blue);
+
+								// Setup the format of the output source voice
+								WAVEFORMATEX Format = { 0 };
+								Format.nChannels = NewDevice.NumChannels;
+								Format.nSamplesPerSec = RenderingSampleRate;		// NOTE: We use the Rendering sample rate here.
+								Format.wFormatTag = WAVE_FORMAT_IEEE_FLOAT;
+								Format.nAvgBytesPerSec = Format.nSamplesPerSec * sizeof(float) * Format.nChannels;
+								Format.nBlockAlign = sizeof(float) * Format.nChannels;
+								Format.wBitsPerSample = sizeof(float) * 8;
+
+								// Create the output source voice
+								HRESULT Result = Results.XAudio2System->CreateSourceVoice(
+									&Results.OutputAudioStreamSourceVoice,
+									&Format,
+									XAUDIO2_VOICE_NOPITCH,
+									XAUDIO2_DEFAULT_FREQ_RATIO,
+									&Callbacks,
+									nullptr,
+									nullptr
+								);
+
+								if (FAILED(Result))
+								{
+									XAUDIO2_LOG_RESULT("XAudio2System->CreateSourceVoice", Result);
+									Results.OutputAudioStreamMasteringVoice->DestroyVoice();
+									Results.OutputAudioStreamMasteringVoice = nullptr;
+									SAFE_RELEASE(Results.XAudio2System);
+									return {};
+								}
 							}
-						}
 
-						// Listen session for changes to this device.
-						RegisterForSessionEvents(NewDevice.DeviceId);
+							// Listen session for changes to this device.
+							RegisterForSessionEvents(NewDevice.DeviceId);
 
-						FPlatformMisc::CoUninitialize();
+							FPlatformMisc::CoUninitialize();
 
-						// Success.
-						Results.DeviceInfo = NewDevice;
-						return Results;
-					};
+							// Success.
+							Results.DeviceInfo = NewDevice;
+							return Results;
+						};
+					}
 
 					// NULL our system / master / source voices. (as they are about to be torn down async).
 					{
@@ -530,8 +540,8 @@ namespace Audio
 							AudioStreamInfo.DeviceInfo = ActiveDeviceSwap.Get().DeviceInfo;
 
 							// Display our new XAudio2 Mastering voice details.
-							UE_LOG(LogAudioMixer, Display, TEXT("Successful Swap new Device is (NumChannels=%u, SampleRate=%u, DeviceID=%s, Name=%s), InstanceID=%d"),
-								(uint32)AudioStreamInfo.DeviceInfo.NumChannels, (uint32)AudioStreamInfo.DeviceInfo.SampleRate, *AudioStreamInfo.DeviceInfo.DeviceId, *AudioStreamInfo.DeviceInfo.Name, InstanceID);
+							UE_LOG(LogAudioMixer, Display, TEXT("Successful Swap new Device is (NumChannels=%u, SampleRate=%u, DeviceID=%s, Name=%s), Reason=%s, InstanceID=%d"),
+								(uint32)AudioStreamInfo.DeviceInfo.NumChannels, (uint32)AudioStreamInfo.DeviceInfo.SampleRate, *AudioStreamInfo.DeviceInfo.DeviceId, *AudioStreamInfo.DeviceInfo.Name, *ActiveDeviceSwap.Get().SwapReason, InstanceID);
 
 							// Reinitialize the output circular buffer to match the buffer math of the new audio device.
 							const int32 NumOutputSamples = AudioStreamInfo.NumOutputFrames * AudioStreamInfo.DeviceInfo.NumChannels;
@@ -1626,6 +1636,7 @@ namespace Audio
 
 			NewAudioDeviceId = DeviceID;
 			bMoveAudioStreamToNewAudioDevice = true;
+			DeviceSwapReason = InReason;
 			
 			AudioDeviceSwapCriticalSection.Unlock();
 
diff --git a/Engine/Source/Runtime/Windows/AudioMixerXAudio2/Private/AudioMixerPlatformXAudio2.h b/Engine/Source/Runtime/Windows/AudioMixerXAudio2/Private/AudioMixerPlatformXAudio2.h
--- a/Engine/Source/Runtime/Windows/AudioMixerXAudio2/Private/AudioMixerPlatformXAudio2.h
+++ b/Engine/Source/Runtime/Windows/AudioMixerXAudio2/Private/AudioMixerPlatformXAudio2.h
@@ -113,6 +113,7 @@ namespace Audio
 			IXAudio2MasteringVoice* OutputAudioStreamMasteringVoice = nullptr;
 			IXAudio2SourceVoice* OutputAudioStreamSourceVoice = nullptr;
 			FAudioPlatformDeviceInfo DeviceInfo;
+			FString SwapReason;
 		};
 		TFuture<FXAudio2AsyncCreateResult> ActiveDeviceSwap;
 		bool CheckThreadedDeviceSwap();
@@ -147,6 +148,7 @@ namespace Audio
 		FString OriginalAudioDeviceId;
 		FString NewAudioDeviceId;
 		double LastDeviceSwapTime;
+		FString DeviceSwapReason;
 
 		// When we are running the null device,
 		// we check whether a new audio device was connected every second or so.
