Fix for double swap on default render device change.

diff --git a/Engine/Source/Runtime/Windows/AudioMixerXAudio2/Private/AudioMixerPlatformWindows.cpp b/Engine/Source/Runtime/Windows/AudioMixerXAudio2/Private/AudioMixerPlatformWindows.cpp
--- a/Engine/Source/Runtime/Windows/AudioMixerXAudio2/Private/AudioMixerPlatformWindows.cpp
+++ b/Engine/Source/Runtime/Windows/AudioMixerXAudio2/Private/AudioMixerPlatformWindows.cpp
@@ -66,7 +66,7 @@ public:
 		if (DeviceEnumerator)
 		{
 			TComPtr<IMMDevice> DefaultDevice;
-			if (SUCCEEDED(DeviceEnumerator->GetDefaultAudioEndpoint(eRender, eMultimedia, &DefaultDevice)))
+			if (SUCCEEDED(DeviceEnumerator->GetDefaultAudioEndpoint(eRender, eConsole, &DefaultDevice)))
 			{
 				RegisterForSessionNotifications(DefaultDevice);
 			}
@@ -1271,14 +1271,14 @@ namespace Audio
 		FName GetDefaultOutputDevice() const
 		{
 			FReadScopeLock MapReadLock(CacheMutationLock);
-			if( !DefaultRenderId[(int32)EAudioDeviceRole::Multimedia].IsNone())
-			{
-				return DefaultRenderId[(int32)EAudioDeviceRole::Multimedia];
-			}
 			if (!DefaultRenderId[(int32)EAudioDeviceRole::Console].IsNone())
 			{
 				return DefaultRenderId[(int32)EAudioDeviceRole::Console];
 			}
+			if( !DefaultRenderId[(int32)EAudioDeviceRole::Multimedia].IsNone())
+			{
+				return DefaultRenderId[(int32)EAudioDeviceRole::Multimedia];
+			}
 			return NAME_None;
 		}
 		
@@ -1366,8 +1366,12 @@ namespace Audio
 
 	void FMixerPlatformXAudio2::OnDefaultRenderDeviceChanged(const EAudioDeviceRole InAudioDeviceRole, const FString& DeviceId)
 	{		
-		// Ignore changes that are for a change in default communication device.
-		if (InAudioDeviceRole != EAudioDeviceRole::Communications)
+		// There's 3 defaults in windows (communications, console, multimedia). These technically can all be different devices.		
+		// However the Windows UX only allows console+multimedia to be toggle as a pair. This means you get two notifications
+		// for default device changing typically. To prevent a trouble trigger we only listen to "Console" here. For more information on 
+		// device roles: https://docs.microsoft.com/en-us/windows/win32/coreaudio/device-roles
+		
+		if (InAudioDeviceRole == EAudioDeviceRole::Console)
 		{
 			UE_LOG(LogAudioMixer, Warning, TEXT("FMixerPlatformXAudio2: Changing default audio render device to new device: Role=%s, DeviceName=%s, InstanceID=%d"), 
 				FWindowsMMDeviceCache::ToString(InAudioDeviceRole), *WindowsNotificationClient->GetFriendlyName(DeviceId), InstanceID);
