Fixing for channel weirdness around device swaps. Remove PhysicalSpeaker property from being monitored for swapping and improve logging around channels changing.

diff --git a/Engine/Source/Runtime/AudioMixerCore/Public/AudioMixer.h b/Engine/Source/Runtime/AudioMixerCore/Public/AudioMixer.h
--- a/Engine/Source/Runtime/AudioMixerCore/Public/AudioMixer.h
+++ b/Engine/Source/Runtime/AudioMixerCore/Public/AudioMixer.h
@@ -378,7 +378,7 @@ namespace Audio
 		{
 			int32 NumChannels = 0;
 			int32 SampleRate = 0;
-			uint32 ChannelConfig = 0;
+			uint32 ChannelBitmask = 0;
 		};
 
 		enum class EDisconnectReason
@@ -399,7 +399,6 @@ namespace Audio
 		virtual void OnDeviceRemoved(const FString& DeviceId, bool bIsRenderDevice) {}
 		virtual void OnDeviceStateChanged(const FString& DeviceId, const EAudioDeviceState InState, bool bIsRenderDevice) {}
 		virtual void OnFormatChanged(const FString& InDeviceId, const FFormatChangedData& InFormat) {}
-		virtual void OnSpeakerConfigChanged(const FString& InDeviceId, uint32 InSpeakerBitmask) {}		
 		virtual void OnSessionDisconnect(EDisconnectReason InReason) {}
 		
 		virtual FString GetDeviceId() const { return FString(); }
diff --git a/Engine/Source/Runtime/Windows/AudioMixerXAudio2/Private/AudioMixerPlatformWindows.cpp b/Engine/Source/Runtime/Windows/AudioMixerXAudio2/Private/AudioMixerPlatformWindows.cpp
--- a/Engine/Source/Runtime/Windows/AudioMixerXAudio2/Private/AudioMixerPlatformWindows.cpp
+++ b/Engine/Source/Runtime/Windows/AudioMixerXAudio2/Private/AudioMixerPlatformWindows.cpp
@@ -434,9 +434,7 @@ public:
 	{
 		UE_CLOG(Audio::IAudioMixer::ShouldLogDeviceSwaps(), LogAudioMixer, Verbose, TEXT("OnPropertyValueChanged: %s : %s"), *GetFriendlyName(pwstrDeviceId), *PropToString(key));
 		
-		if (key.fmtid == PKEY_AudioEngine_DeviceFormat.fmtid ||
-			/*key.fmtid == PKEY_AudioEngine_OEMFormat.fmtid ||*/
-			key.fmtid == PKEY_AudioEndpoint_PhysicalSpeakers.fmtid )
+		if (key.fmtid == PKEY_AudioEngine_DeviceFormat.fmtid )
 		{
 			// Get device.
 			FString DeviceId = pwstrDeviceId;
@@ -454,9 +452,8 @@ public:
 					PROPVARIANT Prop;
 					PropVariantInit(&Prop);
 
-					if (key.fmtid == PKEY_AudioEngine_DeviceFormat.fmtid ||
-			key.fmtid == PKEY_AudioEngine_OEMFormat.fmtid)
-		{
+					if (key.fmtid == PKEY_AudioEngine_DeviceFormat.fmtid )
+					{
 						// WAVEFORMATEX blobs.
 						if (SUCCEEDED(PropertyStore->GetValue(key, &Prop)) && Prop.blob.pBlobData)
 						{
@@ -465,7 +462,7 @@ public:
 							Audio::IAudioMixerDeviceChangedListener::FFormatChangedData FormatChanged;
 							FormatChanged.NumChannels = FMath::Clamp((int32)WaveFormatEx->nChannels, 2, 8);
 							FormatChanged.SampleRate = WaveFormatEx->nSamplesPerSec;
-							FormatChanged.ChannelConfig = WaveFormatEx->wFormatTag == WAVE_FORMAT_EXTENSIBLE ?
+							FormatChanged.ChannelBitmask = WaveFormatEx->wFormatTag == WAVE_FORMAT_EXTENSIBLE ?
 								((const WAVEFORMATEXTENSIBLE*)WaveFormatEx)->dwChannelMask : 0;
 
 							FScopeLock ScopeLock(&MutationCs);
@@ -474,18 +471,7 @@ public:
 								i->OnFormatChanged(DeviceId, FormatChanged);
 							}
 						}
-					}
-					else if(key.fmtid == PKEY_AudioEndpoint_PhysicalSpeakers.fmtid)
-					{
-						if (SUCCEEDED(PropertyStore->GetValue(PKEY_AudioEndpoint_PhysicalSpeakers, &Prop)))
-						{
-							FScopeLock ScopeLock(&MutationCs);
-							for (Audio::IAudioMixerDeviceChangedListener* i : Listeners)
-			{
-								i->OnSpeakerConfigChanged(DeviceId, Prop.uintVal);
-							}
-						}
-					}
+					}				
 					PropVariantClear(&Prop);
 				}
 			}
@@ -689,7 +675,7 @@ namespace Audio
 			int32 NumChannels = 0;
 			int32 SampleRate = 0;
 			EEndpointType Type = EEndpointType::Unknown;
-			uint32 SpeakerConfig = 0;				// Bitfield used to build output channels, for easy comparison.
+			uint32 ChannelBitmask = 0;				// Bitfield used to build output channels, for easy comparison.
 
 			TArray<EAudioMixerChannel::Type> OutputChannels;	// TODO. Generate this from the ChannelNum and bitmask when we are asked for it.
 			mutable FRWLock MutationLock;
@@ -704,7 +690,7 @@ namespace Audio
 				NumChannels = InOther.NumChannels;
 				SampleRate = InOther.SampleRate;
 				Type = InOther.Type;
-				SpeakerConfig = InOther.SpeakerConfig;
+				ChannelBitmask = InOther.ChannelBitmask;
 				OutputChannels = InOther.OutputChannels;
 				return *this;
 			}
@@ -718,7 +704,7 @@ namespace Audio
 				NumChannels = MoveTemp(InOther.NumChannels);
 				SampleRate = MoveTemp(InOther.SampleRate);
 				Type = MoveTemp(InOther.Type);
-				SpeakerConfig = MoveTemp(InOther.SpeakerConfig);
+				ChannelBitmask = MoveTemp(InOther.ChannelBitmask);
 				OutputChannels = MoveTemp(InOther.OutputChannels);
 				return *this;
 			}
@@ -754,7 +740,7 @@ namespace Audio
 		}
 		virtual ~FWindowsMMDeviceCache() = default;
 
-		bool EnumSpeakerMask(uint32 InMask, FCacheEntry& OutInfo)
+		bool EnumerateChannelMask(uint32 InMask, FCacheEntry& OutInfo)
 		{
 			// Loop through the extensible format channel flags in the standard order and build our output channel array
 			// From https://msdn.microsoft.com/en-us/library/windows/hardware/dn653308(v=vs.85).aspx
@@ -785,7 +771,8 @@ namespace Audio
 				SPEAKER_RESERVED,
 			};
 
-			OutInfo.SpeakerConfig = InMask;
+			OutInfo.ChannelBitmask = InMask;
+			OutInfo.OutputChannels.Reset();
 
 			// No need to enumerate speakers for capture devices.
 			if (OutInfo.Type == FCacheEntry::EEndpointType::Capture)
@@ -858,7 +845,7 @@ namespace Audio
 			return true;
 		}
 
-		bool EnumerateSpeakers(const WAVEFORMATEX* InFormat, FCacheEntry& OutInfo)
+		bool EnumerateChannelFormat(const WAVEFORMATEX* InFormat, FCacheEntry& OutInfo)
 		{
 			OutInfo.OutputChannels.Empty();
 
@@ -867,7 +854,7 @@ namespace Audio
 			{
 				// Cast to the extensible format to get access to extensible data
 				const WAVEFORMATEXTENSIBLE* WaveFormatExtensible = (const WAVEFORMATEXTENSIBLE*)(InFormat);
-				return EnumSpeakerMask(WaveFormatExtensible->dwChannelMask, OutInfo);
+				return EnumerateChannelMask(WaveFormatExtensible->dwChannelMask, OutInfo);
 			}
 			else
 			{
@@ -941,7 +928,7 @@ namespace Audio
 						OutInfo.NumChannels = FMath::Clamp((int32)WaveFormatEx->nChannels, 2, 8);
 						OutInfo.SampleRate = WaveFormatEx->nSamplesPerSec;
 
-						EnumerateSpeakers(WaveFormatEx, OutInfo);
+						EnumerateChannelFormat(WaveFormatEx, OutInfo);
 						PropVariantClear(&DeviceFormat);
 						return true;
 					}
@@ -1142,6 +1129,18 @@ namespace Audio
 			}
 			#undef CASE_TO_STRING
 		}
+		static const FString ToString(const TArray<EAudioMixerChannel::Type>& InChannels)
+		{
+			FString ChannelString;
+			static const int32 ApproxChannelNameLength = 18;
+			ChannelString.Reserve(ApproxChannelNameLength * InChannels.Num());
+			for (EAudioMixerChannel::Type i : InChannels)
+			{
+				ChannelString.Append(EAudioMixerChannel::ToString(i));
+				ChannelString.Append(TEXT("|"));
+			}
+			return ChannelString;
+		}
 
 		void OnDeviceStateChanged(const FString& DeviceId, const EAudioDeviceState InState, bool ) override
 		{
@@ -1165,7 +1164,7 @@ namespace Audio
 		void OnFormatChanged(const FString& InDeviceId, const FFormatChangedData& InFormat) override
 		{
 			FName DeviceName(InDeviceId);
-			bool bDirty = false;
+			bool bNeedToEnumerateChannels = false;
 		
 			{
 				FReadScopeLock ReadLock(CacheMutationLock);
@@ -1174,52 +1173,35 @@ namespace Audio
 					if (Found->NumChannels != InFormat.NumChannels)
 					{
 						FWriteScopeLock WriteLock(Found->MutationLock);
-						UE_CLOG(Audio::IAudioMixer::ShouldLogDeviceSwaps(),LogAudioMixer, Verbose, TEXT("FWindowsMMDeviceCache: Device '%s' changed default format from %d channels to %d."), *InDeviceId, Found->NumChannels, InFormat.NumChannels);
+						UE_CLOG(Audio::IAudioMixer::ShouldLogDeviceSwaps(),LogAudioMixer, Verbose, TEXT("FWindowsMMDeviceCache: DeviceID='%s', Name='%s' changed default format from %d channels to %d."), *InDeviceId, *Found->FriendlyName, Found->NumChannels, InFormat.NumChannels);
 						Found->NumChannels = InFormat.NumChannels;
-						bDirty = true;
+						bNeedToEnumerateChannels = true;
 					}
 					if (Found->SampleRate != InFormat.SampleRate)
 					{
 						FWriteScopeLock WriteLock(Found->MutationLock);
-						UE_CLOG(Audio::IAudioMixer::ShouldLogDeviceSwaps(),LogAudioMixer, Verbose, TEXT("FWindowsMMDeviceCache: Device '%s' changed default format from %dhz to %dhz."), *InDeviceId, Found->SampleRate, InFormat.SampleRate);
+						UE_CLOG(Audio::IAudioMixer::ShouldLogDeviceSwaps(),LogAudioMixer, Verbose, TEXT("FWindowsMMDeviceCache: DeviceID='%s', Name='%s' changed default format from %dhz to %dhz."), *InDeviceId, *Found->FriendlyName, Found->SampleRate, InFormat.SampleRate);
 						Found->SampleRate = InFormat.SampleRate;
-						bDirty = true;
 					}
-					if (Found->SpeakerConfig != InFormat.ChannelConfig)
+					if (Found->ChannelBitmask != InFormat.ChannelBitmask)
 					{
 						FWriteScopeLock WriteLock(Found->MutationLock);
-						UE_CLOG(Audio::IAudioMixer::ShouldLogDeviceSwaps(),LogAudioMixer, Verbose, TEXT("FWindowsMMDeviceCache: Device '%s' changed default format from 0x%x to 0x%x"), *InDeviceId, Found->SpeakerConfig, InFormat.ChannelConfig);
-						Found->SpeakerConfig = InFormat.ChannelConfig;
-						bDirty = true;
+						UE_CLOG(Audio::IAudioMixer::ShouldLogDeviceSwaps(),LogAudioMixer, Verbose, TEXT("FWindowsMMDeviceCache: DeviceID='%s', Name='%s' changed default format from 0x%x to 0x%x bitmask"), *InDeviceId, *Found->FriendlyName, Found->ChannelBitmask, InFormat.ChannelBitmask);
+						Found->ChannelBitmask = InFormat.ChannelBitmask;
+						bNeedToEnumerateChannels = true;
 					}
-				}
-			}
-		}
-
-		void OnSpeakerConfigChanged(const FString& InDeviceId, uint32 InSpeakerBitmask ) override
-		{
-			FName DeviceName(InDeviceId);
-			FReadScopeLock ReadLock(CacheMutationLock);
-
-			if (FCacheEntry* Found = Cache.Find(DeviceName))
-			{
-				bool bDirty = false;
-				if (Found->SpeakerConfig != InSpeakerBitmask)
-				{
-					FWriteScopeLock WriteLock(Found->MutationLock);
-					
-					UE_CLOG(Audio::IAudioMixer::ShouldLogDeviceSwaps(),LogAudioMixer, Verbose, TEXT("Device '%s' changed default format from 0x%x to 0x%x"),
-						*InDeviceId, Found->SpeakerConfig, InSpeakerBitmask);
-					Found->SpeakerConfig = InSpeakerBitmask;
 
-					// Update speaker array based on new bit mask. (maybe we just do this when have successfully changed).
-					EnumSpeakerMask(InSpeakerBitmask, *Found);
-					
-					bDirty = true;
-				}			
+					if (bNeedToEnumerateChannels)
+					{
+						FWriteScopeLock WriteLock(Found->MutationLock);
+						UE_CLOG(Audio::IAudioMixer::ShouldLogDeviceSwaps(), LogAudioMixer, Verbose, TEXT("FWindowsMMDeviceCache: Channel Change, DeviceID='%s', Name='%s' OLD=[%s]"), *InDeviceId, *Found->FriendlyName, *ToString(Found->OutputChannels));
+						EnumerateChannelMask(InFormat.ChannelBitmask, *Found);
+						UE_CLOG(Audio::IAudioMixer::ShouldLogDeviceSwaps(), LogAudioMixer, Verbose, TEXT("FWindowsMMDeviceCache: Channel Change, DeviceID='%s', Name='%s' NEW=[%s]"), *InDeviceId, *Found->FriendlyName, *ToString(Found->OutputChannels));
+					}
+				}
 			}
 		}
-
+			
 		void MakeDeviceInfo(const FCacheEntry& InEntry, FAudioPlatformDeviceInfo& OutInfo) const
 		{
 			OutInfo.Reset();
