Fix for rare deadlock in deviceswap

diff --git a/Engine/Source/Runtime/Windows/AudioMixerXAudio2/Private/AudioMixerPlatformWindows.cpp b/Engine/Source/Runtime/Windows/AudioMixerXAudio2/Private/AudioMixerPlatformWindows.cpp
--- a/Engine/Source/Runtime/Windows/AudioMixerXAudio2/Private/AudioMixerPlatformWindows.cpp
+++ b/Engine/Source/Runtime/Windows/AudioMixerXAudio2/Private/AudioMixerPlatformWindows.cpp
@@ -1165,44 +1165,56 @@ namespace Audio
 		{
 			FName DeviceName(InDeviceId);
 			bool bNeedToEnumerateChannels = false;
+			bool bDirty = false;
 		
+			FReadScopeLock MapReadLock(CacheMutationLock);
+			if (FCacheEntry* Found = Cache.Find(DeviceName))
 			{
-				FReadScopeLock ReadLock(CacheMutationLock);
-				if (FCacheEntry* Found = Cache.Find(DeviceName))
+				// Make a copy of the entry
+				FCacheEntry EntryCopy(InDeviceId);
 				{
-					if (Found->NumChannels != InFormat.NumChannels)
-					{
-						FWriteScopeLock WriteLock(Found->MutationLock);
-						UE_CLOG(Audio::IAudioMixer::ShouldLogDeviceSwaps(),LogAudioMixer, Verbose, TEXT("FWindowsMMDeviceCache: DeviceID='%s', Name='%s' changed default format from %d channels to %d."), *InDeviceId, *Found->FriendlyName, Found->NumChannels, InFormat.NumChannels);
-						Found->NumChannels = InFormat.NumChannels;
-						bNeedToEnumerateChannels = true;
-					}
-					if (Found->SampleRate != InFormat.SampleRate)
-					{
-						FWriteScopeLock WriteLock(Found->MutationLock);
-						UE_CLOG(Audio::IAudioMixer::ShouldLogDeviceSwaps(),LogAudioMixer, Verbose, TEXT("FWindowsMMDeviceCache: DeviceID='%s', Name='%s' changed default format from %dhz to %dhz."), *InDeviceId, *Found->FriendlyName, Found->SampleRate, InFormat.SampleRate);
-						Found->SampleRate = InFormat.SampleRate;
-					}
-					if (Found->ChannelBitmask != InFormat.ChannelBitmask)
-					{
-						FWriteScopeLock WriteLock(Found->MutationLock);
-						UE_CLOG(Audio::IAudioMixer::ShouldLogDeviceSwaps(),LogAudioMixer, Verbose, TEXT("FWindowsMMDeviceCache: DeviceID='%s', Name='%s' changed default format from 0x%x to 0x%x bitmask"), *InDeviceId, *Found->FriendlyName, Found->ChannelBitmask, InFormat.ChannelBitmask);
-						Found->ChannelBitmask = InFormat.ChannelBitmask;
-						bNeedToEnumerateChannels = true;
-					}
+					FReadScopeLock FoundReadLock(Found->MutationLock);
+					EntryCopy = *Found;
+				}
 
-					if (bNeedToEnumerateChannels)
-					{
-						FWriteScopeLock WriteLock(Found->MutationLock);
-						UE_CLOG(Audio::IAudioMixer::ShouldLogDeviceSwaps(), LogAudioMixer, Verbose, TEXT("FWindowsMMDeviceCache: Channel Change, DeviceID='%s', Name='%s' OLD=[%s]"), *InDeviceId, *Found->FriendlyName, *ToString(Found->OutputChannels));
-						EnumerateChannelMask(InFormat.ChannelBitmask, *Found);
-						UE_CLOG(Audio::IAudioMixer::ShouldLogDeviceSwaps(), LogAudioMixer, Verbose, TEXT("FWindowsMMDeviceCache: Channel Change, DeviceID='%s', Name='%s' NEW=[%s]"), *InDeviceId, *Found->FriendlyName, *ToString(Found->OutputChannels));
-					}
+				if (EntryCopy.NumChannels != InFormat.NumChannels)
+				{
+					UE_CLOG(Audio::IAudioMixer::ShouldLogDeviceSwaps(),LogAudioMixer, Verbose, TEXT("FWindowsMMDeviceCache: DeviceID='%s', Name='%s' changed default format from %d channels to %d."), *InDeviceId, *EntryCopy.FriendlyName, EntryCopy.NumChannels, InFormat.NumChannels);
+					EntryCopy.NumChannels = InFormat.NumChannels;
+					bNeedToEnumerateChannels = true;
+					bDirty = true;
+				}
+				if (EntryCopy.SampleRate != InFormat.SampleRate)
+				{
+					UE_CLOG(Audio::IAudioMixer::ShouldLogDeviceSwaps(),LogAudioMixer, Verbose, TEXT("FWindowsMMDeviceCache: DeviceID='%s', Name='%s' changed default format from %dhz to %dhz."), *InDeviceId, *EntryCopy.FriendlyName, EntryCopy.SampleRate, InFormat.SampleRate);
+					EntryCopy.SampleRate = InFormat.SampleRate;
+					bDirty = true;
+				}
+				if (EntryCopy.ChannelBitmask != InFormat.ChannelBitmask)
+				{
+					UE_CLOG(Audio::IAudioMixer::ShouldLogDeviceSwaps(),LogAudioMixer, Verbose, TEXT("FWindowsMMDeviceCache: DeviceID='%s', Name='%s' changed default format from 0x%x to 0x%x bitmask"), *InDeviceId, *EntryCopy.FriendlyName, EntryCopy.ChannelBitmask, InFormat.ChannelBitmask);
+					EntryCopy.ChannelBitmask = InFormat.ChannelBitmask;
+					bNeedToEnumerateChannels = true;
+					bDirty = true;
+				}
+
+				if (bNeedToEnumerateChannels)
+				{
+					UE_CLOG(Audio::IAudioMixer::ShouldLogDeviceSwaps(), LogAudioMixer, Verbose, TEXT("FWindowsMMDeviceCache: Channel Change, DeviceID='%s', Name='%s' OLD=[%s]"), *InDeviceId, *EntryCopy.FriendlyName, *ToString(EntryCopy.OutputChannels));
+					EnumerateChannelMask(InFormat.ChannelBitmask, EntryCopy);
+					UE_CLOG(Audio::IAudioMixer::ShouldLogDeviceSwaps(), LogAudioMixer, Verbose, TEXT("FWindowsMMDeviceCache: Channel Change, DeviceID='%s', Name='%s' NEW=[%s]"), *InDeviceId, *EntryCopy.FriendlyName, *ToString(EntryCopy.OutputChannels));
+				}
+
+				// Update the entire entry with one write.
+				if (bDirty)
+				{
+					FWriteScopeLock FoundWriteLock(Found->MutationLock);
+					*Found = EntryCopy;
 				}
 			}
 		}
 			
-		void MakeDeviceInfo(const FCacheEntry& InEntry, FAudioPlatformDeviceInfo& OutInfo) const
+		void MakeDeviceInfo(const FCacheEntry& InEntry, FName InDefaultDevice, FAudioPlatformDeviceInfo& OutInfo) const
 		{
 			OutInfo.Reset();
 			OutInfo.Name				= InEntry.FriendlyName;
@@ -1211,13 +1223,7 @@ namespace Audio
 			OutInfo.SampleRate			= InEntry.SampleRate;
 			OutInfo.OutputChannelArray	= InEntry.OutputChannels;
 			OutInfo.Format				= EAudioMixerStreamDataFormat::Float;
-			
-			{
-				FReadScopeLock Lock(CacheMutationLock);
-				bool bConsoleRenderDefault = DefaultRenderId[(int32)EAudioDeviceRole::Console] == InEntry.DeviceId;
-				bool bMultimediaRenderDefault = DefaultRenderId[(int32)EAudioDeviceRole::Multimedia] == InEntry.DeviceId;
-				OutInfo.bIsSystemDefault = bConsoleRenderDefault || bMultimediaRenderDefault;
-			}
+			OutInfo.bIsSystemDefault	= InEntry.DeviceId == InDefaultDevice;
 		}
 	
 		virtual TArray<FAudioPlatformDeviceInfo> GetAllActiveOutputDevices() const override
@@ -1226,11 +1232,14 @@ namespace Audio
 
 			// Find all active devices.
 			TArray<FAudioPlatformDeviceInfo> ActiveDevices;
-	
-			// Read lock on map.
-			FReadScopeLock MapReadLock(CacheMutationLock);
+
+			// Read lock
+			FReadScopeLock ReadLock(CacheMutationLock);
 			ActiveDevices.Reserve(Cache.Num());
 
+			// Ask for defaults once, as we are inside a read lock.
+			FName DefaultRenderDeviceId = GetDefaultOutputDevice_NoLock();
+
 			// Walk cache, read lock for each entry.
 			for (const auto& i : Cache)
 			{
@@ -1240,9 +1249,7 @@ namespace Audio
 					i.Value.Type == FCacheEntry::EEndpointType::Render )
 				{
 					FAudioPlatformDeviceInfo& Info = ActiveDevices.Emplace_GetRef();
-					
-					// Note: This opens another read lock to pull defaults, but we already own the lock.
-					MakeDeviceInfo(i.Value, Info);
+					MakeDeviceInfo(i.Value, DefaultRenderDeviceId, Info);
 				}
 			}
 
@@ -1250,9 +1257,8 @@ namespace Audio
 			return ActiveDevices;
 		}
 
-		FName GetDefaultOutputDevice() const
+		FName GetDefaultOutputDevice_NoLock() const
 		{
-			FReadScopeLock MapReadLock(CacheMutationLock);
 			if (!DefaultRenderId[(int32)EAudioDeviceRole::Console].IsNone())
 			{
 				return DefaultRenderId[(int32)EAudioDeviceRole::Console];
@@ -1275,10 +1281,13 @@ namespace Audio
 
 			FReadScopeLock MapReadLock(CacheMutationLock);
 			
+			// Ask for default here as we are inside the read lock.
+			const FName DefaultOutputDevice = GetDefaultOutputDevice_NoLock();
+
 			// Asking for Default?
 			if (InDeviceID.IsNone())
 			{
-				InDeviceID = GetDefaultOutputDevice();
+				InDeviceID = DefaultOutputDevice;
 				if (InDeviceID.IsNone())
 				{
 					// No default set, fail.
@@ -1294,7 +1303,7 @@ namespace Audio
 					Found->Type == FCacheEntry::EEndpointType::Render )
 				{
 					FAudioPlatformDeviceInfo Info;
-					MakeDeviceInfo(*Found, Info);
+					MakeDeviceInfo(*Found, DefaultOutputDevice, Info);
 					return Info;
 				}
 			}
