Fix handling of xaudio2 errors on device swap

diff --git a/Engine/Source/Runtime/AudioMixer/Private/AudioMixerSourceManager.cpp b/Engine/Source/Runtime/AudioMixer/Private/AudioMixerSourceManager.cpp
--- a/Engine/Source/Runtime/AudioMixer/Private/AudioMixerSourceManager.cpp
+++ b/Engine/Source/Runtime/AudioMixer/Private/AudioMixerSourceManager.cpp
@@ -1484,7 +1484,11 @@ namespace Audio
 					SourceSubmixOutputBuffers[SourceId].SetPreAttenuationSourceBuffer(nullptr);
 				}
 
-				InSubmixPtr->AddOrSetSourceVoice(MixerSources[SourceId], InSubmixSend.SendLevel, InSubmixSend.SubmixSendStage);
+				FMixerSourceVoice* SourceVoice = MixerSources[SourceId];
+				if (ensureAlways(nullptr != SourceVoice))
+				{
+					InSubmixPtr->AddOrSetSourceVoice(SourceVoice, InSubmixSend.SendLevel, InSubmixSend.SubmixSendStage);
+				}
 			}
 		});
 	}
diff --git a/Engine/Source/Runtime/Windows/AudioMixerXAudio2/Private/AudioMixerPlatformXAudio2.cpp b/Engine/Source/Runtime/Windows/AudioMixerXAudio2/Private/AudioMixerPlatformXAudio2.cpp
--- a/Engine/Source/Runtime/Windows/AudioMixerXAudio2/Private/AudioMixerPlatformXAudio2.cpp
+++ b/Engine/Source/Runtime/Windows/AudioMixerXAudio2/Private/AudioMixerPlatformXAudio2.cpp
@@ -42,6 +42,11 @@
 
 #include "Misc/CoreDelegates.h"
 
+#define XAUDIO2_LOG_RESULT(FunctionName, Result) \
+	{ \
+		FString ErrorString = FString::Printf(TEXT("%s -> 0x%X: %s (line: %d)"), TEXT( FunctionName ), Result, *GetErrorString(Result), __LINE__); \
+		UE_LOG(LogAudioMixer, Error, TEXT("XAudio2 Error: %s"), *ErrorString);																		\
+	}
 
 // Macro to check result code for XAudio2 failure, get the string version, log, and goto a cleanup
 #define XAUDIO2_GOTO_CLEANUP_ON_FAIL(Result)																										 \
@@ -61,6 +66,8 @@
 		return false;																														 \
 	}
 
+
+
 static FString GetErrorString(HRESULT Result)
 {
 	switch (Result)
@@ -1288,13 +1295,14 @@ namespace Audio
 			bIsInDeviceSwap = false;
 		}
 
+		// In order to resume audio playback, this function must return true. 
+		// All code paths below return true, even it they encounter an error.
+		
 		if (bTrySwitchToHardwareDevice)
 		{
 			if (!ResetXAudio2System())
 			{
 				// Reinitializing the XAudio2System failed, so we have to exit here.
-				BeginGeneratingAudio();
-				StartRunningNullDevice();
 				return true;
 			}
 
@@ -1321,27 +1329,42 @@ namespace Audio
 			// Get the output device info at this new index
 			GetOutputDeviceInfo(AudioStreamInfo.OutputDeviceIndex, AudioStreamInfo.DeviceInfo);
 
+			HRESULT Result;
 			// Create a new master voice
 			// XAudio2 for HoloLens has different parameters to CreateMasteringVoice
 			// See https://blogs.msdn.microsoft.com/chuckw/2012/04/02/xaudio2-and-windows-8/
 #if PLATFORM_HOLOLENS
-			XAUDIO2_RETURN_ON_FAIL(XAudio2System->CreateMasteringVoice(
+			Result = XAudio2System->CreateMasteringVoice(
 				&OutputAudioStreamMasteringVoice, 
 				AudioStreamInfo.DeviceInfo.NumChannels, 
 				AudioStreamInfo.DeviceInfo.SampleRate, 
 				0, 
 				AllAudioDevices->GetAt(AudioStreamInfo.OutputDeviceIndex)->Id->Data(), 
-				nullptr));
+				nullptr);
+			if (FAILED(Result))
+			{
+				XAUDIO2_LOG_RESULT("XAudio2System->CreateMasteringVoice", Result);
+				// Switch to running null device by setting OutputAudioStreamMasteringVoice to null.
+				// Will default to null device when calling `ResumePlaybackOnNewDevice()`
+				OutputAudioStreamMasteringVoice = nullptr;
+			}
 #else
 			// open up on the default device
-			XAUDIO2_RETURN_ON_FAIL(XAudio2System->CreateMasteringVoice(
+			Result = XAudio2System->CreateMasteringVoice(
 				&OutputAudioStreamMasteringVoice,
 				AudioStreamInfo.DeviceInfo.NumChannels,
 				AudioStreamInfo.DeviceInfo.SampleRate,
 				0,
 				*AudioStreamInfo.DeviceInfo.DeviceId,
 				nullptr,
-				AudioCategory_GameEffects));
+				AudioCategory_GameEffects);
+			if (FAILED(Result))
+			{
+				XAUDIO2_LOG_RESULT("XAudio2System->CreateMasteringVoice", Result);
+				// Switch to running null device by setting OutputAudioStreamMasteringVoice to null.
+				// Will default to null device when calling `ResumePlaybackOnNewDevice()`
+				OutputAudioStreamMasteringVoice = nullptr;
+			}
 #endif
 
 			// Setup the format of the output source voice
@@ -1354,23 +1377,21 @@ namespace Audio
 			Format.wBitsPerSample = sizeof(float) * 8;
 
 			// Create the output source voice
-			HRESULT Result = XAudio2System->CreateSourceVoice(&OutputAudioStreamSourceVoice, &Format, XAUDIO2_VOICE_NOPITCH, 2.0f, &OutputVoiceCallback);
+			Result = XAudio2System->CreateSourceVoice(&OutputAudioStreamSourceVoice, &Format, XAUDIO2_VOICE_NOPITCH, 2.0f, &OutputVoiceCallback);
 			if (FAILED(Result))
 			{
-				FString ErrorString = FString::Printf(TEXT("%s -> 0x%X: %s (line: %d)"), TEXT( "XAudio2System->CreateSourceVoice" ), Result, *GetErrorString(Result), __LINE__);
-				UE_LOG(LogAudioMixer, Error, TEXT("XAudio2 Error: %s"), *ErrorString);																 
-
+				XAUDIO2_LOG_RESULT("XAudio2System->CreateSourceVoice", Result);
 				// Switch to running null device by setting OutputAudioStreamSourceVoice to null.
 				// Will default to null device when calling `ResumePlaybackOnNewDevice()`
 				OutputAudioStreamSourceVoice = nullptr;
-
-				// Return true to signal that playback must be restarted. 
-				return true;
 			}
 
 			// Reinitialize the output circular buffer to match the buffer math of the new audio device.
 			const int32 NumOutputSamples = AudioStreamInfo.NumOutputFrames * AudioStreamInfo.DeviceInfo.NumChannels;
-			OutputBuffer.Init(AudioStreamInfo.AudioMixer, NumOutputSamples, NumOutputBuffers, AudioStreamInfo.DeviceInfo.Format);
+			if (ensure(NumOutputSamples > 0))
+			{
+				OutputBuffer.Init(AudioStreamInfo.AudioMixer, NumOutputSamples, NumOutputBuffers, AudioStreamInfo.DeviceInfo.Format);
+			}
 		}
 		else
 		{	
