D3D11 RHI compiles for platform extension Windows platforms.

#rnx

diff --git a/Engine/Source/Runtime/Windows/D3D11RHI/Private/D3D11Commands.cpp b/Engine/Source/Runtime/Windows/D3D11RHI/Private/D3D11Commands.cpp
--- a/Engine/Source/Runtime/Windows/D3D11RHI/Private/D3D11Commands.cpp
+++ b/Engine/Source/Runtime/Windows/D3D11RHI/Private/D3D11Commands.cpp
@@ -15,7 +15,7 @@
 #include "SceneUtils.h"
 #include "EngineGlobals.h"
 
-#if PLATFORM_DESKTOP
+#if !PLATFORM_HOLOLENS
 // For Depth Bounds Test interface
 #include "Windows/AllowWindowsPlatformTypes.h"
 	#include "nvapi.h"
