Disable moving audio devices on platforms that don't support it

diff --git a/Engine/Source/Runtime/Windows/AudioMixerXAudio2/Private/AudioMixerPlatformXAudio2.cpp b/Engine/Source/Runtime/Windows/AudioMixerXAudio2/Private/AudioMixerPlatformXAudio2.cpp
--- a/Engine/Source/Runtime/Windows/AudioMixerXAudio2/Private/AudioMixerPlatformXAudio2.cpp
+++ b/Engine/Source/Runtime/Windows/AudioMixerXAudio2/Private/AudioMixerPlatformXAudio2.cpp
@@ -216,6 +216,7 @@ namespace Audio
 
 	bool FMixerPlatformXAudio2::AllowDeviceSwap()
 	{
+#if PLATFORM_WINDOWS
 		double CurrentTime = FPlatformTime::Seconds();
 
 		// If we're already in the process of swapping, we do not want to "double-trigger" a swap
@@ -232,6 +233,7 @@ namespace Audio
 			LastDeviceSwapTime = CurrentTime;
 			return true;
 		}
+#endif
 		return false;
 	}
 
@@ -1186,6 +1188,7 @@ namespace Audio
 
 	bool FMixerPlatformXAudio2::CheckAudioDeviceChange()
 	{
+#if PLATFORM_WINDOWS && XAUDIO_SUPPORTS_DEVICE_DETAILS
 		FScopeLock Lock(&AudioDeviceSwapCriticalSection);
 
 		if (bMoveAudioStreamToNewAudioDevice)
@@ -1194,6 +1197,7 @@ namespace Audio
 
 			return MoveAudioStreamToNewAudioDevice(NewAudioDeviceId);
 		}
+#endif
 		return false;
 	}
 
@@ -1221,13 +1225,14 @@ namespace Audio
 
 	bool FMixerPlatformXAudio2::MoveAudioStreamToNewAudioDevice(const FString& InNewDeviceId)
 	{
+		bool bDidStopGeneratingAudio = false;
 #if PLATFORM_WINDOWS && XAUDIO_SUPPORTS_DEVICE_DETAILS
 
 		uint32 NumDevices = 0;
 		// XAudio2 for HoloLens doesn't have GetDeviceCount, use local wrapper instead
 		if (!GetNumOutputDevices(NumDevices))
 		{
-			return false;
+			return bDidStopGeneratingAudio;
 		}
 
 		// If we're running the null device, This function is called every second or so.
@@ -1240,7 +1245,7 @@ namespace Audio
 		if (bContinueUsingNullDevice)
 		{
 			// Audio device was not changed. Return false to avoid downstream device change logic.
-			return false;
+			return bDidStopGeneratingAudio;
 		}
 
 		UE_LOG(LogAudioMixer, Log, TEXT("Resetting audio stream to device id %s"), *InNewDeviceId);
@@ -1249,6 +1254,7 @@ namespace Audio
 		// different channel formats. Stop generating audio to protect against
 		// accessing the OutputBuffer.
 		StopGeneratingAudio();
+		bDidStopGeneratingAudio = true;
 
 		// Stop currently running device
 		if (bIsUsingNullDevice)
@@ -1260,7 +1266,7 @@ namespace Audio
 			// Not initialized!
 			if (!bIsInitialized)
 			{
-				return true;
+				return bDidStopGeneratingAudio;
 			}
 
 			// If an XAudio2 callback is in flight,
@@ -1303,7 +1309,7 @@ namespace Audio
 			if (!ResetXAudio2System())
 			{
 				// Reinitializing the XAudio2System failed, so we have to exit here.
-				return true;
+				return bDidStopGeneratingAudio;
 			}
 
 			// Now get info on the new audio device we're trying to reset to
@@ -1401,12 +1407,12 @@ namespace Audio
 			// NullDevice is started when OutputAudioStreamSourceVoice is null
 			check(nullptr == OutputAudioStreamSourceVoice);
 
-			return true;
+			return bDidStopGeneratingAudio;
 		}
 		
 #endif // #if PLATFORM_WINDOWS
 
-		return true; 
+		return bDidStopGeneratingAudio; 
 	}
 
 	void FMixerPlatformXAudio2::ResumePlaybackOnNewDevice()
