Improvements to PixelFormat capability checking
- Adding EPixelFormatCapabilities::UAV
- Adding PixelFormatHasCapabilities helper
- D3D format queries need to use the different UAV/Resource formats

diff --git a/Engine/Source/Runtime/D3D12RHI/Private/D3D12Device.cpp b/Engine/Source/Runtime/D3D12RHI/Private/D3D12Device.cpp
--- a/Engine/Source/Runtime/D3D12RHI/Private/D3D12Device.cpp
+++ b/Engine/Source/Runtime/D3D12RHI/Private/D3D12Device.cpp
@@ -158,6 +158,16 @@ typedef HRESULT(WINAPI *FDXGIGetDebugInterface1)(UINT, REFIID, void **);
 
 ID3D12CommandQueue* gD3D12CommandQueue;
 
+static D3D12_FEATURE_DATA_FORMAT_SUPPORT GetFormatSupport(ID3D12Device* InDevice, DXGI_FORMAT InFormat)
+{
+	D3D12_FEATURE_DATA_FORMAT_SUPPORT FormatSupport{};
+	FormatSupport.Format = InFormat;
+
+	InDevice->CheckFeatureSupport(D3D12_FEATURE_FORMAT_SUPPORT, &FormatSupport, sizeof(FormatSupport));
+
+	return FormatSupport;
+}
+
 void FD3D12Device::SetupAfterDeviceCreation()
 {
 	ID3D12Device* Direct3DDevice = GetParentAdapter()->GetD3DDevice();
@@ -171,57 +181,57 @@ void FD3D12Device::SetupAfterDeviceCreation()
 
 		if (PlatformFormat != DXGI_FORMAT_UNKNOWN)
 		{
-			D3D12_FEATURE_DATA_FORMAT_SUPPORT FormatSupport{};
-			FormatSupport.Format = PlatformFormat;
+			const D3D12_FEATURE_DATA_FORMAT_SUPPORT FormatSupport    = GetFormatSupport(Direct3DDevice, PlatformFormat);
+			const D3D12_FEATURE_DATA_FORMAT_SUPPORT SRVFormatSupport = GetFormatSupport(Direct3DDevice, FindShaderResourceDXGIFormat(PlatformFormat, false));
+			const D3D12_FEATURE_DATA_FORMAT_SUPPORT UAVFormatSupport = GetFormatSupport(Direct3DDevice, FindUnorderedAccessDXGIFormat(PlatformFormat));
+			const D3D12_FEATURE_DATA_FORMAT_SUPPORT RTVFormatSupport = GetFormatSupport(Direct3DDevice, FindShaderResourceDXGIFormat(PlatformFormat, false));
+			const D3D12_FEATURE_DATA_FORMAT_SUPPORT DSVFormatSupport = GetFormatSupport(Direct3DDevice, FindDepthStencilDXGIFormat(PlatformFormat));
 
-			auto ConvertCap1 = [&Capabilities, &FormatSupport](EPixelFormatCapabilities UnrealCap, D3D12_FORMAT_SUPPORT1 InFlags)
+			auto ConvertCap1 = [&Capabilities](const D3D12_FEATURE_DATA_FORMAT_SUPPORT& InSupport, EPixelFormatCapabilities UnrealCap, D3D12_FORMAT_SUPPORT1 InFlags)
 			{
-				if (EnumHasAnyFlags(FormatSupport.Support1, InFlags))
+				if (EnumHasAnyFlags(InSupport.Support1, InFlags))
 				{
 					EnumAddFlags(Capabilities, UnrealCap);
 				}
 			};
-			auto ConvertCap2 = [&Capabilities, &FormatSupport](EPixelFormatCapabilities UnrealCap, D3D12_FORMAT_SUPPORT2 InFlags)
+			auto ConvertCap2 = [&Capabilities](const D3D12_FEATURE_DATA_FORMAT_SUPPORT& InSupport, EPixelFormatCapabilities UnrealCap, D3D12_FORMAT_SUPPORT2 InFlags)
 			{
-				if (EnumHasAnyFlags(FormatSupport.Support2, InFlags))
+				if (EnumHasAnyFlags(InSupport.Support2, InFlags))
 				{
 					EnumAddFlags(Capabilities, UnrealCap);
 				}
 			};
 
-			HRESULT hr = Direct3DDevice->CheckFeatureSupport(D3D12_FEATURE_FORMAT_SUPPORT, &FormatSupport, sizeof(FormatSupport));
-			if (SUCCEEDED(hr))
-			{
-				ConvertCap1(EPixelFormatCapabilities::Texture1D,			D3D12_FORMAT_SUPPORT1_TEXTURE1D);
-				ConvertCap1(EPixelFormatCapabilities::Texture2D,			D3D12_FORMAT_SUPPORT1_TEXTURE2D);
-				ConvertCap1(EPixelFormatCapabilities::Texture3D,			D3D12_FORMAT_SUPPORT1_TEXTURE3D);
-				ConvertCap1(EPixelFormatCapabilities::TextureCube,			D3D12_FORMAT_SUPPORT1_TEXTURECUBE);
-				ConvertCap1(EPixelFormatCapabilities::Buffer,               D3D12_FORMAT_SUPPORT1_BUFFER);
-				ConvertCap1(EPixelFormatCapabilities::VertexBuffer,         D3D12_FORMAT_SUPPORT1_IA_VERTEX_BUFFER);
-				ConvertCap1(EPixelFormatCapabilities::IndexBuffer,          D3D12_FORMAT_SUPPORT1_IA_INDEX_BUFFER);
-
-				if (EnumHasAnyFlags(Capabilities, EPixelFormatCapabilities::AnyTexture))
-				{
-					ConvertCap1(EPixelFormatCapabilities::RenderTarget,     D3D12_FORMAT_SUPPORT1_RENDER_TARGET);
-					ConvertCap1(EPixelFormatCapabilities::DepthStencil,     D3D12_FORMAT_SUPPORT1_DEPTH_STENCIL);
-					ConvertCap1(EPixelFormatCapabilities::TextureMipmaps,   D3D12_FORMAT_SUPPORT1_MIP);
-					ConvertCap1(EPixelFormatCapabilities::TextureLoad,      D3D12_FORMAT_SUPPORT1_SHADER_LOAD);
-					ConvertCap1(EPixelFormatCapabilities::TextureSample,    D3D12_FORMAT_SUPPORT1_SHADER_SAMPLE);
-					ConvertCap1(EPixelFormatCapabilities::TextureGather,    D3D12_FORMAT_SUPPORT1_SHADER_GATHER);
-					ConvertCap2(EPixelFormatCapabilities::TextureAtomics,   D3D12_FORMAT_SUPPORT2_UAV_ATOMIC_EXCHANGE);
-					ConvertCap1(EPixelFormatCapabilities::TextureBlendable, D3D12_FORMAT_SUPPORT1_BLENDABLE);
-				}
+			ConvertCap1(FormatSupport, EPixelFormatCapabilities::Texture1D,               D3D12_FORMAT_SUPPORT1_TEXTURE1D);
+			ConvertCap1(FormatSupport, EPixelFormatCapabilities::Texture2D,               D3D12_FORMAT_SUPPORT1_TEXTURE2D);
+			ConvertCap1(FormatSupport, EPixelFormatCapabilities::Texture3D,               D3D12_FORMAT_SUPPORT1_TEXTURE3D);
+			ConvertCap1(FormatSupport, EPixelFormatCapabilities::TextureCube,             D3D12_FORMAT_SUPPORT1_TEXTURECUBE);
+			ConvertCap1(FormatSupport, EPixelFormatCapabilities::Buffer,                  D3D12_FORMAT_SUPPORT1_BUFFER);
+			ConvertCap1(FormatSupport, EPixelFormatCapabilities::VertexBuffer,            D3D12_FORMAT_SUPPORT1_IA_VERTEX_BUFFER);
+			ConvertCap1(FormatSupport, EPixelFormatCapabilities::IndexBuffer,             D3D12_FORMAT_SUPPORT1_IA_INDEX_BUFFER);
 
-				if (EnumHasAnyFlags(Capabilities, EPixelFormatCapabilities::Buffer))
-				{
-					ConvertCap1(EPixelFormatCapabilities::BufferLoad,       D3D12_FORMAT_SUPPORT1_SHADER_LOAD);
-					ConvertCap2(EPixelFormatCapabilities::BufferStore,      D3D12_FORMAT_SUPPORT2_UAV_TYPED_STORE);
-					ConvertCap2(EPixelFormatCapabilities::BufferAtomics,    D3D12_FORMAT_SUPPORT2_UAV_ATOMIC_EXCHANGE);
-				}
+			if (EnumHasAnyFlags(Capabilities, EPixelFormatCapabilities::AnyTexture))
+			{
+				ConvertCap1(FormatSupport, EPixelFormatCapabilities::RenderTarget,        D3D12_FORMAT_SUPPORT1_RENDER_TARGET);
+				ConvertCap1(FormatSupport, EPixelFormatCapabilities::DepthStencil,        D3D12_FORMAT_SUPPORT1_DEPTH_STENCIL);
+				ConvertCap1(FormatSupport, EPixelFormatCapabilities::TextureMipmaps,      D3D12_FORMAT_SUPPORT1_MIP);
+				ConvertCap1(SRVFormatSupport, EPixelFormatCapabilities::TextureLoad,      D3D12_FORMAT_SUPPORT1_SHADER_LOAD);
+				ConvertCap1(SRVFormatSupport, EPixelFormatCapabilities::TextureSample,    D3D12_FORMAT_SUPPORT1_SHADER_SAMPLE);
+				ConvertCap1(SRVFormatSupport, EPixelFormatCapabilities::TextureGather,    D3D12_FORMAT_SUPPORT1_SHADER_GATHER);
+				ConvertCap2(UAVFormatSupport, EPixelFormatCapabilities::TextureAtomics,   D3D12_FORMAT_SUPPORT2_UAV_ATOMIC_EXCHANGE);
+				ConvertCap1(SRVFormatSupport, EPixelFormatCapabilities::TextureBlendable, D3D12_FORMAT_SUPPORT1_BLENDABLE);
+			}
 
-				ConvertCap2(EPixelFormatCapabilities::TypedUAVLoad,			D3D12_FORMAT_SUPPORT2_UAV_TYPED_LOAD);
-				ConvertCap2(EPixelFormatCapabilities::TypedUAVStore,		D3D12_FORMAT_SUPPORT2_UAV_TYPED_STORE);
+			if (EnumHasAnyFlags(Capabilities, EPixelFormatCapabilities::Buffer))
+			{
+				ConvertCap1(SRVFormatSupport, EPixelFormatCapabilities::BufferLoad,       D3D12_FORMAT_SUPPORT1_SHADER_LOAD);
+				ConvertCap2(UAVFormatSupport, EPixelFormatCapabilities::BufferStore,      D3D12_FORMAT_SUPPORT2_UAV_TYPED_STORE);
+				ConvertCap2(UAVFormatSupport, EPixelFormatCapabilities::BufferAtomics,    D3D12_FORMAT_SUPPORT2_UAV_ATOMIC_EXCHANGE);
 			}
+
+			ConvertCap1(UAVFormatSupport, EPixelFormatCapabilities::UAV,                  D3D12_FORMAT_SUPPORT1_TYPED_UNORDERED_ACCESS_VIEW);
+			ConvertCap2(UAVFormatSupport, EPixelFormatCapabilities::TypedUAVLoad,         D3D12_FORMAT_SUPPORT2_UAV_TYPED_LOAD);
+			ConvertCap2(UAVFormatSupport, EPixelFormatCapabilities::TypedUAVStore,        D3D12_FORMAT_SUPPORT2_UAV_TYPED_STORE);
 		}
 
 		PixelFormatInfo.Capabilities = Capabilities;
diff --git a/Engine/Source/Runtime/RHI/Public/RHI.h b/Engine/Source/Runtime/RHI/Public/RHI.h
--- a/Engine/Source/Runtime/RHI/Public/RHI.h
+++ b/Engine/Source/Runtime/RHI/Public/RHI.h
@@ -757,8 +757,9 @@ enum class EPixelFormatCapabilities : uint32
     BufferStore      = 1ull << 17,
     BufferAtomics    = 1ull << 18,
 
-    TypedUAVLoad     = 1ull << 19,
-    TypedUAVStore    = 1ull << 20,
+	UAV              = 1ull << 19,
+    TypedUAVLoad     = 1ull << 20,
+	TypedUAVStore    = 1ull << 21,
 
 	AnyTexture = Texture1D | Texture2D | Texture3D | TextureCube,
 };
@@ -798,6 +799,11 @@ struct FPixelFormatInfo
 
 extern RHI_API FPixelFormatInfo GPixelFormats[PF_MAX];		// Maps members of EPixelFormat to a FPixelFormatInfo describing the format.
 
+inline bool PixelFormatHasCapabilities(EPixelFormat InFormat, EPixelFormatCapabilities InCapabilities)
+{
+	return EnumHasAllFlags(GPixelFormats[InFormat].Capabilities, InCapabilities);
+}
+
 //
 //	CalculateImageBytes
 //
diff --git a/Engine/Source/Runtime/Windows/D3D11RHI/Private/D3D11Device.cpp b/Engine/Source/Runtime/Windows/D3D11RHI/Private/D3D11Device.cpp
--- a/Engine/Source/Runtime/Windows/D3D11RHI/Private/D3D11Device.cpp
+++ b/Engine/Source/Runtime/Windows/D3D11RHI/Private/D3D11Device.cpp
@@ -384,72 +384,104 @@ uint32 FD3D11DynamicRHI::GetMaxMSAAQuality(uint32 SampleCount)
 	return 0xffffffff;
 }
 
+struct FFormatSupport
+{
+	D3D11_FORMAT_SUPPORT FormatSupport;
+	D3D11_FORMAT_SUPPORT2 FormatSupport2;
+};
+
+static FFormatSupport GetFormatSupport(FD3D11Device* InDevice, DXGI_FORMAT InFormat)
+{
+	FFormatSupport Result{};
+
+	{
+		D3D11_FEATURE_DATA_FORMAT_SUPPORT FormatSupport{};
+		FormatSupport.InFormat = InFormat;
+
+		HRESULT SupportHR = InDevice->CheckFeatureSupport(D3D11_FEATURE_FORMAT_SUPPORT, &FormatSupport, sizeof(FormatSupport));
+		if (SUCCEEDED(SupportHR))
+		{
+			Result.FormatSupport = (D3D11_FORMAT_SUPPORT)FormatSupport.OutFormatSupport;
+		}
+	}
+
+	{
+		D3D11_FEATURE_DATA_FORMAT_SUPPORT2 FormatSupport2{};
+		FormatSupport2.InFormat = InFormat;
+
+		HRESULT Support2HR = InDevice->CheckFeatureSupport(D3D11_FEATURE_FORMAT_SUPPORT2, &FormatSupport2, sizeof(FormatSupport2));
+		if (SUCCEEDED(Support2HR))
+		{
+			Result.FormatSupport2 = (D3D11_FORMAT_SUPPORT2)FormatSupport2.OutFormatSupport2;
+		}
+	}
+
+	return Result;
+}
+
 void FD3D11DynamicRHI::SetupAfterDeviceCreation()
 {
 	for (uint32 FormatIndex = PF_Unknown; FormatIndex < PF_MAX; FormatIndex++)
 	{
 		FPixelFormatInfo& PixelFormatInfo = GPixelFormats[FormatIndex];
 		const DXGI_FORMAT PlatformFormat = static_cast<DXGI_FORMAT>(PixelFormatInfo.PlatformFormat);
+		const DXGI_FORMAT UAVFormat = FindUnorderedAccessDXGIFormat(PlatformFormat);
 
 		EPixelFormatCapabilities Capabilities = EPixelFormatCapabilities::None;
 
 		if (PlatformFormat != DXGI_FORMAT_UNKNOWN)
 		{
-			D3D11_FEATURE_DATA_FORMAT_SUPPORT FormatSupport{};
-			FormatSupport.InFormat = PlatformFormat;
+			const FFormatSupport FormatSupport    = GetFormatSupport(Direct3DDevice, PlatformFormat);
+			const FFormatSupport SRVFormatSupport = GetFormatSupport(Direct3DDevice, FindShaderResourceDXGIFormat(PlatformFormat, false));
+			const FFormatSupport UAVFormatSupport = GetFormatSupport(Direct3DDevice, FindUnorderedAccessDXGIFormat(PlatformFormat));
+			const FFormatSupport RTVFormatSupport = GetFormatSupport(Direct3DDevice, FindShaderResourceDXGIFormat(PlatformFormat, false));
+			const FFormatSupport DSVFormatSupport = GetFormatSupport(Direct3DDevice, FindDepthStencilDXGIFormat(PlatformFormat));
 
-			D3D11_FEATURE_DATA_FORMAT_SUPPORT2 FormatSupport2{};
-			FormatSupport2.InFormat = PlatformFormat;
-
-			auto ConvertCap1 = [&Capabilities, &FormatSupport](EPixelFormatCapabilities UnrealCap, D3D11_FORMAT_SUPPORT InFlag)
+			auto ConvertCap1 = [&Capabilities](const FFormatSupport& InSupport, EPixelFormatCapabilities UnrealCap, D3D11_FORMAT_SUPPORT InFlag)
 			{
-				if (EnumHasAllFlags((D3D11_FORMAT_SUPPORT)FormatSupport.OutFormatSupport, InFlag))
+				if (EnumHasAllFlags(InSupport.FormatSupport, InFlag))
 				{
 					EnumAddFlags(Capabilities, UnrealCap);
 				}
 			};
-			auto ConvertCap2 = [&Capabilities, &FormatSupport2](EPixelFormatCapabilities UnrealCap, D3D11_FORMAT_SUPPORT2 InFlag)
+			auto ConvertCap2 = [&Capabilities](const FFormatSupport& InSupport, EPixelFormatCapabilities UnrealCap, D3D11_FORMAT_SUPPORT2 InFlag)
 			{
-				if (EnumHasAllFlags((D3D11_FORMAT_SUPPORT2)FormatSupport2.OutFormatSupport2, InFlag))
+				if (EnumHasAllFlags(InSupport.FormatSupport2, InFlag))
 				{
 					EnumAddFlags(Capabilities, UnrealCap);
 				}
 			};
 
-			HRESULT SupportHR = Direct3DDevice->CheckFeatureSupport(D3D11_FEATURE_FORMAT_SUPPORT, &FormatSupport, sizeof(FormatSupport));
-			HRESULT Support2HR = Direct3DDevice->CheckFeatureSupport(D3D11_FEATURE_FORMAT_SUPPORT2, &FormatSupport2, sizeof(FormatSupport2));
-			if (SUCCEEDED(SupportHR) || SUCCEEDED(Support2HR))
-			{
-				ConvertCap1(EPixelFormatCapabilities::Texture1D,			D3D11_FORMAT_SUPPORT_TEXTURE1D);
-				ConvertCap1(EPixelFormatCapabilities::Texture2D,			D3D11_FORMAT_SUPPORT_TEXTURE2D);
-				ConvertCap1(EPixelFormatCapabilities::Texture3D,			D3D11_FORMAT_SUPPORT_TEXTURE3D);
-				ConvertCap1(EPixelFormatCapabilities::TextureCube,			D3D11_FORMAT_SUPPORT_TEXTURECUBE);
-				ConvertCap1(EPixelFormatCapabilities::Buffer,               D3D11_FORMAT_SUPPORT_BUFFER);
-				ConvertCap1(EPixelFormatCapabilities::VertexBuffer,         D3D11_FORMAT_SUPPORT_IA_VERTEX_BUFFER);
-				ConvertCap1(EPixelFormatCapabilities::IndexBuffer,          D3D11_FORMAT_SUPPORT_IA_INDEX_BUFFER);
-
-				if (EnumHasAnyFlags(Capabilities, EPixelFormatCapabilities::AnyTexture))
-				{
-					ConvertCap1(EPixelFormatCapabilities::RenderTarget,     D3D11_FORMAT_SUPPORT_RENDER_TARGET);
-					ConvertCap1(EPixelFormatCapabilities::DepthStencil,     D3D11_FORMAT_SUPPORT_DEPTH_STENCIL);
-					ConvertCap1(EPixelFormatCapabilities::TextureMipmaps,   D3D11_FORMAT_SUPPORT_DEPTH_STENCIL);
-					ConvertCap1(EPixelFormatCapabilities::TextureLoad,      D3D11_FORMAT_SUPPORT_DEPTH_STENCIL);
-					ConvertCap1(EPixelFormatCapabilities::TextureSample,    D3D11_FORMAT_SUPPORT_DEPTH_STENCIL);
-					ConvertCap1(EPixelFormatCapabilities::TextureGather,    D3D11_FORMAT_SUPPORT_DEPTH_STENCIL);
-					ConvertCap1(EPixelFormatCapabilities::TextureAtomics,   D3D11_FORMAT_SUPPORT_DEPTH_STENCIL);
-					ConvertCap1(EPixelFormatCapabilities::TextureBlendable, D3D11_FORMAT_SUPPORT_DEPTH_STENCIL);
-				}
-
-				if (EnumHasAnyFlags(Capabilities, EPixelFormatCapabilities::Buffer))
-				{
-					ConvertCap1(EPixelFormatCapabilities::BufferLoad,       D3D11_FORMAT_SUPPORT_SHADER_LOAD);
-					ConvertCap2(EPixelFormatCapabilities::BufferStore,      D3D11_FORMAT_SUPPORT2_UAV_TYPED_STORE);
-					ConvertCap2(EPixelFormatCapabilities::BufferAtomics,    D3D11_FORMAT_SUPPORT2_UAV_ATOMIC_EXCHANGE);
-				}
-
-				ConvertCap2(EPixelFormatCapabilities::TypedUAVLoad,         D3D11_FORMAT_SUPPORT2_UAV_TYPED_LOAD);
-				ConvertCap2(EPixelFormatCapabilities::TypedUAVStore,        D3D11_FORMAT_SUPPORT2_UAV_TYPED_STORE);
-			}
+            ConvertCap1(FormatSupport, EPixelFormatCapabilities::Texture1D,               D3D11_FORMAT_SUPPORT_TEXTURE1D);
+            ConvertCap1(FormatSupport, EPixelFormatCapabilities::Texture2D,               D3D11_FORMAT_SUPPORT_TEXTURE2D);
+            ConvertCap1(FormatSupport, EPixelFormatCapabilities::Texture3D,               D3D11_FORMAT_SUPPORT_TEXTURE3D);
+            ConvertCap1(FormatSupport, EPixelFormatCapabilities::TextureCube,             D3D11_FORMAT_SUPPORT_TEXTURECUBE);
+            ConvertCap1(FormatSupport, EPixelFormatCapabilities::Buffer,                  D3D11_FORMAT_SUPPORT_BUFFER);
+            ConvertCap1(FormatSupport, EPixelFormatCapabilities::VertexBuffer,            D3D11_FORMAT_SUPPORT_IA_VERTEX_BUFFER);
+            ConvertCap1(FormatSupport, EPixelFormatCapabilities::IndexBuffer,             D3D11_FORMAT_SUPPORT_IA_INDEX_BUFFER);
+
+            if (EnumHasAnyFlags(Capabilities, EPixelFormatCapabilities::AnyTexture))
+            {
+                ConvertCap1(FormatSupport, EPixelFormatCapabilities::RenderTarget,        D3D11_FORMAT_SUPPORT_RENDER_TARGET);
+                ConvertCap1(FormatSupport, EPixelFormatCapabilities::DepthStencil,        D3D11_FORMAT_SUPPORT_DEPTH_STENCIL);
+                ConvertCap1(FormatSupport, EPixelFormatCapabilities::TextureMipmaps,      D3D11_FORMAT_SUPPORT_MIP);
+                ConvertCap1(SRVFormatSupport, EPixelFormatCapabilities::TextureLoad,      D3D11_FORMAT_SUPPORT_SHADER_LOAD);
+                ConvertCap1(SRVFormatSupport, EPixelFormatCapabilities::TextureSample,    D3D11_FORMAT_SUPPORT_SHADER_SAMPLE);
+                ConvertCap1(SRVFormatSupport, EPixelFormatCapabilities::TextureGather,    D3D11_FORMAT_SUPPORT_SHADER_GATHER);
+                ConvertCap2(UAVFormatSupport, EPixelFormatCapabilities::TextureAtomics,   D3D11_FORMAT_SUPPORT2_UAV_ATOMIC_EXCHANGE);
+                ConvertCap1(RTVFormatSupport, EPixelFormatCapabilities::TextureBlendable, D3D11_FORMAT_SUPPORT_BLENDABLE);
+            }
+
+            if (EnumHasAnyFlags(Capabilities, EPixelFormatCapabilities::Buffer))
+            {
+                ConvertCap1(SRVFormatSupport, EPixelFormatCapabilities::BufferLoad,       D3D11_FORMAT_SUPPORT_SHADER_LOAD);
+                ConvertCap2(UAVFormatSupport, EPixelFormatCapabilities::BufferStore,      D3D11_FORMAT_SUPPORT2_UAV_TYPED_STORE);
+                ConvertCap2(UAVFormatSupport, EPixelFormatCapabilities::BufferAtomics,    D3D11_FORMAT_SUPPORT2_UAV_ATOMIC_EXCHANGE);
+            }
+
+			ConvertCap1(UAVFormatSupport, EPixelFormatCapabilities::UAV,                  D3D11_FORMAT_SUPPORT_TYPED_UNORDERED_ACCESS_VIEW);
+			ConvertCap2(UAVFormatSupport, EPixelFormatCapabilities::TypedUAVLoad,         D3D11_FORMAT_SUPPORT2_UAV_TYPED_LOAD);
+			ConvertCap2(UAVFormatSupport, EPixelFormatCapabilities::TypedUAVStore,        D3D11_FORMAT_SUPPORT2_UAV_TYPED_STORE);
 		}
 
 		PixelFormatInfo.Capabilities = Capabilities;
