U5 - More removal of SuspendRenderingThread

diff --git a/Engine/Source/Runtime/D3D12RHI/Private/D3D12Adapter.cpp b/Engine/Source/Runtime/D3D12RHI/Private/D3D12Adapter.cpp
--- a/Engine/Source/Runtime/D3D12RHI/Private/D3D12Adapter.cpp
+++ b/Engine/Source/Runtime/D3D12RHI/Private/D3D12Adapter.cpp
@@ -822,7 +822,7 @@ void FD3D12Adapter::InitializeDevices()
 	check(IsInGameThread());
 
 	// Wait for the rendering thread to go idle.
-	SCOPED_SUSPEND_RENDERING_THREAD(false);
+	FlushRenderingCommands();
 
 	// If the device we were using has been removed, release it and the resources we created for it.
 	if (bDeviceRemoved)
diff --git a/Engine/Source/Runtime/VulkanRHI/Private/VulkanRHI.cpp b/Engine/Source/Runtime/VulkanRHI/Private/VulkanRHI.cpp
--- a/Engine/Source/Runtime/VulkanRHI/Private/VulkanRHI.cpp
+++ b/Engine/Source/Runtime/VulkanRHI/Private/VulkanRHI.cpp
@@ -737,11 +737,11 @@ void FVulkanDynamicRHI::InitInstance()
 {
 	check(IsInGameThread());
 
-	// Wait for the rendering thread to go idle.
-	SCOPED_SUSPEND_RENDERING_THREAD(false);
-
 	if (!Device)
 	{
+		// Wait for the rendering thread to go idle.
+		FlushRenderingCommands();
+
 		check(!GIsRHIInitialized);
 
 		FVulkanPlatform::OverridePlatformHandlers(true);
diff --git a/Engine/Source/Runtime/Windows/D3D11RHI/Private/Windows/WindowsD3D11Device.cpp b/Engine/Source/Runtime/Windows/D3D11RHI/Private/Windows/WindowsD3D11Device.cpp
--- a/Engine/Source/Runtime/Windows/D3D11RHI/Private/Windows/WindowsD3D11Device.cpp
+++ b/Engine/Source/Runtime/Windows/D3D11RHI/Private/Windows/WindowsD3D11Device.cpp
@@ -1667,7 +1667,7 @@ void FD3D11DynamicRHI::InitD3DDevice()
 	if(!Direct3DDevice)
 	{
 		// Wait for the rendering thread to go idle.
-		SCOPED_SUSPEND_RENDERING_THREAD(false);
+		FlushRenderingCommands();
 
 		UE_LOG(LogD3D11RHI, Log, TEXT("Creating new Direct3DDevice"));
 		check(!GIsRHIInitialized);
