Fix for deadlock with MMNotification client

diff --git a/Engine/Source/Runtime/Windows/AudioMixerXAudio2/Private/AudioMixerPlatformWindows.cpp b/Engine/Source/Runtime/Windows/AudioMixerXAudio2/Private/AudioMixerPlatformWindows.cpp
--- a/Engine/Source/Runtime/Windows/AudioMixerXAudio2/Private/AudioMixerPlatformWindows.cpp
+++ b/Engine/Source/Runtime/Windows/AudioMixerXAudio2/Private/AudioMixerPlatformWindows.cpp
@@ -9,6 +9,8 @@
 #include "AudioMixerPlatformXAudio2.h"
 #include "AudioMixer.h"
 #include "AudioDeviceNotificationSubsystem.h"
+#include "Misc/ScopeRWLock.h"
+#include <atomic>
 
 #if PLATFORM_WINDOWS
 
@@ -74,9 +76,7 @@ public:
 	}
 
 	bool RegisterForSessionNotifications(const TComPtr<IMMDevice>& InDevice)
-	{
-		FScopeLock ScopeLock(&MutationCs);
-
+	{	
 		// If we're already listening to this device, we can early out.
 		if (DeviceListeningToSessionEvents == InDevice)
 		{
@@ -115,9 +115,7 @@ public:
 	}
 
 	void UnregisterForSessionNotifications()
-	{
-		FScopeLock ScopeLock(&MutationCs);
-		
+	{	
 		// Unregister for any device we're already listening to.
 		if (SessionControls)
 		{
@@ -204,7 +202,6 @@ public:
 
 	HRESULT STDMETHODCALLTYPE OnDefaultDeviceChanged(EDataFlow InFlow, ERole InRole, LPCWSTR pwstrDeviceId) override
 	{
-		FScopeLock ScopeLock(&MutationCs);
 
 		UE_CLOG(Audio::IAudioMixer::ShouldLogDeviceSwaps(),LogAudioMixer, Warning, 
 			TEXT("FWindowsMMNotificationClient: OnDefaultDeviceChanged: %s, %s, %s - %s"), ToString(InFlow), ToString(InRole), pwstrDeviceId, *GetFriendlyName(pwstrDeviceId));
@@ -232,6 +229,7 @@ public:
 		FString DeviceString(pwstrDeviceId);
 		if (InFlow == eRender)
 		{
+			FReadScopeLock Lock(ListenersSetRwLock);
 			for (Audio::IAudioMixerDeviceChangedListener* Listener : Listeners)
 			{
 				Listener->OnDefaultRenderDeviceChanged(AudioDeviceRole, DeviceString);
@@ -239,6 +237,7 @@ public:
 		}
 		else if (InFlow == eCapture)
 		{
+			FReadScopeLock Lock(ListenersSetRwLock);
 			for (Audio::IAudioMixerDeviceChangedListener* Listener : Listeners)
 			{
 				Listener->OnDefaultCaptureDeviceChanged(AudioDeviceRole, DeviceString);
@@ -246,6 +245,7 @@ public:
 		}
 		else
 		{
+			FReadScopeLock Lock(ListenersSetRwLock);
 			for (Audio::IAudioMixerDeviceChangedListener* Listener : Listeners)
 			{
 				Listener->OnDefaultCaptureDeviceChanged(AudioDeviceRole, DeviceString);
@@ -278,8 +278,6 @@ public:
 
 	HRESULT STDMETHODCALLTYPE OnDeviceAdded(LPCWSTR pwstrDeviceId) override
 	{
-		FScopeLock ScopeLock(&MutationCs);
-
 		UE_CLOG(Audio::IAudioMixer::ShouldLogDeviceSwaps(), LogAudioMixer, Display, TEXT("FWindowsMMNotificationClient: OnDeviceAdded: %s"), *GetFriendlyName(pwstrDeviceId));
 			
 		if (Audio::IAudioMixer::ShouldIgnoreDeviceSwaps())
@@ -288,18 +286,17 @@ public:
 		}
 
 		bool bIsRender = IsRenderDevice(pwstrDeviceId);
-
+		FString DeviceString(pwstrDeviceId);
+		FReadScopeLock ReadLock(ListenersSetRwLock);
 		for (Audio::IAudioMixerDeviceChangedListener* Listener : Listeners)
 		{
-			Listener->OnDeviceAdded(FString(pwstrDeviceId), bIsRender);
+			Listener->OnDeviceAdded(DeviceString, bIsRender);
 		}
 		return S_OK;
 	};
 
 	HRESULT STDMETHODCALLTYPE OnDeviceRemoved(LPCWSTR pwstrDeviceId) override
-	{
-		FScopeLock ScopeLock(&MutationCs);
-
+	{		
 		UE_CLOG(Audio::IAudioMixer::ShouldLogDeviceSwaps(),LogAudioMixer, Display, TEXT("FWindowsMMNotificationClient: OnDeviceRemoved: %s"), *GetFriendlyName(pwstrDeviceId));
 
 		if (Audio::IAudioMixer::ShouldIgnoreDeviceSwaps())
@@ -308,17 +305,17 @@ public:
 		}
 
 		bool bIsRender = IsRenderDevice(pwstrDeviceId);
+		FString DeviceString(pwstrDeviceId);
+		FReadScopeLock ReadLock(ListenersSetRwLock);		
 		for (Audio::IAudioMixerDeviceChangedListener* Listener : Listeners)
 		{
-			Listener->OnDeviceRemoved(FString(pwstrDeviceId), bIsRender);
+			Listener->OnDeviceRemoved(DeviceString, bIsRender);
 		}
 		return S_OK;
 	}
 
 	HRESULT STDMETHODCALLTYPE OnDeviceStateChanged(LPCWSTR pwstrDeviceId, DWORD dwNewState) override
 	{
-		FScopeLock ScopeLock(&MutationCs);
-
 		UE_CLOG(Audio::IAudioMixer::ShouldLogDeviceSwaps(), LogAudioMixer, Display, TEXT("FWindowsMMNotificationClient: OnDeviceStateChanged: %s, %d"), *GetFriendlyName(pwstrDeviceId), dwNewState);
 
 		if (Audio::IAudioMixer::ShouldIgnoreDeviceSwaps())
@@ -330,10 +327,12 @@ public:
 		if (dwNewState == DEVICE_STATE_ACTIVE || dwNewState == DEVICE_STATE_DISABLED || dwNewState == DEVICE_STATE_UNPLUGGED || dwNewState == DEVICE_STATE_NOTPRESENT)
 		{
 			Audio::EAudioDeviceState State = ConvertWordToDeviceState(dwNewState);
-						
+		
+			FString DeviceString(pwstrDeviceId);
+			FReadScopeLock ReadLock(ListenersSetRwLock);
 			for (Audio::IAudioMixerDeviceChangedListener* Listener : Listeners)
 			{
-				Listener->OnDeviceStateChanged(FString(pwstrDeviceId), State, bIsRender);
+				Listener->OnDeviceStateChanged(DeviceString, State, bIsRender);
 			}
 		}
 		return S_OK;
@@ -465,7 +464,7 @@ public:
 							FormatChanged.ChannelBitmask = WaveFormatEx->wFormatTag == WAVE_FORMAT_EXTENSIBLE ?
 								((const WAVEFORMATEXTENSIBLE*)WaveFormatEx)->dwChannelMask : 0;
 
-							FScopeLock ScopeLock(&MutationCs);
+							FReadScopeLock ReadLock(ListenersSetRwLock);
 							for (Audio::IAudioMixerDeviceChangedListener* i : Listeners)
 							{
 								i->OnFormatChanged(DeviceId, FormatChanged);
@@ -534,13 +533,15 @@ public:
 
 	void RegisterDeviceChangedListener(Audio::IAudioMixerDeviceChangedListener* DeviceChangedListener)
 	{
-		FScopeLock ScopeLock(&MutationCs);
+		// Modifying container so get full write lock
+		FWriteScopeLock Lock(ListenersSetRwLock);
 		Listeners.Add(DeviceChangedListener);
 	}
 
 	void UnRegisterDeviceDeviceChangedListener(Audio::IAudioMixerDeviceChangedListener* DeviceChangedListener)
 	{
-		FScopeLock ScopeLock(&MutationCs);
+		// Modifying container so get full write lock
+		FWriteScopeLock Lock(ListenersSetRwLock);
 		Listeners.Remove(DeviceChangedListener);
 	}
 
@@ -595,13 +596,16 @@ public:
 		UE_CLOG(Audio::IAudioMixer::ShouldLogDeviceSwaps(), LogAudioMixer, Verbose, TEXT("Session Disconnect: Reason=%s, DeviceBound=%s, HasDisconnectSessionHappened=%d"), 
 			ToString(InDisconnectReason), *GetFriendlyName(DeviceListeningToSessionEvents), (int32)bHasDisconnectSessionHappened);
 
-		FScopeLock ScopeLock(&MutationCs);
+
 		if (!bHasDisconnectSessionHappened)
 		{
-			Audio::IAudioMixerDeviceChangedListener::EDisconnectReason Reason = AudioSessionDisconnectToEDisconnectReason(InDisconnectReason);
-			for (Audio::IAudioMixerDeviceChangedListener* i : Listeners)
 			{
-				i->OnSessionDisconnect(Reason);
+				FReadScopeLock Lock(ListenersSetRwLock);
+				Audio::IAudioMixerDeviceChangedListener::EDisconnectReason Reason = AudioSessionDisconnectToEDisconnectReason(InDisconnectReason);
+				for (Audio::IAudioMixerDeviceChangedListener* i : Listeners)
+				{
+					i->OnSessionDisconnect(Reason);
+				}
 			}
 			
 			// Mark this true.
@@ -634,13 +638,13 @@ public:
 private:
 	LONG Ref;
 	TSet<Audio::IAudioMixerDeviceChangedListener*> Listeners;
-	FCriticalSection MutationCs;
+	FRWLock ListenersSetRwLock;
 	TComPtr<IMMDeviceEnumerator> DeviceEnumerator;
 	TComPtr<IAudioSessionManager> SessionManager;
 	TComPtr<IAudioSessionControl> SessionControls;
 	TComPtr<IMMDevice> DeviceListeningToSessionEvents;
 	bool bComInitialized;
-	bool bHasDisconnectSessionHappened;
+	std::atomic<bool> bHasDisconnectSessionHappened;
 };
 
 namespace Audio
