Fix for shutdown crash with in flight threaded device swap.

diff --git a/Engine/Source/Runtime/AudioMixerCore/Private/AudioMixer.cpp b/Engine/Source/Runtime/AudioMixerCore/Private/AudioMixer.cpp
--- a/Engine/Source/Runtime/AudioMixerCore/Private/AudioMixer.cpp
+++ b/Engine/Source/Runtime/AudioMixerCore/Private/AudioMixer.cpp
@@ -84,7 +84,7 @@ FAutoConsoleVariableRef CVarRecycleThreads(
 	TEXT("0 off, 1 on"),
 	ECVF_Default);
 
-static int32 OverrunTimeoutCVar = 1000;
+static int32 OverrunTimeoutCVar = 5000;
 FAutoConsoleVariableRef CVarOverrunTimeout(
 	TEXT("au.OverrunTimeoutMSec"),
 	OverrunTimeoutCVar,
diff --git a/Engine/Source/Runtime/Windows/AudioMixerXAudio2/Private/AudioMixerPlatformXAudio2.cpp b/Engine/Source/Runtime/Windows/AudioMixerXAudio2/Private/AudioMixerPlatformXAudio2.cpp
--- a/Engine/Source/Runtime/Windows/AudioMixerXAudio2/Private/AudioMixerPlatformXAudio2.cpp
+++ b/Engine/Source/Runtime/Windows/AudioMixerXAudio2/Private/AudioMixerPlatformXAudio2.cpp
@@ -337,6 +337,8 @@ namespace Audio
 						SwapReason = DeviceSwapReason
 					]() mutable->FXAudio2AsyncCreateResult
 						{
+							uint64_t StartTimeCycles = FPlatformTime::Cycles64();
+							
 							SCOPED_NAMED_EVENT(FMixerPlatformXAudio2_AsyncDeleteCreate, FColor::Blue);
 							UE_LOG(LogAudioMixer, Verbose, TEXT("FMixerPlatformXAudio2::CheckThreadedDeviceSwap() - AsyncTask Start. Because=%s"), *SwapReason);
 
@@ -467,12 +469,20 @@ namespace Audio
 							FPlatformMisc::CoUninitialize();
 
 							// Success.
+							Results.SuccessfullDurationMs = FPlatformTime::ToMilliseconds64(FPlatformTime::Cycles64() - StartTimeCycles);
 							Results.DeviceInfo = NewDevice;
 							return Results;
-						};
-					}
+						}; // End of Lambda.
 
-					// NULL our system / master / source voices. (as they are about to be torn down async).
+						// Still inside CS here.
+						// NULL our system / master / source voices. (as they are about to be torn down async).
+						bIsInDeviceSwap = true;
+						XAudio2System = nullptr;
+						OutputAudioStreamMasteringVoice = nullptr;
+						OutputAudioStreamSourceVoice = nullptr;
+					} // End of Lambda Creation scope (with CS).
+
+				
 					{
 						SCOPED_NAMED_EVENT(FMixerPlatformXAudio2_AsyncDeleteCreate_StopSource, FColor::Blue);
 
@@ -486,12 +496,6 @@ namespace Audio
 							}
 						}
 
-						// Lock CS so any callback mid swap earlies out
-						FScopeLock Lock(&AudioDeviceSwapCriticalSection);
-						bIsInDeviceSwap = true;
-						XAudio2System = nullptr;
-						OutputAudioStreamMasteringVoice = nullptr;
-						OutputAudioStreamSourceVoice = nullptr;
 					}
 
 					// Start a new device swap.
@@ -540,8 +544,8 @@ namespace Audio
 							AudioStreamInfo.DeviceInfo = ActiveDeviceSwap.Get().DeviceInfo;
 
 							// Display our new XAudio2 Mastering voice details.
-							UE_LOG(LogAudioMixer, Display, TEXT("Successful Swap new Device is (NumChannels=%u, SampleRate=%u, DeviceID=%s, Name=%s), Reason=%s, InstanceID=%d"),
-								(uint32)AudioStreamInfo.DeviceInfo.NumChannels, (uint32)AudioStreamInfo.DeviceInfo.SampleRate, *AudioStreamInfo.DeviceInfo.DeviceId, *AudioStreamInfo.DeviceInfo.Name, *ActiveDeviceSwap.Get().SwapReason, InstanceID);
+							UE_LOG(LogAudioMixer, Display, TEXT("Successful Swap new Device is (NumChannels=%u, SampleRate=%u, DeviceID=%s, Name=%s), Reason=%s, InstanceID=%d, DurationMS=%.2f"),
+								(uint32)AudioStreamInfo.DeviceInfo.NumChannels, (uint32)AudioStreamInfo.DeviceInfo.SampleRate, *AudioStreamInfo.DeviceInfo.DeviceId, *AudioStreamInfo.DeviceInfo.Name, *ActiveDeviceSwap.Get().SwapReason, InstanceID, ActiveDeviceSwap.Get().SuccessfullDurationMs);
 
 							// Reinitialize the output circular buffer to match the buffer math of the new audio device.
 							const int32 NumOutputSamples = AudioStreamInfo.NumOutputFrames * AudioStreamInfo.DeviceInfo.NumChannels;
@@ -1501,6 +1505,14 @@ namespace Audio
 			return false;
 		}
 
+		// If we have a active device swap in flight, we need to wait for it.
+		static const FTimespan TimeoutOnDeviceSwap(0,0,0,3,0); // 3 secs.
+		if (ActiveDeviceSwap.IsValid() && !ActiveDeviceSwap.IsReady() && !ActiveDeviceSwap.WaitFor(TimeoutOnDeviceSwap))
+		{
+			UE_LOG(LogAudioMixer, Warning, TEXT("Timeout waiting for inflight device-swap. InstanceID=%d"), InstanceID);
+			return false;
+		}
+
 		if (bIsDeviceOpen && !StopAudioStream())
 		{
 			return false;
@@ -1578,9 +1590,8 @@ namespace Audio
 			// Signal that the thread that is running the update that we're stopping
 			if (OutputAudioStreamSourceVoice)
 			{
-				FScopeLock ScopeLock(&DeviceSwapCriticalSection);
-				OutputAudioStreamSourceVoice->DestroyVoice();
-				OutputAudioStreamSourceVoice = nullptr;
+				OutputAudioStreamSourceVoice->Stop(0, 0); // Don't wait for tails, stop as quick as you can.
+
 			}
 
 			check(AudioStreamInfo.StreamState == EAudioOutputStreamState::Stopped);
diff --git a/Engine/Source/Runtime/Windows/AudioMixerXAudio2/Private/AudioMixerPlatformXAudio2.h b/Engine/Source/Runtime/Windows/AudioMixerXAudio2/Private/AudioMixerPlatformXAudio2.h
--- a/Engine/Source/Runtime/Windows/AudioMixerXAudio2/Private/AudioMixerPlatformXAudio2.h
+++ b/Engine/Source/Runtime/Windows/AudioMixerXAudio2/Private/AudioMixerPlatformXAudio2.h
@@ -114,6 +114,7 @@ namespace Audio
 			IXAudio2SourceVoice* OutputAudioStreamSourceVoice = nullptr;
 			FAudioPlatformDeviceInfo DeviceInfo;
 			FString SwapReason;
+			float SuccessfullDurationMs = 0.f;
 		};
 		TFuture<FXAudio2AsyncCreateResult> ActiveDeviceSwap;
 		bool CheckThreadedDeviceSwap();
