Remove UE4 references from old XAudio2 modules

diff --git a/Engine/Source/Runtime/Windows/XAudio2/Private/XAudio2Source.cpp b/Engine/Source/Runtime/Windows/XAudio2/Private/XAudio2Source.cpp
--- a/Engine/Source/Runtime/Windows/XAudio2/Private/XAudio2Source.cpp
+++ b/Engine/Source/Runtime/Windows/XAudio2/Private/XAudio2Source.cpp
@@ -1717,7 +1717,7 @@ void FSpatializationHelper::Init()
 	Emitter.pCone->InnerReverb = 0.0f;
 	Emitter.pCone->OuterReverb = 1.0f;
 
-	Emitter.ChannelCount = UE4_XAUDIO3D_INPUTCHANNELS;
+	Emitter.ChannelCount = UE_XAUDIO3D_INPUTCHANNELS;
 	Emitter.ChannelRadius = 0.0f;
 	// we aren't using the helper to spatialize multichannel files so we can set this nullptr
 	Emitter.pChannelAzimuths = nullptr;
@@ -1748,7 +1748,7 @@ void FSpatializationHelper::Init()
 	// Zero the matrix coefficients
 	FMemory::Memzero(MatrixCoefficients, sizeof(float)*UE_ARRAY_COUNT(MatrixCoefficients));
 
-	DSPSettings.SrcChannelCount = UE4_XAUDIO3D_INPUTCHANNELS;
+	DSPSettings.SrcChannelCount = UE_XAUDIO3D_INPUTCHANNELS;
 	DSPSettings.DstChannelCount = SPEAKER_COUNT;
 	DSPSettings.pMatrixCoefficients = MatrixCoefficients;
 	DSPSettings.pDelayTimes = nullptr;
@@ -1901,7 +1901,7 @@ void FSpatializationHelper::DumpSpatializationState() const
 	FLocal::DumpDistanceCurve(TEXT("  "), TEXT("ReverbVolumeCurve"), &ReverbVolumeCurve);
 
 	// EmitterAzimuths
-	FLocal::DumpChannelArray(TEXT("  "), TEXT("EmitterAzimuths"), UE4_XAUDIO3D_INPUTCHANNELS, EmitterAzimuths);
+	FLocal::DumpChannelArray(TEXT("  "), TEXT("EmitterAzimuths"), UE_XAUDIO3D_INPUTCHANNELS, EmitterAzimuths);
 
 	// MatrixCoefficients
 	FLocal::DumpChannelArray(TEXT("  "), TEXT("MatrixCoefficients"), UE_ARRAY_COUNT(MatrixCoefficients), MatrixCoefficients);
diff --git a/Engine/Source/Runtime/Windows/XAudio2/Private/XAudio2Support.h b/Engine/Source/Runtime/Windows/XAudio2/Private/XAudio2Support.h
--- a/Engine/Source/Runtime/Windows/XAudio2/Private/XAudio2Support.h
+++ b/Engine/Source/Runtime/Windows/XAudio2/Private/XAudio2Support.h
@@ -167,7 +167,7 @@ private:
 #define SPEAKER_5POINT0          ( SPEAKER_FRONT_LEFT | SPEAKER_FRONT_RIGHT | SPEAKER_FRONT_CENTER | SPEAKER_SIDE_LEFT | SPEAKER_SIDE_RIGHT )
 #define SPEAKER_6POINT1          ( SPEAKER_FRONT_LEFT | SPEAKER_FRONT_RIGHT | SPEAKER_FRONT_CENTER | SPEAKER_LOW_FREQUENCY | SPEAKER_SIDE_LEFT | SPEAKER_SIDE_RIGHT | SPEAKER_BACK_CENTER )
 
-#define UE4_XAUDIO3D_INPUTCHANNELS 1
+#define UE_XAUDIO3D_INPUTCHANNELS 1
 
 struct FPCMBufferInfo
 {
@@ -688,7 +688,7 @@ class FSpatializationHelper
 	X3DAUDIO_DISTANCE_CURVE_POINT ReverbVolumeCurvePoint[2];
 	X3DAUDIO_DISTANCE_CURVE       ReverbVolumeCurve;
 
-	float                         EmitterAzimuths[UE4_XAUDIO3D_INPUTCHANNELS];
+	float                         EmitterAzimuths[UE_XAUDIO3D_INPUTCHANNELS];
 
 	// TODO: Hardcoding this to 8 because X3DAudioCalculate is ignoring the destination speaker count we put in and
 	//       using the number of speakers on the output device.  For 7.1 this means that it writes to 8 speakers,
