Changes from jessica.chen at Intel to prepare for Gen12+ platforms. Changes were made to remove the excess calls to IsRHIDeviceIntel().

diff --git a/Engine/Source/Runtime/D3D12RHI/Private/Windows/WindowsD3D12Device.cpp b/Engine/Source/Runtime/D3D12RHI/Private/Windows/WindowsD3D12Device.cpp
--- a/Engine/Source/Runtime/D3D12RHI/Private/Windows/WindowsD3D12Device.cpp
+++ b/Engine/Source/Runtime/D3D12RHI/Private/Windows/WindowsD3D12Device.cpp
@@ -721,6 +721,7 @@ FDynamicRHI* FD3D12DynamicRHIModule::CreateRHI(ERHIFeatureLevel::Type RequestedF
 	if (ChosenAdapters.Num() > 0 && ChosenAdapters[0].IsValid())
 	{
 		FGenericCrashContext::SetEngineData(TEXT("RHI.IntegratedGPU"), ChosenAdapters[0].Get()->GetDesc().bIsIntegrated ? TEXT("true") : TEXT("false"));
+		GRHIDeviceIsIntegrated = ChosenAdapters[0].Get()->GetDesc().bIsIntegrated;
 	}
 
 	GD3D12RHI = new FD3D12DynamicRHI(ChosenAdapters, bPixEventEnabled);
diff --git a/Engine/Source/Runtime/RHI/Private/RHI.cpp b/Engine/Source/Runtime/RHI/Private/RHI.cpp
--- a/Engine/Source/Runtime/RHI/Private/RHI.cpp
+++ b/Engine/Source/Runtime/RHI/Private/RHI.cpp
@@ -1450,6 +1450,8 @@ bool GRHISupportsConservativeRasterization = false;
 
 bool GRHISupportsPipelineFileCache = false;
 
+bool GRHIDeviceIsIntegrated = false;
+
 /** Whether we are profiling GPU hitches. */
 bool GTriggerGPUHitchProfile = false;
 
@@ -1731,7 +1733,6 @@ RHI_API bool IsRHIDeviceAMD()
 RHI_API bool IsRHIDeviceIntel()
 {
 	check(GRHIVendorId != 0);
-	// Intel GPUs are integrated and use both DedicatedVideoMemory and SharedSystemMemory.
 	return GRHIVendorId == 0x8086;
 }
 
diff --git a/Engine/Source/Runtime/RHI/Public/RHI.h b/Engine/Source/Runtime/RHI/Public/RHI.h
--- a/Engine/Source/Runtime/RHI/Public/RHI.h
+++ b/Engine/Source/Runtime/RHI/Public/RHI.h
@@ -637,6 +637,9 @@ extern RHI_API uint32 GRHIRayTracingInstanceDescriptorSize;
 /** Whether or not the RHI supports shader wave operations (shader model 6.0). */
 extern RHI_API bool GRHISupportsWaveOperations;
 
+/** Whether or not the current GPU is integrated into to CPU */
+extern RHI_API bool GRHIDeviceIsIntegrated;
+
 /** 
 * Specifies the minimum and maximum number of lanes in the SIMD wave that this GPU can support. I.e. 32 on NVIDIA, 64 on AMD.
 * Valid values are in range [4..128] (as per SM 6.0 specification) or 0 if unknown.
diff --git a/Engine/Source/Runtime/Renderer/Private/DistanceFieldAmbientOcclusion.cpp b/Engine/Source/Runtime/Renderer/Private/DistanceFieldAmbientOcclusion.cpp
--- a/Engine/Source/Runtime/Renderer/Private/DistanceFieldAmbientOcclusion.cpp
+++ b/Engine/Source/Runtime/Renderer/Private/DistanceFieldAmbientOcclusion.cpp
@@ -675,7 +675,7 @@ bool SupportsDistanceFieldAO(ERHIFeatureLevel::Type FeatureLevel, EShaderPlatfor
 		// Better to disable entirely than to display garbage
 		&& !GRHIDeviceIsAMDPreGCNArchitecture
 		// Intel HD 4000 hangs in the RHICreateTexture3D call to allocate the large distance field atlas, and virtually no Intel cards can afford it anyway
-		&& !IsRHIDeviceIntel()
+		&& !GRHIDeviceIsIntegrated
 		&& FeatureLevel >= ERHIFeatureLevel::SM5
 		&& DoesPlatformSupportDistanceFieldAO(ShaderPlatform)
 		&& IsUsingDistanceFields(ShaderPlatform);
@@ -719,7 +719,7 @@ bool FSceneRenderer::ShouldPrepareDistanceFieldScene() const
 		return false;
 	}
 
-	if (IsRHIDeviceIntel())
+	if (GRHIDeviceIsIntegrated)
 	{
 		// Intel HD 4000 hangs in the RHICreateTexture3D call to allocate the large distance field atlas, and virtually no Intel cards can afford it anyway
 		return false;
diff --git a/Engine/Source/Runtime/Renderer/Private/HairStrands/HairStrandsVisibility.cpp b/Engine/Source/Runtime/Renderer/Private/HairStrands/HairStrandsVisibility.cpp
--- a/Engine/Source/Runtime/Renderer/Private/HairStrands/HairStrandsVisibility.cpp
+++ b/Engine/Source/Runtime/Renderer/Private/HairStrands/HairStrandsVisibility.cpp
@@ -123,11 +123,7 @@ enum EHairVisibilityRenderMode
 
 inline bool DoesSupportRasterCompute()
 {
-#if PLATFORM_WINDOWS
-	return (IsRHIDeviceNVIDIA() || IsRHIDeviceAMD()) && GRHISupportsAtomicUInt64;
-#else
 	return GRHISupportsAtomicUInt64;
-#endif
 }
 
 inline EHairVisibilityRenderMode GetHairVisibilityRenderMode()
diff --git a/Engine/Source/Runtime/Windows/D3D11RHI/Private/Windows/WindowsD3D11Device.cpp b/Engine/Source/Runtime/Windows/D3D11RHI/Private/Windows/WindowsD3D11Device.cpp
--- a/Engine/Source/Runtime/Windows/D3D11RHI/Private/Windows/WindowsD3D11Device.cpp
+++ b/Engine/Source/Runtime/Windows/D3D11RHI/Private/Windows/WindowsD3D11Device.cpp
@@ -1069,6 +1069,8 @@ void FD3D11DynamicRHIModule::FindAdapter()
 	{
 		UE_LOG(LogD3D11RHI, Error, TEXT("Failed to choose a D3D11 Adapter."));
 	}
+
+	GRHIDeviceIsIntegrated = ChosenAdapter.bIsIntegrated;
 }
 
 FDynamicRHI* FD3D11DynamicRHIModule::CreateRHI(ERHIFeatureLevel::Type RequestedFeatureLevel)
@@ -1751,7 +1753,7 @@ void FD3D11DynamicRHI::InitD3DDevice()
 		TRefCountPtr<IDXGIAdapter3> DxgiAdapter3;
 		DXGI_QUERY_VIDEO_MEMORY_INFO LocalVideoMemoryInfo;
 		FD3D11GlobalStats::GTotalGraphicsMemory = 0;
-		if ( IsRHIDeviceIntel() )
+		if (Adapter.bIsIntegrated)
 		{
 			// It's all system memory.
 			FD3D11GlobalStats::GTotalGraphicsMemory = FD3D11GlobalStats::GDedicatedVideoMemory;
@@ -2122,7 +2124,7 @@ void FD3D11DynamicRHI::InitD3DDevice()
 #endif // INTEL_METRICSDISCOVERY
 
 		// Disable the RHI thread by default for devices that will likely suffer in performance
-		if (IsRHIDeviceIntel() || FPlatformMisc::NumberOfCores() < 4)
+		if (Adapter.bIsIntegrated || FPlatformMisc::NumberOfCores() < 4)
 		{
 			GRHISupportsRHIThread = false;
 		}
