Adding RHIIsTypedUAVStoreSupported, implementing on D3D11 and D3D12.
Bringing over texture format changes from Release-5.0
- D3D format queries need to use the different UAV/Resource formats
- Check TypedUAVLoad support for D3D12

diff --git a/Engine/Source/Runtime/D3D12RHI/Private/D3D12Device.cpp b/Engine/Source/Runtime/D3D12RHI/Private/D3D12Device.cpp
--- a/Engine/Source/Runtime/D3D12RHI/Private/D3D12Device.cpp
+++ b/Engine/Source/Runtime/D3D12RHI/Private/D3D12Device.cpp
@@ -595,4 +595,9 @@ D3D12_RESOURCE_ALLOCATION_INFO FD3D12Device::GetResourceAllocationInfo(const D3D
 bool FD3D12DynamicRHI::RHIIsTypedUAVLoadSupported(EPixelFormat PixelFormat)
 {
 	return EnumHasAnyFlags(GPixelFormats[PixelFormat].Capabilities, EPixelFormatCapabilities::TypedUAVLoad);
-}
\ No newline at end of file
+}
+
+bool FD3D12DynamicRHI::RHIIsTypedUAVStoreSupported(EPixelFormat PixelFormat)
+{
+	return EnumHasAnyFlags(GPixelFormats[PixelFormat].Capabilities, EPixelFormatCapabilities::TypedUAVStore);
+}
diff --git a/Engine/Source/Runtime/D3D12RHI/Private/D3D12RHIPrivate.h b/Engine/Source/Runtime/D3D12RHI/Private/D3D12RHIPrivate.h
--- a/Engine/Source/Runtime/D3D12RHI/Private/D3D12RHIPrivate.h
+++ b/Engine/Source/Runtime/D3D12RHI/Private/D3D12RHIPrivate.h
@@ -398,6 +398,7 @@ public:
 	virtual IRHITransientResourceAllocator* RHICreateTransientResourceAllocator() override;
 
 	virtual bool RHIIsTypedUAVLoadSupported(EPixelFormat PixelFormat) override;
+	virtual bool RHIIsTypedUAVStoreSupported(EPixelFormat PixelFormat) override;
 
 #if WITH_MGPU
 	virtual IRHICommandContextContainer* RHIGetCommandContextContainer(int32 Index, int32 Num, FRHIGPUMask GPUMask)final override;
diff --git a/Engine/Source/Runtime/RHI/Public/DynamicRHI.h b/Engine/Source/Runtime/RHI/Public/DynamicRHI.h
--- a/Engine/Source/Runtime/RHI/Public/DynamicRHI.h
+++ b/Engine/Source/Runtime/RHI/Public/DynamicRHI.h
@@ -1461,6 +1461,7 @@ public:
 	virtual void RHIPollRenderQueryResults() {}
 
 	virtual bool RHIIsTypedUAVLoadSupported(EPixelFormat PixelFormat) { return true; }
+	virtual bool RHIIsTypedUAVStoreSupported(EPixelFormat PixelFormat) { return true; }
 
 	virtual uint16 RHIGetPlatformTextureMaxSampleCount() { return 8; };
 
diff --git a/Engine/Source/Runtime/Windows/D3D11RHI/Private/D3D11Device.cpp b/Engine/Source/Runtime/Windows/D3D11RHI/Private/D3D11Device.cpp
--- a/Engine/Source/Runtime/Windows/D3D11RHI/Private/D3D11Device.cpp
+++ b/Engine/Source/Runtime/Windows/D3D11RHI/Private/D3D11Device.cpp
@@ -532,6 +532,11 @@ bool FD3D11DynamicRHI::RHIIsTypedUAVLoadSupported(EPixelFormat PixelFormat)
 	return EnumHasAllFlags(GPixelFormats[PixelFormat].Capabilities, EPixelFormatCapabilities::TypedUAVLoad);
 }
 
+bool FD3D11DynamicRHI::RHIIsTypedUAVStoreSupported(EPixelFormat PixelFormat)
+{
+	return EnumHasAllFlags(GPixelFormats[PixelFormat].Capabilities, EPixelFormatCapabilities::TypedUAVStore);
+}
+
 void FD3D11DynamicRHI::UpdateMSAASettings()
 {	
 	check(DX_MAX_MSAA_COUNT == 8);
diff --git a/Engine/Source/Runtime/Windows/D3D11RHI/Private/D3D11RHIPrivate.h b/Engine/Source/Runtime/Windows/D3D11RHI/Private/D3D11RHIPrivate.h
--- a/Engine/Source/Runtime/Windows/D3D11RHI/Private/D3D11RHIPrivate.h
+++ b/Engine/Source/Runtime/Windows/D3D11RHI/Private/D3D11RHIPrivate.h
@@ -683,6 +683,7 @@ public:
 	virtual void RHICalibrateTimers() override;
 
 	virtual bool RHIIsTypedUAVLoadSupported(EPixelFormat PixelFormat) override;
+	virtual bool RHIIsTypedUAVStoreSupported(EPixelFormat PixelFormat) override;
 
 	// Accessors.
 	ID3D11Device* GetDevice() const
