Cleanup XAudio2Windows glut and break out into smaller more manageable pieces.

diff --git a/Engine/Source/Runtime/Windows/AudioMixerXAudio2/Private/AudioMixerPlatformWindows.cpp b/Engine/Source/Runtime/Windows/AudioMixerXAudio2/Private/AudioMixerPlatformWindows.cpp
--- a/Engine/Source/Runtime/Windows/AudioMixerXAudio2/Private/AudioMixerPlatformWindows.cpp
+++ b/Engine/Source/Runtime/Windows/AudioMixerXAudio2/Private/AudioMixerPlatformWindows.cpp
@@ -10,650 +10,25 @@
 #include "AudioMixer.h"
 #include "AudioDeviceNotificationSubsystem.h"
 #include "Misc/ScopeRWLock.h"
+
 #include <atomic>
 
 #if PLATFORM_WINDOWS
 
+#include "Windows/MinWindows.h"
 #include "Windows/COMPointer.h"
-#include "ScopedCom.h"
+#include "ScopedCom.h"					// FScopedComString
 
 #include "Windows/AllowWindowsPlatformTypes.h"
-#include "Windows/AllowWindowsPlatformAtomics.h"
-
 THIRD_PARTY_INCLUDES_START
-#define INITGUID
-#include <mmdeviceapi.h>
-#include <audiopolicy.h>
-#include <functiondiscoverykeys_devpkey.h>
+// Linkage for  Windows GUIDs included by Notification/DeviceInfoCache, otherwise they are extern.
+#include <initguid.h>
 THIRD_PARTY_INCLUDES_END
+#include "Windows/HideWindowsPlatformTypes.h"
 
-#if !NO_LOGGING
-// PSStringFromPropertyKey needs this lib, which we use for logging the Property GUIDs.
-#pragma comment(lib, "Propsys.lib")
-#endif //!NO_LOGGING
-
-static Audio::EAudioDeviceState ConvertWordToDeviceState(DWORD InWord)
-{
-	switch (InWord)
-	{
-	case DEVICE_STATE_ACTIVE: return Audio::EAudioDeviceState::Active;
-	case DEVICE_STATE_DISABLED: return Audio::EAudioDeviceState::Disabled;
-	case DEVICE_STATE_UNPLUGGED: return Audio::EAudioDeviceState::Unplugged;
-	case DEVICE_STATE_NOTPRESENT: return Audio::EAudioDeviceState::NotPresent;
-	default:
-		break;
-	}
-	checkNoEntry();
-	return Audio::EAudioDeviceState::NotPresent;
-}
-
-
-class FWindowsMMNotificationClient final : public IMMNotificationClient, 
-										   public IAudioSessionEvents
-{
-public:
-	FWindowsMMNotificationClient()
-		: Ref(1)
-		, DeviceEnumerator(nullptr)
-		, bHasDisconnectSessionHappened(false)
-	{
-		bComInitialized = FWindowsPlatformMisc::CoInitialize();
-		HRESULT Result = CoCreateInstance(__uuidof(MMDeviceEnumerator), nullptr, CLSCTX_INPROC_SERVER, IID_PPV_ARGS(&DeviceEnumerator));
-		if (Result == S_OK)
-		{
-			DeviceEnumerator->RegisterEndpointNotificationCallback(this);
-		}
-
-		// Register for session events from default endpoint.
-		if (DeviceEnumerator)
-		{
-			TComPtr<IMMDevice> DefaultDevice;
-			if (SUCCEEDED(DeviceEnumerator->GetDefaultAudioEndpoint(eRender, eConsole, &DefaultDevice)))
-			{
-				RegisterForSessionNotifications(DefaultDevice);
-			}
-		}
-	}
-
-	bool RegisterForSessionNotifications(const TComPtr<IMMDevice>& InDevice)
-	{	
-		FScopeLock Lock(&SessionRegistrationCS);
-
-		// If we're already listening to this device, we can early out.
-		if (DeviceListeningToSessionEvents == InDevice)
-		{
-			return true;
-		}
-		
-		UnregisterForSessionNotifications();
-
-		DeviceListeningToSessionEvents = InDevice;
-
-		
-		if(InDevice) 
-		{
-			if (SUCCEEDED(InDevice->Activate(__uuidof(IAudioSessionManager), CLSCTX_INPROC_SERVER, NULL, (void**)&SessionManager)))
-			{
-				if (SUCCEEDED(SessionManager->GetAudioSessionControl(NULL, 0, &SessionControls)))
-				{
-					if (SUCCEEDED(SessionControls->RegisterAudioSessionNotification(this)))
-					{
-						UE_LOG(LogAudioMixer, Verbose, TEXT("FWindowsMMNotificationClient: Registering for sessions events for '%s'"), *GetFriendlyName(DeviceListeningToSessionEvents.Get()));
-						return true;
-					}
-				}
-			}
-		}
-		return false;
-	}
-
-	bool RegisterForSessionNotifications(const FString& InDeviceId)
-	{
-		if (TComPtr<IMMDevice> Device = GetDevice(InDeviceId))
-		{
-			return RegisterForSessionNotifications(Device);
-		}
-		return false;
-	}
-
-	void UnregisterForSessionNotifications()
-	{	
-		FScopeLock Lock(&SessionRegistrationCS);
-
-		// Unregister for any device we're already listening to.
-		if (SessionControls)
-		{
-			UE_LOG(LogAudioMixer, Verbose, TEXT("FWindowsMMNotificationClient: Unregistering for sessions events for device '%s'"), DeviceListeningToSessionEvents ? *GetFriendlyName(DeviceListeningToSessionEvents.Get()) : TEXT("None"));
-			SessionControls->UnregisterAudioSessionNotification(this);
-			SessionControls.Reset();
-		}
-		if (SessionManager)
-		{
-			SessionManager.Reset();
-		}
-		
-		DeviceListeningToSessionEvents.Reset();
-
-		// Reset this flag.
-		bHasDisconnectSessionHappened = false;
-	}
-	
-	~FWindowsMMNotificationClient()
-	{
-		UnregisterForSessionNotifications();
-
-		if (DeviceEnumerator)
-		{
-			DeviceEnumerator->UnregisterEndpointNotificationCallback(this);
-		}
-
-		if (bComInitialized)
-		{
-			FWindowsPlatformMisc::CoUninitialize();
-		}
-	}
-
-	#define CASE_TO_STRING(X) case X: return TEXT(#X)
-	static const TCHAR* ToString(EDataFlow InFlow)
-	{
-#if !NO_LOGGING
-		switch (InFlow)
-		{
-			CASE_TO_STRING(eRender);
-			CASE_TO_STRING(eCapture);
-			CASE_TO_STRING(eAll);
-			default:
-				checkNoEntry();
-				break;
-		}
-#endif //!NO_LOGGING
-		return TEXT("Unknown");
-	}
-	static const TCHAR* ToString(ERole InRole)
-	{
-#if !NO_LOGGING
-		switch (InRole)
-		{
-			CASE_TO_STRING(eConsole);
-			CASE_TO_STRING(eMultimedia);
-			CASE_TO_STRING(eCommunications);
-			default:
-				checkNoEntry();
-				break;
-		}
-#endif //!NO_LOGGING
-		return TEXT("Unknown");
-	}	
-	static const TCHAR* ToString(AudioSessionDisconnectReason InDisconnectReason)
-	{
-#if !NO_LOGGING
-		switch (InDisconnectReason)
-		{
-			CASE_TO_STRING(DisconnectReasonDeviceRemoval);
-			CASE_TO_STRING(DisconnectReasonServerShutdown);
-			CASE_TO_STRING(DisconnectReasonFormatChanged);
-			CASE_TO_STRING(DisconnectReasonSessionLogoff);
-			CASE_TO_STRING(DisconnectReasonSessionDisconnected);
-			CASE_TO_STRING(DisconnectReasonExclusiveModeOverride);
-		default:
-			checkNoEntry();
-			break;
-		}
-#endif //!NO_LOGGING
-		return TEXT("Unknown");
-	}
-	#undef CASE_TO_STRING
-
-	HRESULT STDMETHODCALLTYPE OnDefaultDeviceChanged(EDataFlow InFlow, ERole InRole, LPCWSTR pwstrDeviceId) override
-	{
-
-		UE_CLOG(Audio::IAudioMixer::ShouldLogDeviceSwaps(),LogAudioMixer, Warning, 
-			TEXT("FWindowsMMNotificationClient: OnDefaultDeviceChanged: %s, %s, %s - %s"), ToString(InFlow), ToString(InRole), pwstrDeviceId, *GetFriendlyName(pwstrDeviceId));
-
-		Audio::EAudioDeviceRole AudioDeviceRole;
-
-		if (Audio::IAudioMixer::ShouldIgnoreDeviceSwaps())
-		{
-			return S_OK;
-		}
-
-		if (InRole == eConsole)
-		{
-			AudioDeviceRole = Audio::EAudioDeviceRole::Console;
-		}
-		else if (InRole == eMultimedia)
-		{
-			AudioDeviceRole = Audio::EAudioDeviceRole::Multimedia;
-		}
-		else
-		{
-			AudioDeviceRole = Audio::EAudioDeviceRole::Communications;
-		}
-
-		FString DeviceString(pwstrDeviceId);
-		if (InFlow == eRender)
-		{
-			FReadScopeLock Lock(ListenersSetRwLock);
-			for (Audio::IAudioMixerDeviceChangedListener* Listener : Listeners)
-			{
-				Listener->OnDefaultRenderDeviceChanged(AudioDeviceRole, DeviceString);
-			}
-		}
-		else if (InFlow == eCapture)
-		{
-			FReadScopeLock Lock(ListenersSetRwLock);
-			for (Audio::IAudioMixerDeviceChangedListener* Listener : Listeners)
-			{
-				Listener->OnDefaultCaptureDeviceChanged(AudioDeviceRole, DeviceString);
-			}
-		}
-		else
-		{
-			FReadScopeLock Lock(ListenersSetRwLock);
-			for (Audio::IAudioMixerDeviceChangedListener* Listener : Listeners)
-			{
-				Listener->OnDefaultCaptureDeviceChanged(AudioDeviceRole, DeviceString);
-				Listener->OnDefaultRenderDeviceChanged(AudioDeviceRole, DeviceString);
-			}
-		}
-
-
-		return S_OK;
-	}
-
-	// TODO: Ideally we'd use the cache instead of ask for this.
-	bool IsRenderDevice(const FString& InDeviceId) const
-	{
-		bool bIsRender = true;
-		if (TComPtr<IMMDevice> Device = GetDevice(InDeviceId))
-		{
-			TComPtr<IMMEndpoint> Endpoint;
-			if (SUCCEEDED(Device->QueryInterface(IID_PPV_ARGS(&Endpoint))))
-			{
-				EDataFlow DataFlow = eRender;
-				if (SUCCEEDED(Endpoint->GetDataFlow(&DataFlow)))
-				{
-					bIsRender = DataFlow == eRender;
-				}
-			}
-		}
-		return bIsRender;
-	}
-
-	HRESULT STDMETHODCALLTYPE OnDeviceAdded(LPCWSTR pwstrDeviceId) override
-	{
-		UE_CLOG(Audio::IAudioMixer::ShouldLogDeviceSwaps(), LogAudioMixer, Display, TEXT("FWindowsMMNotificationClient: OnDeviceAdded: %s"), *GetFriendlyName(pwstrDeviceId));
-			
-		if (Audio::IAudioMixer::ShouldIgnoreDeviceSwaps())
-		{
-			return S_OK;
-		}
-		
-		bool bIsRender = IsRenderDevice(pwstrDeviceId);
-		FString DeviceString(pwstrDeviceId);
-		FReadScopeLock ReadLock(ListenersSetRwLock);
-		for (Audio::IAudioMixerDeviceChangedListener* Listener : Listeners)
-		{
-			Listener->OnDeviceAdded(DeviceString, bIsRender);
-		}
-		return S_OK;
-	};
-
-	HRESULT STDMETHODCALLTYPE OnDeviceRemoved(LPCWSTR pwstrDeviceId) override
-	{		
-		UE_CLOG(Audio::IAudioMixer::ShouldLogDeviceSwaps(),LogAudioMixer, Display, TEXT("FWindowsMMNotificationClient: OnDeviceRemoved: %s"), *GetFriendlyName(pwstrDeviceId));
-
-		if (Audio::IAudioMixer::ShouldIgnoreDeviceSwaps())
-		{
-			return S_OK;
-		}
-
-		bool bIsRender = IsRenderDevice(pwstrDeviceId);
-		FString DeviceString(pwstrDeviceId);
-		FReadScopeLock ReadLock(ListenersSetRwLock);		
-		for (Audio::IAudioMixerDeviceChangedListener* Listener : Listeners)
-		{
-			Listener->OnDeviceRemoved(DeviceString, bIsRender);
-		}
-		return S_OK;
-	}
-
-	HRESULT STDMETHODCALLTYPE OnDeviceStateChanged(LPCWSTR pwstrDeviceId, DWORD dwNewState) override
-	{
-		UE_CLOG(Audio::IAudioMixer::ShouldLogDeviceSwaps(), LogAudioMixer, Display, TEXT("FWindowsMMNotificationClient: OnDeviceStateChanged: %s, %d"), *GetFriendlyName(pwstrDeviceId), dwNewState);
-
-		if (Audio::IAudioMixer::ShouldIgnoreDeviceSwaps())
-		{
-			return S_OK;
-		}
-		
-		bool bIsRender = IsRenderDevice(pwstrDeviceId);
-		if (dwNewState == DEVICE_STATE_ACTIVE || dwNewState == DEVICE_STATE_DISABLED || dwNewState == DEVICE_STATE_UNPLUGGED || dwNewState == DEVICE_STATE_NOTPRESENT)
-		{
-			Audio::EAudioDeviceState State = ConvertWordToDeviceState(dwNewState);
-		
-			FString DeviceString(pwstrDeviceId);
-			FReadScopeLock ReadLock(ListenersSetRwLock);
-			for (Audio::IAudioMixerDeviceChangedListener* Listener : Listeners)
-			{
-				Listener->OnDeviceStateChanged(DeviceString, State, bIsRender);
-			}
-		}
-		return S_OK;
-	}
-
-	FString PropToString(const PROPERTYKEY Key)
-	{
-#if !NO_LOGGING
-		#define IF_PROP_STRING(PROP) if(Key.fmtid == PROP.fmtid) return TEXT(#PROP)
-
-		IF_PROP_STRING(PKEY_AudioEndpoint_PhysicalSpeakers);
-		IF_PROP_STRING(PKEY_AudioEngine_DeviceFormat);
-		IF_PROP_STRING(PKEY_AudioEngine_OEMFormat);
-		IF_PROP_STRING(PKEY_AudioEndpoint_Association);
-		IF_PROP_STRING(PKEY_AudioEndpoint_ControlPanelPageProvider);
-		IF_PROP_STRING(PKEY_AudioEndpoint_Disable_SysFx);
-		IF_PROP_STRING(PKEY_AudioEndpoint_FormFactor);
-		IF_PROP_STRING(PKEY_AudioEndpoint_FullRangeSpeakers);
-		IF_PROP_STRING(PKEY_AudioEndpoint_GUID);
-		IF_PROP_STRING(PKEY_AudioEndpoint_Supports_EventDriven_Mode);
-
-		#undef IF_PROP_STRING
-
-		TCHAR KeyString[PKEYSTR_MAX];
-		HRESULT HR = PSStringFromPropertyKey(Key, KeyString, ARRAYSIZE(KeyString));
-		if(SUCCEEDED(HR))
-		{
-			return FString(KeyString);
-		}
-#endif //!NO_LOGGING
-		return TEXT("Unknown");
-	}
-
-	FString GetFriendlyName(const FString InDeviceID)
-	{
-		if (InDeviceID.IsEmpty())
-		{
-			return TEXT("System Default");
-		}
-
-		FString FriendlyName = TEXT("[No Friendly Name for Device]");
-			
-		// Get device.
-		if (TComPtr<IMMDevice> Device = GetDevice(*InDeviceID))
-		{
-			return GetFriendlyName(Device);
-		}
-		return FriendlyName;
-	}
-	FString GetFriendlyName(const TComPtr<IMMDevice>& InDevice)
-	{
-		FString FriendlyName = TEXT("[No Friendly Name for Device]");
-
-		if (InDevice)
-		{
-			// Get property store.
-			TComPtr<IPropertyStore> PropStore;
-			HRESULT Hr = InDevice->OpenPropertyStore(STGM_READ, &PropStore);
-
-			// Get friendly name.
-			if (SUCCEEDED(Hr) && PropStore)
-			{
-				PROPVARIANT PropString;
-				PropVariantInit(&PropString);
-
-				// Get the endpoint device's friendly-name property.
-				Hr = PropStore->GetValue(PKEY_Device_FriendlyName, &PropString);
-				if (SUCCEEDED(Hr))
-				{
-					// Copy friendly name.
-					if (PropString.pwszVal)
-					{
-						FriendlyName = PropString.pwszVal;
-					}
-				}
-
-				PropVariantClear(&PropString);
-			}
-		}
-		return FriendlyName;
-	}
-
-	TComPtr<IMMDevice> GetDevice(const FString InDeviceID) const
-	{
-		// Get device.
-		TComPtr<IMMDevice> Device;
-		HRESULT Hr = DeviceEnumerator->GetDevice(*InDeviceID, &Device);
-		if (SUCCEEDED(Hr))
-		{
-			return Device;
-		}
-		
-		// Fail.
-		return {};
-	}
-
-	HRESULT STDMETHODCALLTYPE OnPropertyValueChanged(LPCWSTR pwstrDeviceId, const PROPERTYKEY key)
-	{
-		UE_CLOG(Audio::IAudioMixer::ShouldLogDeviceSwaps(), LogAudioMixer, Verbose, TEXT("OnPropertyValueChanged: %s : %s"), *GetFriendlyName(pwstrDeviceId), *PropToString(key));
-		
-		if (key.fmtid == PKEY_AudioEngine_DeviceFormat.fmtid )
-		{
-			// Get device.
-			FString DeviceId = pwstrDeviceId;
-			TComPtr<IMMDevice> Device;			
-			HRESULT Hr = DeviceEnumerator->GetDevice(*DeviceId, &Device);
-
-			// Get property store.
-			TComPtr<IPropertyStore> PropertyStore;
-			if (SUCCEEDED(Hr) && Device)
-			{
-				Hr = Device->OpenPropertyStore(STGM_READ, &PropertyStore);
-				if (SUCCEEDED(Hr) && PropertyStore)
-				{
-					// Device Format
-					PROPVARIANT Prop;
-					PropVariantInit(&Prop);
-
-					if (key.fmtid == PKEY_AudioEngine_DeviceFormat.fmtid )
-					{
-						// WAVEFORMATEX blobs.
-						if (SUCCEEDED(PropertyStore->GetValue(key, &Prop)) && Prop.blob.pBlobData)
-						{
-							const WAVEFORMATEX* WaveFormatEx = (const WAVEFORMATEX*)(Prop.blob.pBlobData);
-
-							Audio::IAudioMixerDeviceChangedListener::FFormatChangedData FormatChanged;
-							FormatChanged.NumChannels = FMath::Clamp((int32)WaveFormatEx->nChannels, 2, 8);
-							FormatChanged.SampleRate = WaveFormatEx->nSamplesPerSec;
-							FormatChanged.ChannelBitmask = WaveFormatEx->wFormatTag == WAVE_FORMAT_EXTENSIBLE ?
-								((const WAVEFORMATEXTENSIBLE*)WaveFormatEx)->dwChannelMask : 0;
-
-							FReadScopeLock ReadLock(ListenersSetRwLock);
-							for (Audio::IAudioMixerDeviceChangedListener* i : Listeners)
-							{
-								i->OnFormatChanged(DeviceId, FormatChanged);
-							}
-						}
-					}				
-					PropVariantClear(&Prop);
-				}
-			}
-		}
-		return S_OK;
-	}
-
-	HRESULT STDMETHODCALLTYPE QueryInterface(const IID& IId, void** UnknownPtrPtr) override
-	{
-		// Three rules of QueryInterface: https://docs.microsoft.com/en-us/windows/win32/com/rules-for-implementing-queryinterface
-		// 1. Objects must have identity.
-		// 2. The set of interfaces on an object instance must be static.
-		// 3. It must be possible to query successfully for any interface on an object from any other interface.
-
-		// If ppvObject(the address) is nullptr, then this method returns E_POINTER.
-		if (!UnknownPtrPtr)
-		{
-			return E_POINTER;
-		}
-
-		// https://docs.microsoft.com/en-us/windows/win32/com/implementing-reference-counting
-		// Whenever a client calls a method(or API function), such as QueryInterface, that returns a new interface pointer, 
-		// the method being called is responsible for incrementing the reference count through the returned pointer.
-		// For example, when a client first creates an object, it receives an interface pointer to an object that, 
-		// from the client's point of view, has a reference count of one. If the client then calls AddRef on the interface pointer, 
-		// the reference count becomes two. The client must call Release twice on the interface pointer to drop all of its references to the object.
-		if (IId == __uuidof(IMMNotificationClient) || IId == __uuidof(IUnknown))
-		{
-			*UnknownPtrPtr = (IMMNotificationClient*)(this);
-			AddRef();
-			return S_OK;
-		}
-		else if (IId == __uuidof(IAudioSessionEvents))
-		{
-			*UnknownPtrPtr = (IAudioSessionEvents*)this;
-			AddRef();
-			return S_OK;
-		}
-
-
-		// This method returns S_OK if the interface is supported, and E_NOINTERFACE otherwise.
-		*UnknownPtrPtr = nullptr;
-		return E_NOINTERFACE;
-	}
-
-	ULONG STDMETHODCALLTYPE AddRef() override
-	{
-		return InterlockedIncrement(&Ref);
-	}
-
-	ULONG STDMETHODCALLTYPE Release() override
-	{
-		ULONG ulRef = InterlockedDecrement(&Ref);
-		if (0 == ulRef)
-		{
-			delete this;
-		}
-		return ulRef;
-	}
-
-	void RegisterDeviceChangedListener(Audio::IAudioMixerDeviceChangedListener* DeviceChangedListener)
-	{
-		// Modifying container so get full write lock
-		FWriteScopeLock Lock(ListenersSetRwLock);
-		Listeners.Add(DeviceChangedListener);
-	}
-
-	void UnRegisterDeviceDeviceChangedListener(Audio::IAudioMixerDeviceChangedListener* DeviceChangedListener)
-	{
-		// Modifying container so get full write lock
-		FWriteScopeLock Lock(ListenersSetRwLock);
-		Listeners.Remove(DeviceChangedListener);
-	}
-	
-	// Begin IAudioSessionEvents
-	HRESULT STDMETHODCALLTYPE OnDisplayNameChanged(
-		LPCWSTR NewDisplayName,
-		LPCGUID EventContext) override
-	{
-		return S_OK;
-	}
-
-	HRESULT STDMETHODCALLTYPE OnIconPathChanged(
-		LPCWSTR NewIconPath,
-		LPCGUID EventContext) override
-	{
-		return S_OK;
-	}
-
-	HRESULT STDMETHODCALLTYPE OnSimpleVolumeChanged(
-		float NewVolume,
-		BOOL NewMute,
-		LPCGUID EventContext) override
-	{
-		return S_OK;
-	}
-
-	HRESULT STDMETHODCALLTYPE OnChannelVolumeChanged(
-		DWORD ChannelCount,
-		float NewChannelVolumeArray[],
-		DWORD ChangedChannel,
-		LPCGUID EventContext) override
-	{
-		return S_OK;
-	}
-
-	HRESULT STDMETHODCALLTYPE OnGroupingParamChanged(
-		LPCGUID NewGroupingParam,
-		LPCGUID EventContext) override
-	{
-		return S_OK;
-	}
-
-	HRESULT STDMETHODCALLTYPE OnStateChanged(
-		AudioSessionState NewState) override
-	{		
-		return S_OK;
-	}
-
-	HRESULT STDMETHODCALLTYPE OnSessionDisconnected(
-		AudioSessionDisconnectReason InDisconnectReason)
-	{	
-		UE_CLOG(Audio::IAudioMixer::ShouldLogDeviceSwaps(), LogAudioMixer, Verbose, TEXT("Session Disconnect: Reason=%s, DeviceBound=%s, HasDisconnectSessionHappened=%d"), 
-			ToString(InDisconnectReason), *GetFriendlyName(DeviceListeningToSessionEvents), (int32)bHasDisconnectSessionHappened);
-
-
-		if (!bHasDisconnectSessionHappened)
-		{
-			{
-				FReadScopeLock Lock(ListenersSetRwLock);
-				Audio::IAudioMixerDeviceChangedListener::EDisconnectReason Reason = AudioSessionDisconnectToEDisconnectReason(InDisconnectReason);
-				for (Audio::IAudioMixerDeviceChangedListener* i : Listeners)
-				{
-					i->OnSessionDisconnect(Reason);
-				}
-			}
-			
-			// Mark this true.
-			bHasDisconnectSessionHappened = true;
-		}
-
-		return S_OK;
-	}
-	// End IAudioSessionEvents
-	
-	static Audio::IAudioMixerDeviceChangedListener::EDisconnectReason 
-	AudioSessionDisconnectToEDisconnectReason(AudioSessionDisconnectReason InDisconnectReason)
-	{
-		using namespace Audio;
-		switch (InDisconnectReason)
-		{
-		case DisconnectReasonDeviceRemoval:			return IAudioMixerDeviceChangedListener::EDisconnectReason::DeviceRemoval;
-		case DisconnectReasonServerShutdown:		return IAudioMixerDeviceChangedListener::EDisconnectReason::ServerShutdown;
-		case DisconnectReasonFormatChanged:			return IAudioMixerDeviceChangedListener::EDisconnectReason::FormatChanged;
-		case DisconnectReasonSessionLogoff:			return IAudioMixerDeviceChangedListener::EDisconnectReason::SessionLogoff;
-		case DisconnectReasonSessionDisconnected:	return IAudioMixerDeviceChangedListener::EDisconnectReason::SessionDisconnected;
-		case DisconnectReasonExclusiveModeOverride:	return IAudioMixerDeviceChangedListener::EDisconnectReason::ExclusiveModeOverride;
-		default: break;
-		}
-
-		checkNoEntry();
-		return IAudioMixerDeviceChangedListener::EDisconnectReason::DeviceRemoval;
-	}
-		
-private:
-	LONG Ref;
-	TSet<Audio::IAudioMixerDeviceChangedListener*> Listeners;
-	FRWLock ListenersSetRwLock;
-	
-	TComPtr<IMMDeviceEnumerator> DeviceEnumerator;
-
-	FCriticalSection SessionRegistrationCS;
-	TComPtr<IAudioSessionManager> SessionManager;
-	TComPtr<IAudioSessionControl> SessionControls;
-	TComPtr<IMMDevice> DeviceListeningToSessionEvents;
-
-	bool bComInitialized;
-	std::atomic<bool> bHasDisconnectSessionHappened;
-};
+#include "WindowsMMNotificationClient.h"
+#include "WindowsMMDeviceInfoCache.h"
+#include "ToStringHelpers.h"
 
 namespace Audio
 {	
@@ -673,657 +48,7 @@ namespace Audio
 			WindowsNotificationClient->UnregisterForSessionNotifications();
 		}
 	}
-
-	struct FWindowsMMDeviceCache : IAudioMixerDeviceChangedListener, IAudioPlatformDeviceInfoCache
-	{
-		struct FCacheEntry
-		{
-			enum class EEndpointType { Unknown, Render, Capture};
-
-			FName DeviceId;							// Key
-			FString FriendlyName;
-			FString DeviceFriendlyName;
-			EAudioDeviceState State;
-			int32 NumChannels = 0;
-			int32 SampleRate = 0;
-			EEndpointType Type = EEndpointType::Unknown;
-			uint32 ChannelBitmask = 0;				// Bitfield used to build output channels, for easy comparison.
-
-			TArray<EAudioMixerChannel::Type> OutputChannels;	// TODO. Generate this from the ChannelNum and bitmask when we are asked for it.
-			mutable FRWLock MutationLock;
-
-			FCacheEntry& operator=(const FCacheEntry& InOther)
-			{
-				// Copy everything but the lock. 
-				DeviceId = InOther.DeviceId;
-				FriendlyName = InOther.FriendlyName;
-				DeviceFriendlyName = InOther.DeviceFriendlyName;
-				State = InOther.State;
-				NumChannels = InOther.NumChannels;
-				SampleRate = InOther.SampleRate;
-				Type = InOther.Type;
-				ChannelBitmask = InOther.ChannelBitmask;
-				OutputChannels = InOther.OutputChannels;
-				return *this;
-			}
-
-			FCacheEntry& operator=(FCacheEntry&& InOther)
-			{
-				DeviceId = MoveTemp(InOther.DeviceId);
-				FriendlyName = MoveTemp(InOther.FriendlyName);
-				DeviceFriendlyName = MoveTemp(InOther.DeviceFriendlyName);
-				State = MoveTemp(InOther.State);
-				NumChannels = MoveTemp(InOther.NumChannels);
-				SampleRate = MoveTemp(InOther.SampleRate);
-				Type = MoveTemp(InOther.Type);
-				ChannelBitmask = MoveTemp(InOther.ChannelBitmask);
-				OutputChannels = MoveTemp(InOther.OutputChannels);
-				return *this;
-			}
-
-			FCacheEntry(const FCacheEntry& InOther)
-			{
-				*this = InOther;
-			}
-
-			FCacheEntry(FCacheEntry&& InOther)
-			{
-				*this = MoveTemp(InOther);
-			}
-
-			FCacheEntry(const FString& InDeviceId)
-				: DeviceId{ InDeviceId }
-			{}
-		};
-
-		TComPtr<IMMDeviceEnumerator> DeviceEnumerator;
-
-		mutable FRWLock CacheMutationLock;							// R/W lock protects map and default arrays.
-		TMap<FName, FCacheEntry> Cache;								// DeviceID GUID -> Info.
-		FName DefaultCaptureId[(int32)EAudioDeviceRole::COUNT];		// Role -> DeviceID GUID
-		FName DefaultRenderId[(int32)EAudioDeviceRole::COUNT];		// Role -> DeviceID GUID
-
-		FWindowsMMDeviceCache()
-		{
-			ensure(SUCCEEDED(CoCreateInstance(__uuidof(MMDeviceEnumerator), NULL, CLSCTX_INPROC_SERVER, IID_PPV_ARGS(&DeviceEnumerator))) && DeviceEnumerator);
-
-			EnumerateEndpoints();
-			EnumerateDefaults();
-		}
-		virtual ~FWindowsMMDeviceCache() = default;
-
-		bool EnumerateChannelMask(uint32 InMask, FCacheEntry& OutInfo)
-		{
-			// Loop through the extensible format channel flags in the standard order and build our output channel array
-			// From https://msdn.microsoft.com/en-us/library/windows/hardware/dn653308(v=vs.85).aspx
-			// The channels in the interleaved stream corresponding to these spatial positions must appear in the order specified above. This holds true even in the 
-			// case of a non-contiguous subset of channels. For example, if a stream contains left, bass enhance and right, then channel 1 is left, channel 2 is right, 
-			// and channel 3 is bass enhance. This enables the linkage of multi-channel streams to well-defined multi-speaker configurations.
-
-			static const uint32 ChannelTypeMap[EAudioMixerChannel::ChannelTypeCount] =
-			{
-				SPEAKER_FRONT_LEFT,
-				SPEAKER_FRONT_RIGHT,
-				SPEAKER_FRONT_CENTER,
-				SPEAKER_LOW_FREQUENCY,
-				SPEAKER_BACK_LEFT,
-				SPEAKER_BACK_RIGHT,
-				SPEAKER_FRONT_LEFT_OF_CENTER,
-				SPEAKER_FRONT_RIGHT_OF_CENTER,
-				SPEAKER_BACK_CENTER,
-				SPEAKER_SIDE_LEFT,
-				SPEAKER_SIDE_RIGHT,
-				SPEAKER_TOP_CENTER,
-				SPEAKER_TOP_FRONT_LEFT,
-				SPEAKER_TOP_FRONT_CENTER,
-				SPEAKER_TOP_FRONT_RIGHT,
-				SPEAKER_TOP_BACK_LEFT,
-				SPEAKER_TOP_BACK_CENTER,
-				SPEAKER_TOP_BACK_RIGHT,
-				SPEAKER_RESERVED,
-			};
-
-			OutInfo.ChannelBitmask = InMask;
-			OutInfo.OutputChannels.Reset();
-
-			// No need to enumerate speakers for capture devices.
-			if (OutInfo.Type == FCacheEntry::EEndpointType::Capture)
-			{
-				return true;
-			}
-
-			uint32 ChanCount = 0;
-			for (uint32 ChannelTypeIndex = 0; ChannelTypeIndex < EAudioMixerChannel::ChannelTypeCount && ChanCount < (uint32)OutInfo.NumChannels; ++ChannelTypeIndex)
-			{
-				if (InMask & ChannelTypeMap[ChannelTypeIndex])
-				{
-					OutInfo.OutputChannels.Add((EAudioMixerChannel::Type)ChannelTypeIndex);
-					++ChanCount;
-				}
-			}
-
-			// We didn't match channel masks for all channels, revert to a default ordering
-			if (ChanCount < (uint32)OutInfo.NumChannels)
-			{
-				UE_CLOG(Audio::IAudioMixer::ShouldLogDeviceSwaps(), LogAudioMixer, Warning, TEXT("FWindowsMMDeviceCache: Did not find the channel type flags for audio device '%s'. Reverting to a default channel ordering."), *OutInfo.FriendlyName);
-
-				OutInfo.OutputChannels.Reset();
-
-				static const EAudioMixerChannel::Type DefaultChannelOrdering[] = {
-					EAudioMixerChannel::FrontLeft,
-					EAudioMixerChannel::FrontRight,
-					EAudioMixerChannel::FrontCenter,
-					EAudioMixerChannel::LowFrequency,
-					EAudioMixerChannel::SideLeft,
-					EAudioMixerChannel::SideRight,
-					EAudioMixerChannel::BackLeft,
-					EAudioMixerChannel::BackRight,
-				};
-
-				const EAudioMixerChannel::Type* ChannelOrdering = DefaultChannelOrdering;
-
-				// Override channel ordering for some special cases
-				if (OutInfo.NumChannels == 4)
-				{
-					static EAudioMixerChannel::Type DefaultChannelOrderingQuad[] = {
-						EAudioMixerChannel::FrontLeft,
-						EAudioMixerChannel::FrontRight,
-						EAudioMixerChannel::BackLeft,
-						EAudioMixerChannel::BackRight,
-					};
-
-					ChannelOrdering = DefaultChannelOrderingQuad;
-				}
-				else if (OutInfo.NumChannels == 6)
-				{
-					static const EAudioMixerChannel::Type DefaultChannelOrdering51[] = {
-						EAudioMixerChannel::FrontLeft,
-						EAudioMixerChannel::FrontRight,
-						EAudioMixerChannel::FrontCenter,
-						EAudioMixerChannel::LowFrequency,
-						EAudioMixerChannel::BackLeft,
-						EAudioMixerChannel::BackRight,
-					};
-
-					ChannelOrdering = DefaultChannelOrdering51;
-				}
-
-				check(OutInfo.NumChannels <= 8);
-				for (int32 Index = 0; Index < OutInfo.NumChannels; ++Index)
-				{
-					OutInfo.OutputChannels.Add(ChannelOrdering[Index]);
-				}
-			}
-			return true;
-		}
-
-		bool EnumerateChannelFormat(const WAVEFORMATEX* InFormat, FCacheEntry& OutInfo)
-		{
-			OutInfo.OutputChannels.Empty();
-
-			// Extensible format supports surround sound so we need to parse the channel configuration to build our channel output array
-			if (InFormat->wFormatTag == WAVE_FORMAT_EXTENSIBLE)
-			{
-				// Cast to the extensible format to get access to extensible data
-				const WAVEFORMATEXTENSIBLE* WaveFormatExtensible = (const WAVEFORMATEXTENSIBLE*)(InFormat);
-				return EnumerateChannelMask(WaveFormatExtensible->dwChannelMask, OutInfo);
-			}
-			else
-			{
-				// Non-extensible formats only support mono or stereo channel output
-				OutInfo.OutputChannels.Add(EAudioMixerChannel::FrontLeft);
-				if (OutInfo.NumChannels == 2)
-				{
-					OutInfo.OutputChannels.Add(EAudioMixerChannel::FrontRight);
-				}
-			}
-
-			// Aways success for now.
-			return true;
-		}
-
-		FCacheEntry::EEndpointType QueryDeviceDataFlow(const TComPtr<IMMDevice>& InDevice) const
-		{
-			TComPtr<IMMEndpoint> Endpoint;
-			if (SUCCEEDED(InDevice->QueryInterface(IID_PPV_ARGS(&Endpoint))))
-			{
-				EDataFlow DataFlow = eRender;
-				if (SUCCEEDED(Endpoint->GetDataFlow(&DataFlow)))
-				{
-					switch (DataFlow)
-					{
-					case eRender:
-						return FCacheEntry::EEndpointType::Render;
-					case eCapture:
-						return FCacheEntry::EEndpointType::Capture;
-					default:
-						break;
-					}
-				}
-			}
-			return FCacheEntry::EEndpointType::Unknown;
-		}
-
-		bool EnumerateDeviceProps(const TComPtr<IMMDevice>& InDevice, FCacheEntry& OutInfo)
-		{
-			// Mark if this is a Render Device or Capture or Unknown.
-			OutInfo.Type = QueryDeviceDataFlow(InDevice);
-
-			// Also query the device state.
-			DWORD DeviceState = DEVICE_STATE_NOTPRESENT;
-			if (SUCCEEDED(InDevice->GetState(&DeviceState)))
-			{
-				OutInfo.State = ConvertWordToDeviceState(DeviceState);
-			}
-
-			TComPtr<IPropertyStore> PropertyStore;
-			if (SUCCEEDED(InDevice->OpenPropertyStore(STGM_READ, &PropertyStore)))
-			{
-				// Friendly Name
-				PROPVARIANT FriendlyName;
-				PropVariantInit(&FriendlyName);
-				if (SUCCEEDED(PropertyStore->GetValue(PKEY_Device_FriendlyName, &FriendlyName)) && FriendlyName.pwszVal)
-				{
-					OutInfo.FriendlyName = FString(FriendlyName.pwszVal);
-					PropVariantClear(&FriendlyName);
-				}
-
-				auto EnumDeviceFormat = [this](const TComPtr<IPropertyStore>& InPropStore, REFPROPERTYKEY InKey, FCacheEntry& OutInfo) -> bool
-				{
-					// Device Format
-					PROPVARIANT DeviceFormat;
-					PropVariantInit(&DeviceFormat);
-
-					if (SUCCEEDED(InPropStore->GetValue(InKey, &DeviceFormat)) && DeviceFormat.blob.pBlobData)
-					{
-						const WAVEFORMATEX* WaveFormatEx = (const WAVEFORMATEX*)(DeviceFormat.blob.pBlobData);
-						OutInfo.NumChannels = FMath::Clamp((int32)WaveFormatEx->nChannels, 2, 8);
-						OutInfo.SampleRate = WaveFormatEx->nSamplesPerSec;
-
-						EnumerateChannelFormat(WaveFormatEx, OutInfo);
-						PropVariantClear(&DeviceFormat);
-						return true;
-					}
-					return false;
-				};
-
-				if (EnumDeviceFormat(PropertyStore, PKEY_AudioEngine_DeviceFormat, OutInfo) ||
-					EnumDeviceFormat(PropertyStore, PKEY_AudioEngine_OEMFormat, OutInfo))
-				{
-				}
-				else
-				{
-					// Log a warning if this device is active as we failed to ask for a format
-					UE_CLOG(DeviceState == DEVICE_STATE_ACTIVE, LogAudioMixer, Warning, TEXT("FWindowsMMDeviceCache: Failed to get Format for active device '%s'"), *OutInfo.FriendlyName);
-				}
-			}
-	
-			// Aways success for now.
-			return true;
-		}
-
-		void EnumerateEndpoints()
-		{
-			// Build a new cache from scratch.
-			TMap<FName, FCacheEntry> NewCache;
-
-			// Get Device Enumerator.
-			if (DeviceEnumerator)
-			{
-				// Get Render Device Collection. (note we ask for ALL states, which include disabled/unplugged devices.).
-				TComPtr<IMMDeviceCollection> DeviceCollection;
-				uint32 DeviceCount = 0;
-				if (SUCCEEDED(DeviceEnumerator->EnumAudioEndpoints(eAll, DEVICE_STATEMASK_ALL, &DeviceCollection)) && DeviceCollection &&
-					SUCCEEDED(DeviceCollection->GetCount(&DeviceCount)))
-				{
-					for (uint32 i = 0; i < DeviceCount; ++i)
-					{
-						TComPtr<IMMDevice> Device;
-						if (SUCCEEDED(DeviceCollection->Item(i, &Device)) && Device)
-						{
-							// Get the device id string (guid)
-							Audio::FScopeComString DeviceIdString;
-							if (SUCCEEDED(Device->GetId(&DeviceIdString.StringPtr)) && DeviceIdString)
-							{
-								FCacheEntry Info{ DeviceIdString.Get() };
-
-								// Enumerate props into our info object.
-								EnumerateDeviceProps(Device, Info);
-
-								UE_LOG(LogAudioMixer, Verbose, TEXT("FWindowsMMDeviceCache: %s Device '%s' ID='%s'"),
-									Info.Type == FCacheEntry::EEndpointType::Capture ? TEXT("Capture") :
-									Info.Type == FCacheEntry::EEndpointType::Render ? TEXT("Render") :
-									TEXT("UNKNOWN!"),
-									*Info.DeviceId.ToString(),
-									*Info.FriendlyName
-								);
-
-								check(!NewCache.Contains(Info.DeviceId));
-								NewCache.Emplace(Info.DeviceId, Info);
-							}
-						}
-					}
-				}
-
-				// Finally, Replace cache with new one.
-				{
-					FWriteScopeLock Lock(CacheMutationLock);
-					Cache = MoveTemp(NewCache);
-				}
-			}
-		}			
-
-		void EnumerateDefaults()
-		{
-			auto GetDefaultDeviceID = [this](EDataFlow InDataFlow, ERole InRole, FName& OutDeviceId) -> bool
-			{
-				// Mark default device.
-				bool bSuccess = false;
-				TComPtr<IMMDevice> DefaultDevice;
-				if (SUCCEEDED(DeviceEnumerator->GetDefaultAudioEndpoint(InDataFlow, InRole, &DefaultDevice)))
-				{
-					Audio::FScopeComString DeviceIdString;
-					if (SUCCEEDED(DefaultDevice->GetId(&DeviceIdString.StringPtr)) && DeviceIdString)
-					{	
-						OutDeviceId = DeviceIdString.Get();
-						bSuccess = true;
-					}
-				}
-				return bSuccess;
-			};
-
-			// Get defaults (render, capture).
-			FWriteScopeLock Lock(CacheMutationLock);
-			static_assert((int32)EAudioDeviceRole::COUNT == ERole_enum_count, "EAudioDeviceRole should be the same as ERole");
-			for (int32 i = 0; i < ERole_enum_count; ++i)
-			{
-				FName DeviceIdName;
-				if (GetDefaultDeviceID(eRender, static_cast<ERole>(i), DeviceIdName))
-				{
-					UE_CLOG(!DeviceIdName.IsNone(), LogAudioMixer, Verbose, TEXT("FWindowsMMDeviceCache: Default Render Role='%s', Device='%s'"), ToString((EAudioDeviceRole)i), *GetFriendlyName(DeviceIdName));
-					DefaultRenderId[i] = DeviceIdName;
-				}
-				if (GetDefaultDeviceID(eCapture, static_cast<ERole>(i), DeviceIdName))
-				{
-					UE_CLOG(!DeviceIdName.IsNone(), LogAudioMixer, Verbose, TEXT("FWindowsMMDeviceCache: Default Capture Role='%s', Device='%s'"), ToString((EAudioDeviceRole)i), *GetFriendlyName(DeviceIdName));
-					DefaultCaptureId[i] = DeviceIdName;
-				}
-			}
-		}
-
-		void OnDefaultCaptureDeviceChanged(const EAudioDeviceRole InAudioDeviceRole, const FString& DeviceId) override
-		{
-			FWriteScopeLock WriteLock(CacheMutationLock);
-			check(InAudioDeviceRole < EAudioDeviceRole::COUNT);
-			DefaultCaptureId[(int32)InAudioDeviceRole] = *DeviceId;
-		}
-
-		void OnDefaultRenderDeviceChanged(const EAudioDeviceRole InAudioDeviceRole, const FString& DeviceId) override
-		{
-			FWriteScopeLock WriteLock(CacheMutationLock);
-			check(InAudioDeviceRole < EAudioDeviceRole::COUNT);
-			DefaultRenderId[(int32)InAudioDeviceRole] = *DeviceId;
-		}
-
-		void OnDeviceAdded(const FString& DeviceId, bool bIsRender) override
-		{
-			if (ensure(DeviceEnumerator))
-			{
-				TComPtr<IMMDevice> Device;
-				if (SUCCEEDED(DeviceEnumerator->GetDevice(*DeviceId, &Device)))
-				{
-					FCacheEntry Info{ *DeviceId };
-					if (EnumerateDeviceProps(Device, Info))
-					{
-						FWriteScopeLock WriteLock(CacheMutationLock);
-						check(!Cache.Contains(Info.DeviceId));
-						Cache.Add(Info.DeviceId, Info);
-					}
-				}
-			}
-		}
-
-		void OnDeviceRemoved(const FString& DeviceId, bool) override
-		{
-			FWriteScopeLock WriteLock(CacheMutationLock);
-			FName DeviceIdName = *DeviceId;
-			check(Cache.Contains(DeviceIdName));
-			Cache.Remove(DeviceIdName);
-		}
-
-		TOptional<FCacheEntry> BuildCacheEntry(const FString& DeviceId)
-		{
-			TComPtr<IMMDevice> Device;
-			if (SUCCEEDED(DeviceEnumerator->GetDevice(*DeviceId, &Device)))
-			{
-				FCacheEntry Info{ *DeviceId };
-				if (EnumerateDeviceProps(Device, Info))
-				{
-					return Info;
-				}
-			}
-			return {};
-		}
-
-		FString GetFriendlyName(FName InDeviceId) const
-		{
-			if (const FCacheEntry* Entry = Cache.Find(InDeviceId))
-			{
-				return Entry->FriendlyName;
-			}
-			return TEXT("Unknown");
-		}
-
-		static const TCHAR* ToString(EAudioDeviceState InState)
-		{
-			#define CASE_TO_STRING(X) case EAudioDeviceState::X: return TEXT(#X)
-			switch (InState)
-			{
-				CASE_TO_STRING(Active);
-				CASE_TO_STRING(Disabled);
-				CASE_TO_STRING(NotPresent);
-				CASE_TO_STRING(Unplugged);
-			default:
-				return TEXT("Unknown");
-			}
-			#undef CASE_TO_STRING
-		}
-		static const TCHAR* ToString(EAudioDeviceRole InRole)
-		{
-			#define CASE_TO_STRING(X) case EAudioDeviceRole::X: return TEXT(#X)
-			switch (InRole)
-			{
-				CASE_TO_STRING(Console);
-				CASE_TO_STRING(Multimedia);
-				CASE_TO_STRING(Communications);
-			default:
-				return TEXT("Unknown");
-			}
-			#undef CASE_TO_STRING
-		}
-		static const FString ToString(const TArray<EAudioMixerChannel::Type>& InChannels)
-		{
-			FString ChannelString;
-			static const int32 ApproxChannelNameLength = 18;
-			ChannelString.Reserve(ApproxChannelNameLength * InChannels.Num());
-			for (EAudioMixerChannel::Type i : InChannels)
-			{
-				ChannelString.Append(EAudioMixerChannel::ToString(i));
-				ChannelString.Append(TEXT("|"));
-			}
-			return ChannelString;
-		}
-
-		void OnDeviceStateChanged(const FString& DeviceId, const EAudioDeviceState InState, bool ) override
-		{
-			TOptional<FCacheEntry> Info = BuildCacheEntry(DeviceId);					
-			
-			FReadScopeLock ReadLock(CacheMutationLock);
-			FName DeviceIdName = *DeviceId;
-			ensureMsgf(Cache.Contains(DeviceIdName), TEXT("Expecting to find '%s' in cache '%s'"), *DeviceId, Info ? *Info->FriendlyName : TEXT("Unknown"));
-
-			if (FCacheEntry* Entry = Cache.Find(DeviceIdName))
-			{
-				FWriteScopeLock WriteLock(Entry->MutationLock);
-
-				UE_CLOG(Audio::IAudioMixer::ShouldLogDeviceSwaps(), LogAudioMixer, Verbose, TEXT("FWindowsMMDeviceCache: DeviceName='%s' - DeviceID='%s' state changed from '%s' to '%s'."),
-					Info ? *Info->FriendlyName : TEXT("Unknown"), *DeviceId, ToString(Entry->State), ToString(InState));
-
-				Entry->State = InState;
-			}
-		}
-
-		void OnFormatChanged(const FString& InDeviceId, const FFormatChangedData& InFormat) override
-		{
-			FName DeviceName(InDeviceId);
-			bool bNeedToEnumerateChannels = false;
-			bool bDirty = false;
-		
-			FReadScopeLock MapReadLock(CacheMutationLock);
-			if (FCacheEntry* Found = Cache.Find(DeviceName))
-			{
-				// Make a copy of the entry
-				FCacheEntry EntryCopy(InDeviceId);
-				{
-					FReadScopeLock FoundReadLock(Found->MutationLock);
-					EntryCopy = *Found;
-				}
-
-				if (EntryCopy.NumChannels != InFormat.NumChannels)
-				{
-					UE_CLOG(Audio::IAudioMixer::ShouldLogDeviceSwaps(),LogAudioMixer, Verbose, TEXT("FWindowsMMDeviceCache: DeviceID='%s', Name='%s' changed default format from %d channels to %d."), *InDeviceId, *EntryCopy.FriendlyName, EntryCopy.NumChannels, InFormat.NumChannels);
-					EntryCopy.NumChannels = InFormat.NumChannels;
-					bNeedToEnumerateChannels = true;
-					bDirty = true;
-				}
-				if (EntryCopy.SampleRate != InFormat.SampleRate)
-				{
-					UE_CLOG(Audio::IAudioMixer::ShouldLogDeviceSwaps(),LogAudioMixer, Verbose, TEXT("FWindowsMMDeviceCache: DeviceID='%s', Name='%s' changed default format from %dhz to %dhz."), *InDeviceId, *EntryCopy.FriendlyName, EntryCopy.SampleRate, InFormat.SampleRate);
-					EntryCopy.SampleRate = InFormat.SampleRate;
-					bDirty = true;
-				}
-				if (EntryCopy.ChannelBitmask != InFormat.ChannelBitmask)
-				{
-					UE_CLOG(Audio::IAudioMixer::ShouldLogDeviceSwaps(),LogAudioMixer, Verbose, TEXT("FWindowsMMDeviceCache: DeviceID='%s', Name='%s' changed default format from 0x%x to 0x%x bitmask"), *InDeviceId, *EntryCopy.FriendlyName, EntryCopy.ChannelBitmask, InFormat.ChannelBitmask);
-					EntryCopy.ChannelBitmask = InFormat.ChannelBitmask;
-					bNeedToEnumerateChannels = true;
-					bDirty = true;
-				}
-
-				if (bNeedToEnumerateChannels)
-				{
-					UE_CLOG(Audio::IAudioMixer::ShouldLogDeviceSwaps(), LogAudioMixer, Verbose, TEXT("FWindowsMMDeviceCache: Channel Change, DeviceID='%s', Name='%s' OLD=[%s]"), *InDeviceId, *EntryCopy.FriendlyName, *ToString(EntryCopy.OutputChannels));
-					EnumerateChannelMask(InFormat.ChannelBitmask, EntryCopy);
-					UE_CLOG(Audio::IAudioMixer::ShouldLogDeviceSwaps(), LogAudioMixer, Verbose, TEXT("FWindowsMMDeviceCache: Channel Change, DeviceID='%s', Name='%s' NEW=[%s]"), *InDeviceId, *EntryCopy.FriendlyName, *ToString(EntryCopy.OutputChannels));
-				}
-
-				// Update the entire entry with one write.
-				if (bDirty)
-				{
-					FWriteScopeLock FoundWriteLock(Found->MutationLock);
-					*Found = EntryCopy;
-				}
-			}
-		}
-			
-		void MakeDeviceInfo(const FCacheEntry& InEntry, FName InDefaultDevice, FAudioPlatformDeviceInfo& OutInfo) const
-		{
-			OutInfo.Reset();
-			OutInfo.Name				= InEntry.FriendlyName;
-			OutInfo.DeviceId			= InEntry.DeviceId.GetPlainNameString();
-			OutInfo.NumChannels			= InEntry.NumChannels;
-			OutInfo.SampleRate			= InEntry.SampleRate;
-			OutInfo.OutputChannelArray	= InEntry.OutputChannels;
-			OutInfo.Format				= EAudioMixerStreamDataFormat::Float;
-			OutInfo.bIsSystemDefault	= InEntry.DeviceId == InDefaultDevice;
-		}
-	
-		virtual TArray<FAudioPlatformDeviceInfo> GetAllActiveOutputDevices() const override
-		{
-			SCOPED_NAMED_EVENT(FWindowsMMDeviceCache_GetAllActiveOutputDevices, FColor::Blue);
-
-			// Find all active devices.
-			TArray<FAudioPlatformDeviceInfo> ActiveDevices;
-
-			// Read lock
-			FReadScopeLock ReadLock(CacheMutationLock);
-			ActiveDevices.Reserve(Cache.Num());
-
-			// Ask for defaults once, as we are inside a read lock.
-			FName DefaultRenderDeviceId = GetDefaultOutputDevice_NoLock();
-
-			// Walk cache, read lock for each entry.
-			for (const auto& i : Cache)
-			{
-				// Read lock on each entry.
-				FReadScopeLock CacheEntryReadLock(i.Value.MutationLock);
-				if (i.Value.State == EAudioDeviceState::Active && 
-					i.Value.Type == FCacheEntry::EEndpointType::Render )
-				{
-					FAudioPlatformDeviceInfo& Info = ActiveDevices.Emplace_GetRef();
-					MakeDeviceInfo(i.Value, DefaultRenderDeviceId, Info);
-				}
-			}
-
-			// RVO
-			return ActiveDevices;
-		}
-
-		FName GetDefaultOutputDevice_NoLock() const
-		{
-			if (!DefaultRenderId[(int32)EAudioDeviceRole::Console].IsNone())
-			{
-				return DefaultRenderId[(int32)EAudioDeviceRole::Console];
-			}
-			if( !DefaultRenderId[(int32)EAudioDeviceRole::Multimedia].IsNone())
-			{
-				return DefaultRenderId[(int32)EAudioDeviceRole::Multimedia];
-			}
-			return NAME_None;
-		}
 		
-		TOptional<FAudioPlatformDeviceInfo> FindDefaultOutputDevice() const override
-		{
-			return FindActiveOutputDevice(NAME_None);
-		}
-
-		TOptional<FAudioPlatformDeviceInfo> FindActiveOutputDevice(FName InDeviceID) const override
-		{
-			SCOPED_NAMED_EVENT(FWindowsMMDeviceCache_FindActiveOutputDevice, FColor::Blue);
-
-			FReadScopeLock MapReadLock(CacheMutationLock);
-			
-			// Ask for default here as we are inside the read lock.
-			const FName DefaultOutputDevice = GetDefaultOutputDevice_NoLock();
-
-			// Asking for Default?
-			if (InDeviceID.IsNone())
-			{
-				InDeviceID = DefaultOutputDevice;
-				if (InDeviceID.IsNone())
-				{
-					// No default set, fail.
-					return {};
-				}
-			}
-
-			// Find entry matching that device ID.
-			if (const FCacheEntry* Found = Cache.Find(InDeviceID))
-			{
-				FReadScopeLock EntryReadLock(Found->MutationLock);
-				if (Found->State == EAudioDeviceState::Active && 
-					Found->Type == FCacheEntry::EEndpointType::Render )
-				{
-					FAudioPlatformDeviceInfo Info;
-					MakeDeviceInfo(*Found, DefaultOutputDevice, Info);
-					return Info;
-				}
-			}
-			// Fail.
-			return {};
-		}
-	};
-
 	void FMixerPlatformXAudio2::RegisterDeviceChangedListener()
 	{
 		if (!WindowsNotificationClient.IsValid())
@@ -1378,7 +103,7 @@ namespace Audio
 		if (InAudioDeviceRole == EAudioDeviceRole::Console)
 		{
 			UE_LOG(LogAudioMixer, Warning, TEXT("FMixerPlatformXAudio2: Changing default audio render device to new device: Role=%s, DeviceName=%s, InstanceID=%d"), 
-				FWindowsMMDeviceCache::ToString(InAudioDeviceRole), *WindowsNotificationClient->GetFriendlyName(DeviceId), InstanceID);
+				Audio::ToString(InAudioDeviceRole), *WindowsNotificationClient->GetFriendlyName(DeviceId), InstanceID);
 
 			RequestDeviceSwap(DeviceId, /* force */true, TEXT("FMixerPlatformXAudio2::OnDefaultRenderDeviceChanged"));
 		}
@@ -1465,8 +190,8 @@ namespace Audio
 	}
 }
 
-#include "Windows/HideWindowsPlatformAtomics.h"
-#include "Windows/HideWindowsPlatformTypes.h"
+//#include "Windows/HideWindowsPlatformAtomics.h"
+//#include "Windows/HideWindowsPlatformTypes.h"
 
 #else 
 // Nothing for XBOXOne
diff --git a/Engine/Source/Runtime/Windows/AudioMixerXAudio2/Private/ConversionHelpers.h b/Engine/Source/Runtime/Windows/AudioMixerXAudio2/Private/ConversionHelpers.h
new file mode 100644
--- /dev/null
+++ b/Engine/Source/Runtime/Windows/AudioMixerXAudio2/Private/ConversionHelpers.h
@@ -0,0 +1,57 @@
+// Copyright Epic Games, Inc. All Rights Reserved.
+
+#pragma once
+
+#if PLATFORM_WINDOWS
+
+#include "AudioMixer.h"	
+#include "Windows/AllowWindowsPlatformTypes.h"
+#include "Windows/MinWindows.h"
+
+THIRD_PARTY_INCLUDES_START
+#include <mmdeviceapi.h>			// IMMNotificationClient
+#include <audiopolicy.h>			// IAudioSessionEvents
+THIRD_PARTY_INCLUDES_END
+
+#include "Windows/HideWindowsPlatformTypes.h"
+
+#endif //PLATFORM_WINDOWS
+
+namespace Audio
+{
+#if PLATFORM_WINDOWS
+	static inline EAudioDeviceState ConvertWordToDeviceState(DWORD InWord)
+	{
+		switch (InWord)
+		{
+		case DEVICE_STATE_ACTIVE: return Audio::EAudioDeviceState::Active;
+		case DEVICE_STATE_DISABLED: return Audio::EAudioDeviceState::Disabled;
+		case DEVICE_STATE_UNPLUGGED: return Audio::EAudioDeviceState::Unplugged;
+		case DEVICE_STATE_NOTPRESENT: return Audio::EAudioDeviceState::NotPresent;
+		default:
+			break;
+		}
+		checkNoEntry();
+		return Audio::EAudioDeviceState::NotPresent;
+	}
+
+	static inline IAudioMixerDeviceChangedListener::EDisconnectReason AudioSessionDisconnectToEDisconnectReason(AudioSessionDisconnectReason InDisconnectReason)
+	{
+		using namespace Audio;
+		switch (InDisconnectReason)
+		{
+		case DisconnectReasonDeviceRemoval:			return IAudioMixerDeviceChangedListener::EDisconnectReason::DeviceRemoval;
+		case DisconnectReasonServerShutdown:		return IAudioMixerDeviceChangedListener::EDisconnectReason::ServerShutdown;
+		case DisconnectReasonFormatChanged:			return IAudioMixerDeviceChangedListener::EDisconnectReason::FormatChanged;
+		case DisconnectReasonSessionLogoff:			return IAudioMixerDeviceChangedListener::EDisconnectReason::SessionLogoff;
+		case DisconnectReasonSessionDisconnected:	return IAudioMixerDeviceChangedListener::EDisconnectReason::SessionDisconnected;
+		case DisconnectReasonExclusiveModeOverride:	return IAudioMixerDeviceChangedListener::EDisconnectReason::ExclusiveModeOverride;
+		default: break;
+		}
+
+		checkNoEntry();
+		return IAudioMixerDeviceChangedListener::EDisconnectReason::DeviceRemoval;
+	}
+
+#endif //PLATFORM_WINDOWS
+}
diff --git a/Engine/Source/Runtime/Windows/AudioMixerXAudio2/Private/ScopedCom.h b/Engine/Source/Runtime/Windows/AudioMixerXAudio2/Private/ScopedCom.h
--- a/Engine/Source/Runtime/Windows/AudioMixerXAudio2/Private/ScopedCom.h
+++ b/Engine/Source/Runtime/Windows/AudioMixerXAudio2/Private/ScopedCom.h
@@ -2,6 +2,13 @@
 
 #pragma once
 
+#include "Windows/AllowWindowsPlatformTypes.h"
+#include "Windows/MinWindows.h"
+
+THIRD_PARTY_INCLUDES_START
+#include <combaseapi.h>
+THIRD_PARTY_INCLUDES_END
+
 namespace Audio
 {
 	struct FScopeComString final
@@ -10,7 +17,7 @@ namespace Audio
 
 		UE_NONCOPYABLE(FScopeComString)
 
-			const LPTSTR Get() const
+		const LPTSTR Get() const
 		{
 			return StringPtr;
 		}
@@ -32,4 +39,6 @@ namespace Audio
 			}
 		}
 	};
-}
\ No newline at end of file
+}
+
+#include "Windows/HideWindowsPlatformTypes.h"
\ No newline at end of file
diff --git a/Engine/Source/Runtime/Windows/AudioMixerXAudio2/Private/ToStringHelpers.cpp b/Engine/Source/Runtime/Windows/AudioMixerXAudio2/Private/ToStringHelpers.cpp
new file mode 100644
--- /dev/null
+++ b/Engine/Source/Runtime/Windows/AudioMixerXAudio2/Private/ToStringHelpers.cpp
@@ -0,0 +1,145 @@
+// Copyright Epic Games, Inc. All Rights Reserved.
+
+#pragma once
+
+#include "ToStringHelpers.h"
+
+#if PLATFORM_WINDOWS
+
+#if !NO_LOGGING
+// PSStringFromPropertyKey needs this lib, which we use for logging the Property GUIDs.
+#pragma comment(lib, "Propsys.lib")
+#endif //!NO_LOGGING
+
+namespace Audio
+{
+	const TCHAR* ToString(AudioSessionDisconnectReason InDisconnectReason)
+	{
+#if !NO_LOGGING
+#define CASE_TO_STRING(X) case AudioSessionDisconnectReason::X: return TEXT(#X)
+		switch (InDisconnectReason)
+		{
+			CASE_TO_STRING(DisconnectReasonDeviceRemoval);
+			CASE_TO_STRING(DisconnectReasonServerShutdown);
+			CASE_TO_STRING(DisconnectReasonFormatChanged);
+			CASE_TO_STRING(DisconnectReasonSessionLogoff);
+			CASE_TO_STRING(DisconnectReasonSessionDisconnected);
+			CASE_TO_STRING(DisconnectReasonExclusiveModeOverride);
+		default:
+			checkNoEntry();
+			break;
+		}
+#undef CASE_TO_STRING
+#endif //!NO_LOGGING
+		return TEXT("Unknown");
+	}
+
+	const TCHAR* ToString(ERole InRole)
+	{
+#if !NO_LOGGING
+#define CASE_TO_STRING(X) case ERole::X: return TEXT(#X)
+		switch (InRole)
+		{
+			CASE_TO_STRING(eConsole);
+			CASE_TO_STRING(eMultimedia);
+			CASE_TO_STRING(eCommunications);
+		default:
+			checkNoEntry();
+			break;
+		}
+#undef CASE_TO_STRING
+#endif //!NO_LOGGING
+		return TEXT("Unknown");
+	}
+
+	const TCHAR* ToString(EDataFlow InFlow)
+	{
+#if !NO_LOGGING
+#define CASE_TO_STRING(X) case EDataFlow::X: return TEXT(#X)
+		switch (InFlow)
+		{
+			CASE_TO_STRING(eRender);
+			CASE_TO_STRING(eCapture);
+			CASE_TO_STRING(eAll);
+		default:
+			checkNoEntry();
+			break;
+		}
+#endif //!NO_LOGGING
+#undef CASE_TO_STRING
+		return TEXT("Unknown");
+	}
+
+	const TCHAR* ToString(EAudioDeviceRole InRole)
+	{
+#define CASE_TO_STRING(X) case EAudioDeviceRole::X: return TEXT(#X)
+		switch (InRole)
+		{
+			CASE_TO_STRING(Console);
+			CASE_TO_STRING(Multimedia);
+			CASE_TO_STRING(Communications);
+		default:
+			return TEXT("Unknown");
+		}
+#undef CASE_TO_STRING
+	}
+
+	const TCHAR* ToString(EAudioDeviceState InState)
+	{
+#define CASE_TO_STRING(X) case EAudioDeviceState::X: return TEXT(#X)
+		switch (InState)
+		{
+			CASE_TO_STRING(Active);
+			CASE_TO_STRING(Disabled);
+			CASE_TO_STRING(NotPresent);
+			CASE_TO_STRING(Unplugged);
+		default:
+			return TEXT("Unknown");
+		}
+#undef CASE_TO_STRING
+	}
+
+	const FString ToFString(const TArray<EAudioMixerChannel::Type>& InChannels)
+	{
+		FString ChannelString;
+		static const int32 ApproxChannelNameLength = 18;
+		ChannelString.Reserve(ApproxChannelNameLength * InChannels.Num());
+		for (EAudioMixerChannel::Type i : InChannels)
+		{
+			ChannelString.Append(EAudioMixerChannel::ToString(i));
+			ChannelString.Append(TEXT("|"));
+		}
+		return ChannelString;
+	}
+
+	FString ToFString(const PROPERTYKEY Key)
+	{
+#if !NO_LOGGING
+#define IF_PROP_STRING(PROP) if(Key.fmtid == PROP.fmtid) return TEXT(#PROP)
+
+		IF_PROP_STRING(PKEY_AudioEndpoint_PhysicalSpeakers);
+		IF_PROP_STRING(PKEY_AudioEngine_DeviceFormat);
+		IF_PROP_STRING(PKEY_AudioEngine_OEMFormat);
+		IF_PROP_STRING(PKEY_AudioEndpoint_Association);
+		IF_PROP_STRING(PKEY_AudioEndpoint_ControlPanelPageProvider);
+		IF_PROP_STRING(PKEY_AudioEndpoint_Disable_SysFx);
+		IF_PROP_STRING(PKEY_AudioEndpoint_FormFactor);
+		IF_PROP_STRING(PKEY_AudioEndpoint_FullRangeSpeakers);
+		IF_PROP_STRING(PKEY_AudioEndpoint_GUID);
+		IF_PROP_STRING(PKEY_AudioEndpoint_Supports_EventDriven_Mode);
+
+#undef IF_PROP_STRING
+
+		TCHAR KeyString[PKEYSTR_MAX];
+		HRESULT HR = PSStringFromPropertyKey(Key, KeyString, ARRAYSIZE(KeyString));
+		if (SUCCEEDED(HR))
+		{
+			return FString(KeyString);
+		}
+#endif //!NO_LOGGING
+		return TEXT("Unknown");
+	}
+
+}
+
+#endif //PLATFORM_WINDOWS
\ No newline at end of file
diff --git a/Engine/Source/Runtime/Windows/AudioMixerXAudio2/Private/ToStringHelpers.h b/Engine/Source/Runtime/Windows/AudioMixerXAudio2/Private/ToStringHelpers.h
new file mode 100644
--- /dev/null
+++ b/Engine/Source/Runtime/Windows/AudioMixerXAudio2/Private/ToStringHelpers.h
@@ -0,0 +1,33 @@
+// Copyright Epic Games, Inc. All Rights Reserved.
+
+#pragma once
+
+#include "AudioMixer.h"
+
+#if PLATFORM_WINDOWS
+
+#include "Windows/AllowWindowsPlatformTypes.h"
+#include "Windows/MinWindows.h"
+
+THIRD_PARTY_INCLUDES_START
+#include <mmdeviceapi.h>			// IMMNotificationClient
+#include <audiopolicy.h>			// IAudioSessionEvents
+THIRD_PARTY_INCLUDES_END
+
+#include "Windows/HideWindowsPlatformTypes.h"
+
+#endif //PLATFORM_WINDOWS
+
+namespace Audio
+{
+#if PLATFORM_WINDOWS
+	const TCHAR* ToString(AudioSessionDisconnectReason InDisconnectReason);
+	const TCHAR* ToString(ERole InRole);
+	const TCHAR* ToString(EDataFlow InFlow);
+	FString ToFString(const PROPERTYKEY Key);
+#endif //PLATFORM_WINDOWS
+
+	const FString ToFString(const TArray<EAudioMixerChannel::Type>& InChannels);
+	const TCHAR* ToString(EAudioDeviceRole InRole);
+	const TCHAR* ToString(EAudioDeviceState InState);	
+}
diff --git a/Engine/Source/Runtime/Windows/AudioMixerXAudio2/Private/WindowsMMDeviceInfoCache.cpp b/Engine/Source/Runtime/Windows/AudioMixerXAudio2/Private/WindowsMMDeviceInfoCache.cpp
new file mode 100644
--- /dev/null
+++ b/Engine/Source/Runtime/Windows/AudioMixerXAudio2/Private/WindowsMMDeviceInfoCache.cpp
@@ -0,0 +1,610 @@
+// Copyright Epic Games, Inc. All Rights Reserved.
+
+#include "WindowsMMDeviceInfoCache.h"
+
+#if PLATFORM_WINDOWS
+
+#include "Windows/AllowWindowsPlatformTypes.h"
+#include "Windows/AllowWindowsPlatformAtomics.h"
+
+THIRD_PARTY_INCLUDES_START
+#include <functiondiscoverykeys_devpkey.h>
+THIRD_PARTY_INCLUDES_END
+
+#include "Windows/HideWindowsPlatformAtomics.h"
+#include "Windows/HideWindowsPlatformTypes.h"
+
+#include "ToStringHelpers.h"
+#include "ConversionHelpers.h"
+
+namespace Audio
+{
+
+FWindowsMMDeviceCache::FCacheEntry& FWindowsMMDeviceCache::FCacheEntry::operator=(FCacheEntry&& InOther)
+{
+	DeviceId = MoveTemp(InOther.DeviceId);
+	FriendlyName = MoveTemp(InOther.FriendlyName);
+	DeviceFriendlyName = MoveTemp(InOther.DeviceFriendlyName);
+	State = MoveTemp(InOther.State);
+	NumChannels = MoveTemp(InOther.NumChannels);
+	SampleRate = MoveTemp(InOther.SampleRate);
+	Type = MoveTemp(InOther.Type);
+	ChannelBitmask = MoveTemp(InOther.ChannelBitmask);
+	OutputChannels = MoveTemp(InOther.OutputChannels);
+	return *this;
+}
+
+FWindowsMMDeviceCache::FCacheEntry& FWindowsMMDeviceCache::FCacheEntry::operator=(const FCacheEntry& InOther)
+{
+	// Copy everything but the lock. 
+	DeviceId = InOther.DeviceId;
+	FriendlyName = InOther.FriendlyName;
+	DeviceFriendlyName = InOther.DeviceFriendlyName;
+	State = InOther.State;
+	NumChannels = InOther.NumChannels;
+	SampleRate = InOther.SampleRate;
+	Type = InOther.Type;
+	ChannelBitmask = InOther.ChannelBitmask;
+	OutputChannels = InOther.OutputChannels;
+	return *this;
+}
+
+FWindowsMMDeviceCache::FCacheEntry::FCacheEntry(const FString& InDeviceId)
+	: DeviceId{ InDeviceId }
+{
+}
+
+FWindowsMMDeviceCache::FCacheEntry::FCacheEntry(FCacheEntry&& InOther)
+{
+	*this = MoveTemp(InOther);
+}
+
+FWindowsMMDeviceCache::FCacheEntry::FCacheEntry(const FCacheEntry& InOther)
+{
+	*this = InOther;
+}
+
+FWindowsMMDeviceCache::FWindowsMMDeviceCache()
+{
+	ensure(SUCCEEDED(CoCreateInstance(__uuidof(MMDeviceEnumerator), NULL, CLSCTX_INPROC_SERVER, IID_PPV_ARGS(&DeviceEnumerator))) && DeviceEnumerator);
+
+	EnumerateEndpoints();
+	EnumerateDefaults();
+}
+
+bool FWindowsMMDeviceCache::EnumerateChannelMask(uint32 InMask, FCacheEntry& OutInfo)
+{
+	// Loop through the extensible format channel flags in the standard order and build our output channel array
+	// From https://msdn.microsoft.com/en-us/library/windows/hardware/dn653308(v=vs.85).aspx
+	// The channels in the interleaved stream corresponding to these spatial positions must appear in the order specified above. This holds true even in the 
+	// case of a non-contiguous subset of channels. For example, if a stream contains left, bass enhance and right, then channel 1 is left, channel 2 is right, 
+	// and channel 3 is bass enhance. This enables the linkage of multi-channel streams to well-defined multi-speaker configurations.
+
+	static const uint32 REMOVE_ME_ChannelTypeMap[EAudioMixerChannel::ChannelTypeCount] =
+	{
+		SPEAKER_FRONT_LEFT,
+		SPEAKER_FRONT_RIGHT,
+		SPEAKER_FRONT_CENTER,
+		SPEAKER_LOW_FREQUENCY,
+		SPEAKER_BACK_LEFT,
+		SPEAKER_BACK_RIGHT,
+		SPEAKER_FRONT_LEFT_OF_CENTER,
+		SPEAKER_FRONT_RIGHT_OF_CENTER,
+		SPEAKER_BACK_CENTER,
+		SPEAKER_SIDE_LEFT,
+		SPEAKER_SIDE_RIGHT,
+		SPEAKER_TOP_CENTER,
+		SPEAKER_TOP_FRONT_LEFT,
+		SPEAKER_TOP_FRONT_CENTER,
+		SPEAKER_TOP_FRONT_RIGHT,
+		SPEAKER_TOP_BACK_LEFT,
+		SPEAKER_TOP_BACK_CENTER,
+		SPEAKER_TOP_BACK_RIGHT,
+		SPEAKER_RESERVED,
+	};
+
+	OutInfo.ChannelBitmask = InMask;
+	OutInfo.OutputChannels.Reset();
+
+	// No need to enumerate speakers for capture devices.
+	if (OutInfo.Type == FCacheEntry::EEndpointType::Capture)
+	{
+		return true;
+	}
+
+	uint32 ChanCount = 0;
+	for (uint32 ChannelTypeIndex = 0; ChannelTypeIndex < EAudioMixerChannel::ChannelTypeCount && ChanCount < (uint32)OutInfo.NumChannels; ++ChannelTypeIndex)
+	{
+		if (InMask & REMOVE_ME_ChannelTypeMap[ChannelTypeIndex])
+		{
+			OutInfo.OutputChannels.Add((EAudioMixerChannel::Type)ChannelTypeIndex);
+			++ChanCount;
+		}
+	}
+
+	// We didn't match channel masks for all channels, revert to a default ordering
+	if (ChanCount < (uint32)OutInfo.NumChannels)
+	{
+		UE_CLOG(Audio::IAudioMixer::ShouldLogDeviceSwaps(), LogAudioMixer, Warning, TEXT("FWindowsMMDeviceCache: Did not find the channel type flags for audio device '%s'. Reverting to a default channel ordering."), *OutInfo.FriendlyName);
+
+		OutInfo.OutputChannels.Reset();
+
+		static const EAudioMixerChannel::Type DefaultChannelOrdering[] = {
+			EAudioMixerChannel::FrontLeft,
+			EAudioMixerChannel::FrontRight,
+			EAudioMixerChannel::FrontCenter,
+			EAudioMixerChannel::LowFrequency,
+			EAudioMixerChannel::SideLeft,
+			EAudioMixerChannel::SideRight,
+			EAudioMixerChannel::BackLeft,
+			EAudioMixerChannel::BackRight,
+		};
+
+		const EAudioMixerChannel::Type* ChannelOrdering = DefaultChannelOrdering;
+
+		// Override channel ordering for some special cases
+		if (OutInfo.NumChannels == 4)
+		{
+			static EAudioMixerChannel::Type DefaultChannelOrderingQuad[] = {
+				EAudioMixerChannel::FrontLeft,
+				EAudioMixerChannel::FrontRight,
+				EAudioMixerChannel::BackLeft,
+				EAudioMixerChannel::BackRight,
+			};
+
+			ChannelOrdering = DefaultChannelOrderingQuad;
+		}
+		else if (OutInfo.NumChannels == 6)
+		{
+			static const EAudioMixerChannel::Type DefaultChannelOrdering51[] = {
+				EAudioMixerChannel::FrontLeft,
+				EAudioMixerChannel::FrontRight,
+				EAudioMixerChannel::FrontCenter,
+				EAudioMixerChannel::LowFrequency,
+				EAudioMixerChannel::BackLeft,
+				EAudioMixerChannel::BackRight,
+			};
+
+			ChannelOrdering = DefaultChannelOrdering51;
+		}
+
+		check(OutInfo.NumChannels <= 8);
+		for (int32 Index = 0; Index < OutInfo.NumChannels; ++Index)
+		{
+			OutInfo.OutputChannels.Add(ChannelOrdering[Index]);
+		}
+	}
+	return true;
+}
+
+bool FWindowsMMDeviceCache::EnumerateChannelFormat(const WAVEFORMATEX* InFormat, FCacheEntry& OutInfo)
+{
+	OutInfo.OutputChannels.Empty();
+
+	// Extensible format supports surround sound so we need to parse the channel configuration to build our channel output array
+	if (InFormat->wFormatTag == WAVE_FORMAT_EXTENSIBLE)
+	{
+		// Cast to the extensible format to get access to extensible data
+		const WAVEFORMATEXTENSIBLE* WaveFormatExtensible = (const WAVEFORMATEXTENSIBLE*)(InFormat);
+		return EnumerateChannelMask(WaveFormatExtensible->dwChannelMask, OutInfo);
+	}
+	else
+	{
+		// Non-extensible formats only support mono or stereo channel output
+		OutInfo.OutputChannels.Add(EAudioMixerChannel::FrontLeft);
+		if (OutInfo.NumChannels == 2)
+		{
+			OutInfo.OutputChannels.Add(EAudioMixerChannel::FrontRight);
+		}
+	}
+
+	// Aways success for now.
+	return true;
+}
+
+FWindowsMMDeviceCache::FCacheEntry::EEndpointType FWindowsMMDeviceCache::QueryDeviceDataFlow(const TComPtr<IMMDevice>& InDevice) const
+{
+	TComPtr<IMMEndpoint> Endpoint;
+	if (SUCCEEDED(InDevice->QueryInterface(IID_PPV_ARGS(&Endpoint))))
+	{
+		EDataFlow DataFlow = eRender;
+		if (SUCCEEDED(Endpoint->GetDataFlow(&DataFlow)))
+		{
+			switch (DataFlow)
+			{
+			case eRender:
+				return FCacheEntry::EEndpointType::Render;
+			case eCapture:
+				return FCacheEntry::EEndpointType::Capture;
+			default:
+				break;
+			}
+		}
+	}
+	return FCacheEntry::EEndpointType::Unknown;
+}
+
+bool FWindowsMMDeviceCache::EnumerateDeviceProps(const TComPtr<IMMDevice>& InDevice, FCacheEntry& OutInfo)
+{
+	// Mark if this is a Render Device or Capture or Unknown.
+	OutInfo.Type = QueryDeviceDataFlow(InDevice);
+
+	// Also query the device state.
+	DWORD DeviceState = DEVICE_STATE_NOTPRESENT;
+	if (SUCCEEDED(InDevice->GetState(&DeviceState)))
+	{
+		OutInfo.State = ConvertWordToDeviceState(DeviceState);
+	}
+
+	TComPtr<IPropertyStore> PropertyStore;
+	if (SUCCEEDED(InDevice->OpenPropertyStore(STGM_READ, &PropertyStore)))
+	{
+		// Friendly Name
+		PROPVARIANT FriendlyName;
+		PropVariantInit(&FriendlyName);
+		if (SUCCEEDED(PropertyStore->GetValue(PKEY_Device_FriendlyName, &FriendlyName)) && FriendlyName.pwszVal)
+		{
+			OutInfo.FriendlyName = FString(FriendlyName.pwszVal);
+			PropVariantClear(&FriendlyName);
+		}
+
+		auto EnumDeviceFormat = [this](const TComPtr<IPropertyStore>& InPropStore, REFPROPERTYKEY InKey, FCacheEntry& OutInfo) -> bool
+		{
+			// Device Format
+			PROPVARIANT DeviceFormat;
+			PropVariantInit(&DeviceFormat);
+
+			if (SUCCEEDED(InPropStore->GetValue(InKey, &DeviceFormat)) && DeviceFormat.blob.pBlobData)
+			{
+				const WAVEFORMATEX* WaveFormatEx = (const WAVEFORMATEX*)(DeviceFormat.blob.pBlobData);
+				OutInfo.NumChannels = FMath::Clamp((int32)WaveFormatEx->nChannels, 2, 8);
+				OutInfo.SampleRate = WaveFormatEx->nSamplesPerSec;
+
+				EnumerateChannelFormat(WaveFormatEx, OutInfo);
+				PropVariantClear(&DeviceFormat);
+				return true;
+			}
+			return false;
+		};
+
+		if (EnumDeviceFormat(PropertyStore, PKEY_AudioEngine_DeviceFormat, OutInfo) ||
+			EnumDeviceFormat(PropertyStore, PKEY_AudioEngine_OEMFormat, OutInfo))
+		{
+		}
+		else
+		{
+			// Log a warning if this device is active as we failed to ask for a format
+			UE_CLOG(DeviceState == DEVICE_STATE_ACTIVE, LogAudioMixer, Warning, TEXT("FWindowsMMDeviceCache: Failed to get Format for active device '%s'"), *OutInfo.FriendlyName);
+		}
+	}
+
+	// Aways success for now.
+	return true;
+}
+
+void FWindowsMMDeviceCache::EnumerateEndpoints()
+{
+	// Build a new cache from scratch.
+	TMap<FName, FCacheEntry> NewCache;
+
+	// Get Device Enumerator.
+	if (DeviceEnumerator)
+	{
+		// Get Render Device Collection. (note we ask for ALL states, which include disabled/unplugged devices.).
+		TComPtr<IMMDeviceCollection> DeviceCollection;
+		uint32 DeviceCount = 0;
+		if (SUCCEEDED(DeviceEnumerator->EnumAudioEndpoints(eAll, DEVICE_STATEMASK_ALL, &DeviceCollection)) && DeviceCollection &&
+			SUCCEEDED(DeviceCollection->GetCount(&DeviceCount)))
+		{
+			for (uint32 i = 0; i < DeviceCount; ++i)
+			{
+				TComPtr<IMMDevice> Device;
+				if (SUCCEEDED(DeviceCollection->Item(i, &Device)) && Device)
+				{
+					// Get the device id string (guid)
+					Audio::FScopeComString DeviceIdString;
+					if (SUCCEEDED(Device->GetId(&DeviceIdString.StringPtr)) && DeviceIdString)
+					{
+						FCacheEntry Info{ DeviceIdString.Get() };
+
+						// Enumerate props into our info object.
+						EnumerateDeviceProps(Device, Info);
+
+						UE_LOG(LogAudioMixer, Verbose, TEXT("FWindowsMMDeviceCache: %s Device '%s' ID='%s'"),
+							Info.Type == FCacheEntry::EEndpointType::Capture ? TEXT("Capture") :
+							Info.Type == FCacheEntry::EEndpointType::Render ? TEXT("Render") :
+							TEXT("UNKNOWN!"),
+							*Info.DeviceId.ToString(),
+							*Info.FriendlyName
+						);
+
+						check(!NewCache.Contains(Info.DeviceId));
+						NewCache.Emplace(Info.DeviceId, Info);
+					}
+				}
+			}
+		}
+
+		// Finally, Replace cache with new one.
+		{
+			FWriteScopeLock Lock(CacheMutationLock);
+			Cache = MoveTemp(NewCache);
+		}
+	}
+}
+
+void FWindowsMMDeviceCache::EnumerateDefaults()
+{
+	auto GetDefaultDeviceID = [this](EDataFlow InDataFlow, ERole InRole, FName& OutDeviceId) -> bool
+	{
+		// Mark default device.
+		bool bSuccess = false;
+		TComPtr<IMMDevice> DefaultDevice;
+		if (SUCCEEDED(DeviceEnumerator->GetDefaultAudioEndpoint(InDataFlow, InRole, &DefaultDevice)))
+		{
+			Audio::FScopeComString DeviceIdString;
+			if (SUCCEEDED(DefaultDevice->GetId(&DeviceIdString.StringPtr)) && DeviceIdString)
+			{
+				OutDeviceId = DeviceIdString.Get();
+				bSuccess = true;
+			}
+		}
+		return bSuccess;
+	};
+
+	// Get defaults (render, capture).
+	FWriteScopeLock Lock(CacheMutationLock);
+	static_assert((int32)EAudioDeviceRole::COUNT == ERole_enum_count, "EAudioDeviceRole should be the same as ERole");
+	for (int32 i = 0; i < ERole_enum_count; ++i)
+	{
+		FName DeviceIdName;
+		if (GetDefaultDeviceID(eRender, static_cast<ERole>(i), DeviceIdName))
+		{
+			UE_CLOG(!DeviceIdName.IsNone(), LogAudioMixer, Verbose, TEXT("FWindowsMMDeviceCache: Default Render Role='%s', Device='%s'"), ToString((EAudioDeviceRole)i), *GetFriendlyName(DeviceIdName));
+			DefaultRenderId[i] = DeviceIdName;
+		}
+		if (GetDefaultDeviceID(eCapture, static_cast<ERole>(i), DeviceIdName))
+		{
+			UE_CLOG(!DeviceIdName.IsNone(), LogAudioMixer, Verbose, TEXT("FWindowsMMDeviceCache: Default Capture Role='%s', Device='%s'"), ToString((EAudioDeviceRole)i), *GetFriendlyName(DeviceIdName));
+			DefaultCaptureId[i] = DeviceIdName;
+		}
+	}
+}
+
+void FWindowsMMDeviceCache::OnDefaultCaptureDeviceChanged(const EAudioDeviceRole InAudioDeviceRole, const FString& DeviceId)
+{
+	FWriteScopeLock WriteLock(CacheMutationLock);
+	check(InAudioDeviceRole < EAudioDeviceRole::COUNT);
+	DefaultCaptureId[(int32)InAudioDeviceRole] = *DeviceId;
+}
+
+void FWindowsMMDeviceCache::OnDefaultRenderDeviceChanged(const EAudioDeviceRole InAudioDeviceRole, const FString& DeviceId)
+{
+	FWriteScopeLock WriteLock(CacheMutationLock);
+	check(InAudioDeviceRole < EAudioDeviceRole::COUNT);
+	DefaultRenderId[(int32)InAudioDeviceRole] = *DeviceId;
+}
+
+void FWindowsMMDeviceCache::OnDeviceAdded(const FString& DeviceId, bool bIsRender)
+{
+	if (ensure(DeviceEnumerator))
+	{
+		TComPtr<IMMDevice> Device;
+		if (SUCCEEDED(DeviceEnumerator->GetDevice(*DeviceId, &Device)))
+		{
+			FCacheEntry Info{ *DeviceId };
+			if (EnumerateDeviceProps(Device, Info))
+			{
+				FWriteScopeLock WriteLock(CacheMutationLock);
+				check(!Cache.Contains(Info.DeviceId));
+				Cache.Add(Info.DeviceId, Info);
+			}
+		}
+	}
+}
+
+void FWindowsMMDeviceCache::OnDeviceRemoved(const FString& DeviceId, bool)
+{
+	FWriteScopeLock WriteLock(CacheMutationLock);
+	FName DeviceIdName = *DeviceId;
+	check(Cache.Contains(DeviceIdName));
+	Cache.Remove(DeviceIdName);
+}
+
+TOptional<FWindowsMMDeviceCache::FCacheEntry> FWindowsMMDeviceCache::BuildCacheEntry(const FString& DeviceId)
+{
+	TComPtr<IMMDevice> Device;
+	if (SUCCEEDED(DeviceEnumerator->GetDevice(*DeviceId, &Device)))
+	{
+		FCacheEntry Info{ *DeviceId };
+		if (EnumerateDeviceProps(Device, Info))
+		{
+			return Info;
+		}
+	}
+	return {};
+}
+
+FString FWindowsMMDeviceCache::GetFriendlyName(FName InDeviceId) const
+{
+	if (const FCacheEntry* Entry = Cache.Find(InDeviceId))
+	{
+		return Entry->FriendlyName;
+	}
+	return TEXT("Unknown");
+}
+
+
+
+void FWindowsMMDeviceCache::OnDeviceStateChanged(const FString& DeviceId, const EAudioDeviceState InState, bool)
+{
+	TOptional<FCacheEntry> Info = BuildCacheEntry(DeviceId);
+
+	FReadScopeLock ReadLock(CacheMutationLock);
+	FName DeviceIdName = *DeviceId;
+	ensureMsgf(Cache.Contains(DeviceIdName), TEXT("Expecting to find '%s' in cache '%s'"), *DeviceId, Info ? *Info->FriendlyName : TEXT("Unknown"));
+
+	if (FCacheEntry* Entry = Cache.Find(DeviceIdName))
+	{
+		FWriteScopeLock WriteLock(Entry->MutationLock);
+
+		UE_CLOG(Audio::IAudioMixer::ShouldLogDeviceSwaps(), LogAudioMixer, Verbose, TEXT("FWindowsMMDeviceCache: DeviceName='%s' - DeviceID='%s' state changed from '%s' to '%s'."),
+			Info ? *Info->FriendlyName : TEXT("Unknown"), *DeviceId, ToString(Entry->State), ToString(InState));
+
+		Entry->State = InState;
+	}
+}
+
+void FWindowsMMDeviceCache::OnFormatChanged(const FString& InDeviceId, const FFormatChangedData& InFormat)
+{
+	FName DeviceName(InDeviceId);
+	bool bNeedToEnumerateChannels = false;
+	bool bDirty = false;
+
+	FReadScopeLock MapReadLock(CacheMutationLock);
+	if (FCacheEntry* Found = Cache.Find(DeviceName))
+	{
+		// Make a copy of the entry
+		FCacheEntry EntryCopy(InDeviceId);
+		{
+			FReadScopeLock FoundReadLock(Found->MutationLock);
+			EntryCopy = *Found;
+		}
+
+		if (EntryCopy.NumChannels != InFormat.NumChannels)
+		{
+			UE_CLOG(Audio::IAudioMixer::ShouldLogDeviceSwaps(), LogAudioMixer, Verbose, TEXT("FWindowsMMDeviceCache: DeviceID='%s', Name='%s' changed default format from %d channels to %d."), *InDeviceId, *EntryCopy.FriendlyName, EntryCopy.NumChannels, InFormat.NumChannels);
+			EntryCopy.NumChannels = InFormat.NumChannels;
+			bNeedToEnumerateChannels = true;
+			bDirty = true;
+		}
+		if (EntryCopy.SampleRate != InFormat.SampleRate)
+		{
+			UE_CLOG(Audio::IAudioMixer::ShouldLogDeviceSwaps(), LogAudioMixer, Verbose, TEXT("FWindowsMMDeviceCache: DeviceID='%s', Name='%s' changed default format from %dhz to %dhz."), *InDeviceId, *EntryCopy.FriendlyName, EntryCopy.SampleRate, InFormat.SampleRate);
+			EntryCopy.SampleRate = InFormat.SampleRate;
+			bDirty = true;
+		}
+		if (EntryCopy.ChannelBitmask != InFormat.ChannelBitmask)
+		{
+			UE_CLOG(Audio::IAudioMixer::ShouldLogDeviceSwaps(), LogAudioMixer, Verbose, TEXT("FWindowsMMDeviceCache: DeviceID='%s', Name='%s' changed default format from 0x%x to 0x%x bitmask"), *InDeviceId, *EntryCopy.FriendlyName, EntryCopy.ChannelBitmask, InFormat.ChannelBitmask);
+			EntryCopy.ChannelBitmask = InFormat.ChannelBitmask;
+			bNeedToEnumerateChannels = true;
+			bDirty = true;
+		}
+
+		if (bNeedToEnumerateChannels)
+		{
+			UE_CLOG(Audio::IAudioMixer::ShouldLogDeviceSwaps(), LogAudioMixer, Verbose, TEXT("FWindowsMMDeviceCache: Channel Change, DeviceID='%s', Name='%s' OLD=[%s]"), *InDeviceId, *EntryCopy.FriendlyName, *ToFString(EntryCopy.OutputChannels));
+			EnumerateChannelMask(InFormat.ChannelBitmask, EntryCopy);
+			UE_CLOG(Audio::IAudioMixer::ShouldLogDeviceSwaps(), LogAudioMixer, Verbose, TEXT("FWindowsMMDeviceCache: Channel Change, DeviceID='%s', Name='%s' NEW=[%s]"), *InDeviceId, *EntryCopy.FriendlyName, *ToFString(EntryCopy.OutputChannels));
+		}
+
+		// Update the entire entry with one write.
+		if (bDirty)
+		{
+			FWriteScopeLock FoundWriteLock(Found->MutationLock);
+			*Found = EntryCopy;
+		}
+	}
+}
+
+void FWindowsMMDeviceCache::MakeDeviceInfo(const FCacheEntry& InEntry, FName InDefaultDevice, FAudioPlatformDeviceInfo& OutInfo) const
+{
+	OutInfo.Reset();
+	OutInfo.Name = InEntry.FriendlyName;
+	OutInfo.DeviceId = InEntry.DeviceId.GetPlainNameString();
+	OutInfo.NumChannels = InEntry.NumChannels;
+	OutInfo.SampleRate = InEntry.SampleRate;
+	OutInfo.OutputChannelArray = InEntry.OutputChannels;
+	OutInfo.Format = EAudioMixerStreamDataFormat::Float;
+	OutInfo.bIsSystemDefault = InEntry.DeviceId == InDefaultDevice;
+}
+
+TArray<Audio::FAudioPlatformDeviceInfo> FWindowsMMDeviceCache::GetAllActiveOutputDevices() const
+{
+	SCOPED_NAMED_EVENT(FWindowsMMDeviceCache_GetAllActiveOutputDevices, FColor::Blue);
+
+	// Find all active devices.
+	TArray<FAudioPlatformDeviceInfo> ActiveDevices;
+
+	// Read lock
+	FReadScopeLock ReadLock(CacheMutationLock);
+	ActiveDevices.Reserve(Cache.Num());
+
+	// Ask for defaults once, as we are inside a read lock.
+	FName DefaultRenderDeviceId = GetDefaultOutputDevice_NoLock();
+
+	// Walk cache, read lock for each entry.
+	for (const auto& i : Cache)
+	{
+		// Read lock on each entry.
+		FReadScopeLock CacheEntryReadLock(i.Value.MutationLock);
+		if (i.Value.State == EAudioDeviceState::Active &&
+			i.Value.Type == FCacheEntry::EEndpointType::Render)
+		{
+			FAudioPlatformDeviceInfo& Info = ActiveDevices.Emplace_GetRef();
+			MakeDeviceInfo(i.Value, DefaultRenderDeviceId, Info);
+		}
+	}
+
+	// RVO
+	return ActiveDevices;
+}
+
+FName FWindowsMMDeviceCache::GetDefaultOutputDevice_NoLock() const
+{
+	if (!DefaultRenderId[(int32)EAudioDeviceRole::Console].IsNone())
+	{
+		return DefaultRenderId[(int32)EAudioDeviceRole::Console];
+	}
+	if (!DefaultRenderId[(int32)EAudioDeviceRole::Multimedia].IsNone())
+	{
+		return DefaultRenderId[(int32)EAudioDeviceRole::Multimedia];
+	}
+	return NAME_None;
+}
+
+TOptional<Audio::FAudioPlatformDeviceInfo> FWindowsMMDeviceCache::FindDefaultOutputDevice() const
+{
+	return FindActiveOutputDevice(NAME_None);
+}
+
+TOptional<Audio::FAudioPlatformDeviceInfo> FWindowsMMDeviceCache::FindActiveOutputDevice(FName InDeviceID) const
+{
+	SCOPED_NAMED_EVENT(FWindowsMMDeviceCache_FindActiveOutputDevice, FColor::Blue);
+
+	FReadScopeLock MapReadLock(CacheMutationLock);
+
+	// Ask for default here as we are inside the read lock.
+	const FName DefaultOutputDevice = GetDefaultOutputDevice_NoLock();
+
+	// Asking for Default?
+	if (InDeviceID.IsNone())
+	{
+		InDeviceID = DefaultOutputDevice;
+		if (InDeviceID.IsNone())
+		{
+			// No default set, fail.
+			return {};
+		}
+	}
+
+	// Find entry matching that device ID.
+	if (const FCacheEntry* Found = Cache.Find(InDeviceID))
+	{
+		FReadScopeLock EntryReadLock(Found->MutationLock);
+		if (Found->State == EAudioDeviceState::Active &&
+			Found->Type == FCacheEntry::EEndpointType::Render)
+		{
+			FAudioPlatformDeviceInfo Info;
+			MakeDeviceInfo(*Found, DefaultOutputDevice, Info);
+			return Info;
+		}
+	}
+	// Fail.
+	return {};
+}
+
+} // namespace Audio
+
+#endif //PLATFORM_WINDOWS
diff --git a/Engine/Source/Runtime/Windows/AudioMixerXAudio2/Private/WindowsMMDeviceInfoCache.h b/Engine/Source/Runtime/Windows/AudioMixerXAudio2/Private/WindowsMMDeviceInfoCache.h
new file mode 100644
--- /dev/null
+++ b/Engine/Source/Runtime/Windows/AudioMixerXAudio2/Private/WindowsMMDeviceInfoCache.h
@@ -0,0 +1,108 @@
+// Copyright Epic Games, Inc. All Rights Reserved.
+
+#pragma once
+
+#if PLATFORM_WINDOWS
+
+#include "AudioMixer.h"
+#include "Misc/ScopeRWLock.h"
+#include "ScopedCom.h"
+#include <atomic>
+
+#include "Windows/AllowWindowsPlatformTypes.h"
+#include "Windows/AllowWindowsPlatformAtomics.h"
+
+#include "Windows/MinWindows.h"
+#include "Windows/COMPointer.h"
+
+THIRD_PARTY_INCLUDES_START
+#include <mmreg.h>				// WAVEFORMATEX
+#include <mmdeviceapi.h>		// IMMDeviceEnumerator
+THIRD_PARTY_INCLUDES_END
+
+#include "Windows/HideWindowsPlatformTypes.h"
+#include "Windows/HideWindowsPlatformAtomics.h"
+
+namespace Audio
+{		
+struct FWindowsMMDeviceCache : IAudioMixerDeviceChangedListener, IAudioPlatformDeviceInfoCache
+{
+	struct FCacheEntry
+	{
+		enum class EEndpointType { Unknown, Render, Capture };
+
+		FName DeviceId;							// Key
+		FString FriendlyName;
+		FString DeviceFriendlyName;
+		EAudioDeviceState State;
+		int32 NumChannels = 0;
+		int32 SampleRate = 0;
+		EEndpointType Type = EEndpointType::Unknown;
+		uint32 ChannelBitmask = 0;				// Bitfield used to build output channels, for easy comparison.
+
+		TArray<EAudioMixerChannel::Type> OutputChannels;	// TODO. Generate this from the ChannelNum and bitmask when we are asked for it.
+		mutable FRWLock MutationLock;
+
+		FCacheEntry& operator=(const FCacheEntry& InOther);
+
+		FCacheEntry& operator=(FCacheEntry&& InOther);
+
+		FCacheEntry(const FCacheEntry& InOther);
+
+		FCacheEntry(FCacheEntry&& InOther);
+
+		FCacheEntry(const FString& InDeviceId);
+	};
+
+	TComPtr<IMMDeviceEnumerator> DeviceEnumerator;
+
+	mutable FRWLock CacheMutationLock;							// R/W lock protects map and default arrays.
+	TMap<FName, FCacheEntry> Cache;								// DeviceID GUID -> Info.
+	FName DefaultCaptureId[(int32)EAudioDeviceRole::COUNT];		// Role -> DeviceID GUID
+	FName DefaultRenderId[(int32)EAudioDeviceRole::COUNT];		// Role -> DeviceID GUID
+
+	FWindowsMMDeviceCache();
+	virtual ~FWindowsMMDeviceCache() = default;
+
+	bool EnumerateChannelMask(uint32 InMask, FCacheEntry& OutInfo);
+
+	bool EnumerateChannelFormat(const WAVEFORMATEX* InFormat, FCacheEntry& OutInfo);
+
+	FCacheEntry::EEndpointType QueryDeviceDataFlow(const TComPtr<IMMDevice>& InDevice) const;
+
+	bool EnumerateDeviceProps(const TComPtr<IMMDevice>& InDevice, FCacheEntry& OutInfo);
+
+	void EnumerateEndpoints();
+
+	void EnumerateDefaults();
+
+	void OnDefaultCaptureDeviceChanged(const EAudioDeviceRole InAudioDeviceRole, const FString& DeviceId) override;
+
+	void OnDefaultRenderDeviceChanged(const EAudioDeviceRole InAudioDeviceRole, const FString& DeviceId) override;
+
+	void OnDeviceAdded(const FString& DeviceId, bool bIsRender) override;
+
+	void OnDeviceRemoved(const FString& DeviceId, bool) override;
+
+	TOptional<FCacheEntry> BuildCacheEntry(const FString& DeviceId);
+
+	FString GetFriendlyName(FName InDeviceId) const;
+
+	void OnDeviceStateChanged(const FString& DeviceId, const EAudioDeviceState InState, bool) override;
+
+	void OnFormatChanged(const FString& InDeviceId, const FFormatChangedData& InFormat) override;
+
+	void MakeDeviceInfo(const FCacheEntry& InEntry, FName InDefaultDevice, FAudioPlatformDeviceInfo& OutInfo) const;
+
+	virtual TArray<FAudioPlatformDeviceInfo> GetAllActiveOutputDevices() const override;
+
+	FName GetDefaultOutputDevice_NoLock() const;
+
+	TOptional<FAudioPlatformDeviceInfo> FindDefaultOutputDevice() const override;
+
+	TOptional<FAudioPlatformDeviceInfo> FindActiveOutputDevice(FName InDeviceID) const override;
+};
+
+}// namespace Audio
+
+#endif //PLATFORM_WINDOWS
diff --git a/Engine/Source/Runtime/Windows/AudioMixerXAudio2/Private/WindowsMMNotificationClient.cpp b/Engine/Source/Runtime/Windows/AudioMixerXAudio2/Private/WindowsMMNotificationClient.cpp
new file mode 100644
--- /dev/null
+++ b/Engine/Source/Runtime/Windows/AudioMixerXAudio2/Private/WindowsMMNotificationClient.cpp
@@ -0,0 +1,499 @@
+// Copyright Epic Games, Inc. All Rights Reserved.
+
+#pragma once
+
+#if PLATFORM_WINDOWS
+
+#include "WindowsMMNotificationClient.h"
+
+#include "Windows/AllowWindowsPlatformTypes.h"
+
+#include "Windows/MinWindows.h"
+
+THIRD_PARTY_INCLUDES_START
+#include <Functiondiscoverykeys_devpkey.h>
+THIRD_PARTY_INCLUDES_END
+
+#include "Windows/HideWindowsPlatformTypes.h"
+
+#include "ToStringHelpers.h"
+#include "ConversionHelpers.h"
+
+namespace Audio
+{
+	bool FWindowsMMNotificationClient::RegisterForSessionNotifications(const TComPtr<IMMDevice>& InDevice)
+	{
+		FScopeLock Lock(&SessionRegistrationCS);
+
+		// If we're already listening to this device, we can early out.
+		if (DeviceListeningToSessionEvents == InDevice)
+		{
+			return true;
+		}
+
+		UnregisterForSessionNotifications();
+
+		DeviceListeningToSessionEvents = InDevice;
+
+		if (InDevice)
+		{
+			if (SUCCEEDED(InDevice->Activate(__uuidof(IAudioSessionManager), CLSCTX_INPROC_SERVER, NULL, (void**)&SessionManager)))
+			{
+				if (SUCCEEDED(SessionManager->GetAudioSessionControl(NULL, 0, &SessionControls)))
+				{
+					if (SUCCEEDED(SessionControls->RegisterAudioSessionNotification(this)))
+					{
+						UE_LOG(LogAudioMixer, Verbose, TEXT("FWindowsMMNotificationClient: Registering for sessions events for '%s'"), *GetFriendlyName(DeviceListeningToSessionEvents.Get()));
+						return true;
+					}
+				}
+			}
+		}
+		return false;
+	}
+
+	bool FWindowsMMNotificationClient::RegisterForSessionNotifications(const FString& InDeviceId)
+	{
+		if (TComPtr<IMMDevice> Device = GetDevice(InDeviceId))
+		{
+			return RegisterForSessionNotifications(Device);
+		}
+		return false;
+	}
+
+
+	HRESULT STDMETHODCALLTYPE FWindowsMMNotificationClient::OnSessionDisconnected(AudioSessionDisconnectReason InDisconnectReason)
+	{
+		UE_CLOG(Audio::IAudioMixer::ShouldLogDeviceSwaps(), LogAudioMixer, Verbose, TEXT("Session Disconnect: Reason=%s, DeviceBound=%s, HasDisconnectSessionHappened=%d"),
+			ToString(InDisconnectReason), *GetFriendlyName(DeviceListeningToSessionEvents), (int32)bHasDisconnectSessionHappened);
+
+		if (!bHasDisconnectSessionHappened)
+		{
+			{
+				FReadScopeLock Lock(ListenersSetRwLock);
+				Audio::IAudioMixerDeviceChangedListener::EDisconnectReason Reason = AudioSessionDisconnectToEDisconnectReason(InDisconnectReason);
+				for (Audio::IAudioMixerDeviceChangedListener* i : Listeners)
+				{
+					i->OnSessionDisconnect(Reason);
+				}
+			}
+
+			// Mark this true.
+			bHasDisconnectSessionHappened = true;
+		}
+
+		return S_OK;
+	}
+
+	HRESULT STDMETHODCALLTYPE FWindowsMMNotificationClient::OnStateChanged(AudioSessionState NewState)
+	{
+		return S_OK;
+	}
+
+	HRESULT STDMETHODCALLTYPE FWindowsMMNotificationClient::OnGroupingParamChanged(LPCGUID NewGroupingParam, LPCGUID EventContext)
+	{
+		return S_OK;
+	}
+
+	HRESULT STDMETHODCALLTYPE FWindowsMMNotificationClient::OnChannelVolumeChanged(DWORD ChannelCount, float NewChannelVolumeArray[], DWORD ChangedChannel, LPCGUID EventContext)
+	{
+		return S_OK;
+	}
+
+	HRESULT STDMETHODCALLTYPE FWindowsMMNotificationClient::OnSimpleVolumeChanged(float NewVolume, BOOL NewMute, LPCGUID EventContext)
+	{
+		return S_OK;
+	}
+
+	HRESULT STDMETHODCALLTYPE FWindowsMMNotificationClient::OnIconPathChanged(LPCWSTR NewIconPath, LPCGUID EventContext)
+	{
+		return S_OK;
+	}
+
+	HRESULT STDMETHODCALLTYPE FWindowsMMNotificationClient::OnDisplayNameChanged(LPCWSTR NewDisplayName, LPCGUID EventContext)
+	{
+		return S_OK;
+	}
+
+	void FWindowsMMNotificationClient::UnRegisterDeviceDeviceChangedListener(Audio::IAudioMixerDeviceChangedListener* DeviceChangedListener)
+	{
+		// Modifying container so get full write lock
+		FWriteScopeLock Lock(ListenersSetRwLock);
+		Listeners.Remove(DeviceChangedListener);
+	}
+
+	void FWindowsMMNotificationClient::RegisterDeviceChangedListener(Audio::IAudioMixerDeviceChangedListener* DeviceChangedListener)
+	{
+		// Modifying container so get full write lock
+		FWriteScopeLock Lock(ListenersSetRwLock);
+		Listeners.Add(DeviceChangedListener);
+	}
+
+	#include "Windows/AllowWindowsPlatformAtomics.h"
+
+	ULONG FWindowsMMNotificationClient::AddRef()
+	{
+		return InterlockedIncrement(&Ref);
+	}
+
+	ULONG FWindowsMMNotificationClient::Release() 
+	{
+		ULONG ulRef = InterlockedDecrement(&Ref);
+		if (0 == ulRef)
+		{
+			delete this;
+		}
+		return ulRef;
+	}
+
+	#include "Windows/HideWindowsPlatformAtomics.h"
+
+	HRESULT STDMETHODCALLTYPE FWindowsMMNotificationClient::QueryInterface(const IID& IId, void** UnknownPtrPtr)
+	{
+		// Three rules of QueryInterface: https://docs.microsoft.com/en-us/windows/win32/com/rules-for-implementing-queryinterface
+		// 1. Objects must have identity.
+		// 2. The set of interfaces on an object instance must be static.
+		// 3. It must be possible to query successfully for any interface on an object from any other interface.
+
+		// If ppvObject(the address) is nullptr, then this method returns E_POINTER.
+		if (!UnknownPtrPtr)
+		{
+			return E_POINTER;
+		}
+
+		// https://docs.microsoft.com/en-us/windows/win32/com/implementing-reference-counting
+		// Whenever a client calls a method(or API function), such as QueryInterface, that returns a new interface pointer, 
+		// the method being called is responsible for incrementing the reference count through the returned pointer.
+		// For example, when a client first creates an object, it receives an interface pointer to an object that, 
+		// from the client's point of view, has a reference count of one. If the client then calls AddRef on the interface pointer, 
+		// the reference count becomes two. The client must call Release twice on the interface pointer to drop all of its references to the object.
+		if (IId == __uuidof(IMMNotificationClient) || IId == __uuidof(IUnknown))
+		{
+			*UnknownPtrPtr = (IMMNotificationClient*)(this);
+			AddRef();
+			return S_OK;
+		}
+		else if (IId == __uuidof(IAudioSessionEvents))
+		{
+			*UnknownPtrPtr = (IAudioSessionEvents*)this;
+			AddRef();
+			return S_OK;
+		}
+
+
+		// This method returns S_OK if the interface is supported, and E_NOINTERFACE otherwise.
+		*UnknownPtrPtr = nullptr;
+		return E_NOINTERFACE;
+	}
+
+	HRESULT STDMETHODCALLTYPE FWindowsMMNotificationClient::OnPropertyValueChanged(LPCWSTR pwstrDeviceId, const PROPERTYKEY key)
+	{
+		UE_CLOG(Audio::IAudioMixer::ShouldLogDeviceSwaps(), LogAudioMixer, Verbose, TEXT("OnPropertyValueChanged: %s : %s"), *GetFriendlyName(pwstrDeviceId), *ToFString(key));
+
+		if (key.fmtid == PKEY_AudioEngine_DeviceFormat.fmtid)
+		{
+			// Get device.
+			FString DeviceId = pwstrDeviceId;
+			TComPtr<IMMDevice> Device;
+			HRESULT Hr = DeviceEnumerator->GetDevice(*DeviceId, &Device);
+
+			// Get property store.
+			TComPtr<IPropertyStore> PropertyStore;
+			if (SUCCEEDED(Hr) && Device)
+			{
+				Hr = Device->OpenPropertyStore(STGM_READ, &PropertyStore);
+				if (SUCCEEDED(Hr) && PropertyStore)
+				{
+					// Device Format
+					PROPVARIANT Prop;
+					PropVariantInit(&Prop);
+
+					if (key.fmtid == PKEY_AudioEngine_DeviceFormat.fmtid)
+					{
+						// WAVEFORMATEX blobs.
+						if (SUCCEEDED(PropertyStore->GetValue(key, &Prop)) && Prop.blob.pBlobData)
+						{
+							const WAVEFORMATEX* WaveFormatEx = (const WAVEFORMATEX*)(Prop.blob.pBlobData);
+
+							Audio::IAudioMixerDeviceChangedListener::FFormatChangedData FormatChanged;
+							FormatChanged.NumChannels = FMath::Clamp((int32)WaveFormatEx->nChannels, 2, 8);
+							FormatChanged.SampleRate = WaveFormatEx->nSamplesPerSec;
+							FormatChanged.ChannelBitmask = WaveFormatEx->wFormatTag == WAVE_FORMAT_EXTENSIBLE ?
+								((const WAVEFORMATEXTENSIBLE*)WaveFormatEx)->dwChannelMask : 0;
+
+							FReadScopeLock ReadLock(ListenersSetRwLock);
+							for (Audio::IAudioMixerDeviceChangedListener* i : Listeners)
+							{
+								i->OnFormatChanged(DeviceId, FormatChanged);
+							}
+						}
+					}
+					PropVariantClear(&Prop);
+				}
+			}
+		}
+		return S_OK;
+	}
+
+	TComPtr<IMMDevice> FWindowsMMNotificationClient::GetDevice(const FString InDeviceID) const
+	{
+		// Get device.
+		TComPtr<IMMDevice> Device;
+		HRESULT Hr = DeviceEnumerator->GetDevice(*InDeviceID, &Device);
+		if (SUCCEEDED(Hr))
+		{
+			return Device;
+		}
+
+		// Fail.
+		return {};
+	}
+
+	FString FWindowsMMNotificationClient::GetFriendlyName(const TComPtr<IMMDevice>& InDevice)
+	{
+		FString FriendlyName = TEXT("[No Friendly Name for Device]");
+
+		if (InDevice)
+		{
+			// Get property store.
+			TComPtr<IPropertyStore> PropStore;
+			HRESULT Hr = InDevice->OpenPropertyStore(STGM_READ, &PropStore);
+
+			// Get friendly name.
+			if (SUCCEEDED(Hr) && PropStore)
+			{
+				PROPVARIANT PropString;
+				PropVariantInit(&PropString);
+
+				// Get the endpoint device's friendly-name property.
+				Hr = PropStore->GetValue(PKEY_Device_FriendlyName, &PropString);
+				if (SUCCEEDED(Hr))
+				{
+					// Copy friendly name.
+					if (PropString.pwszVal)
+					{
+						FriendlyName = PropString.pwszVal;
+					}
+				}
+
+				PropVariantClear(&PropString);
+			}
+		}
+		return FriendlyName;
+	}
+
+	FString FWindowsMMNotificationClient::GetFriendlyName(const FString InDeviceID)
+	{
+		if (InDeviceID.IsEmpty())
+		{
+			return TEXT("System Default");
+		}
+
+		FString FriendlyName = TEXT("[No Friendly Name for Device]");
+
+		// Get device.
+		if (TComPtr<IMMDevice> Device = GetDevice(*InDeviceID))
+		{
+			return GetFriendlyName(Device);
+		}
+		return FriendlyName;
+	}
+
+
+	HRESULT STDMETHODCALLTYPE FWindowsMMNotificationClient::OnDeviceStateChanged(LPCWSTR pwstrDeviceId, DWORD dwNewState)
+	{
+		UE_CLOG(Audio::IAudioMixer::ShouldLogDeviceSwaps(), LogAudioMixer, Display, TEXT("FWindowsMMNotificationClient: OnDeviceStateChanged: %s, %d"), *GetFriendlyName(pwstrDeviceId), dwNewState);
+
+		if (Audio::IAudioMixer::ShouldIgnoreDeviceSwaps())
+		{
+			return S_OK;
+		}
+
+		bool bIsRender = IsRenderDevice(pwstrDeviceId);
+		if (dwNewState == DEVICE_STATE_ACTIVE || dwNewState == DEVICE_STATE_DISABLED || dwNewState == DEVICE_STATE_UNPLUGGED || dwNewState == DEVICE_STATE_NOTPRESENT)
+		{
+			Audio::EAudioDeviceState State = ConvertWordToDeviceState(dwNewState);
+
+			FString DeviceString(pwstrDeviceId);
+			FReadScopeLock ReadLock(ListenersSetRwLock);
+			for (Audio::IAudioMixerDeviceChangedListener* Listener : Listeners)
+			{
+				Listener->OnDeviceStateChanged(DeviceString, State, bIsRender);
+			}
+		}
+		return S_OK;
+	}
+
+	HRESULT STDMETHODCALLTYPE FWindowsMMNotificationClient::OnDeviceRemoved(LPCWSTR pwstrDeviceId)
+	{
+		UE_CLOG(Audio::IAudioMixer::ShouldLogDeviceSwaps(), LogAudioMixer, Display, TEXT("FWindowsMMNotificationClient: OnDeviceRemoved: %s"), *GetFriendlyName(pwstrDeviceId));
+
+		if (Audio::IAudioMixer::ShouldIgnoreDeviceSwaps())
+		{
+			return S_OK;
+		}
+
+		bool bIsRender = IsRenderDevice(pwstrDeviceId);
+		FString DeviceString(pwstrDeviceId);
+		FReadScopeLock ReadLock(ListenersSetRwLock);
+		for (Audio::IAudioMixerDeviceChangedListener* Listener : Listeners)
+		{
+			Listener->OnDeviceRemoved(DeviceString, bIsRender);
+		}
+		return S_OK;
+	}
+
+	HRESULT STDMETHODCALLTYPE FWindowsMMNotificationClient::OnDeviceAdded(LPCWSTR pwstrDeviceId)
+	{
+		UE_CLOG(Audio::IAudioMixer::ShouldLogDeviceSwaps(), LogAudioMixer, Display, TEXT("FWindowsMMNotificationClient: OnDeviceAdded: %s"), *GetFriendlyName(pwstrDeviceId));
+
+		if (Audio::IAudioMixer::ShouldIgnoreDeviceSwaps())
+		{
+			return S_OK;
+		}
+
+		bool bIsRender = IsRenderDevice(pwstrDeviceId);
+		FString DeviceString(pwstrDeviceId);
+		FReadScopeLock ReadLock(ListenersSetRwLock);
+		for (Audio::IAudioMixerDeviceChangedListener* Listener : Listeners)
+		{
+			Listener->OnDeviceAdded(DeviceString, bIsRender);
+		}
+		return S_OK;
+	}
+
+	bool FWindowsMMNotificationClient::IsRenderDevice(const FString& InDeviceId) const
+	{
+		bool bIsRender = true;
+		if (TComPtr<IMMDevice> Device = GetDevice(InDeviceId))
+		{
+			TComPtr<IMMEndpoint> Endpoint;
+			if (SUCCEEDED(Device->QueryInterface(IID_PPV_ARGS(&Endpoint))))
+			{
+				EDataFlow DataFlow = eRender;
+				if (SUCCEEDED(Endpoint->GetDataFlow(&DataFlow)))
+				{
+					bIsRender = DataFlow == eRender;
+				}
+			}
+		}
+		return bIsRender;
+	}
+
+	HRESULT STDMETHODCALLTYPE FWindowsMMNotificationClient::OnDefaultDeviceChanged(EDataFlow InFlow, ERole InRole, LPCWSTR pwstrDeviceId)
+	{
+		UE_CLOG(Audio::IAudioMixer::ShouldLogDeviceSwaps(), LogAudioMixer, Warning,
+			TEXT("FWindowsMMNotificationClient: OnDefaultDeviceChanged: %s, %s, %s - %s"), ToString(InFlow), ToString(InRole), pwstrDeviceId, *GetFriendlyName(pwstrDeviceId));
+
+		Audio::EAudioDeviceRole AudioDeviceRole;
+
+		if (Audio::IAudioMixer::ShouldIgnoreDeviceSwaps())
+		{
+			return S_OK;
+		}
+
+		if (InRole == eConsole)
+		{
+			AudioDeviceRole = Audio::EAudioDeviceRole::Console;
+		}
+		else if (InRole == eMultimedia)
+		{
+			AudioDeviceRole = Audio::EAudioDeviceRole::Multimedia;
+		}
+		else
+		{
+			AudioDeviceRole = Audio::EAudioDeviceRole::Communications;
+		}
+
+		FString DeviceString(pwstrDeviceId);
+		if (InFlow == eRender)
+		{
+			FReadScopeLock Lock(ListenersSetRwLock);
+			for (Audio::IAudioMixerDeviceChangedListener* Listener : Listeners)
+			{
+				Listener->OnDefaultRenderDeviceChanged(AudioDeviceRole, DeviceString);
+			}
+		}
+		else if (InFlow == eCapture)
+		{
+			FReadScopeLock Lock(ListenersSetRwLock);
+			for (Audio::IAudioMixerDeviceChangedListener* Listener : Listeners)
+			{
+				Listener->OnDefaultCaptureDeviceChanged(AudioDeviceRole, DeviceString);
+			}
+		}
+		else
+		{
+			FReadScopeLock Lock(ListenersSetRwLock);
+			for (Audio::IAudioMixerDeviceChangedListener* Listener : Listeners)
+			{
+				Listener->OnDefaultCaptureDeviceChanged(AudioDeviceRole, DeviceString);
+				Listener->OnDefaultRenderDeviceChanged(AudioDeviceRole, DeviceString);
+			}
+		}
+		return S_OK;
+	}
+
+	FWindowsMMNotificationClient::~FWindowsMMNotificationClient()
+	{
+		UnregisterForSessionNotifications();
+
+		if (DeviceEnumerator)
+		{
+			DeviceEnumerator->UnregisterEndpointNotificationCallback(this);
+		}
+
+		if (bComInitialized)
+		{
+			FWindowsPlatformMisc::CoUninitialize();
+		}
+	}
+
+	void FWindowsMMNotificationClient::UnregisterForSessionNotifications()
+	{
+		FScopeLock Lock(&SessionRegistrationCS);
+
+		// Unregister for any device we're already listening to.
+		if (SessionControls)
+		{
+			UE_LOG(LogAudioMixer, Verbose, TEXT("FWindowsMMNotificationClient: Unregistering for sessions events for device '%s'"), DeviceListeningToSessionEvents ? *GetFriendlyName(DeviceListeningToSessionEvents.Get()) : TEXT("None"));
+			SessionControls->UnregisterAudioSessionNotification(this);
+			SessionControls.Reset();
+		}
+		if (SessionManager)
+		{
+			SessionManager.Reset();
+		}
+
+		DeviceListeningToSessionEvents.Reset();
+
+		// Reset this flag.
+		bHasDisconnectSessionHappened = false;
+	}
+
+	FWindowsMMNotificationClient::FWindowsMMNotificationClient() 
+		: Ref(1)
+		, DeviceEnumerator(nullptr)
+		, bHasDisconnectSessionHappened(false)
+	{
+		bComInitialized = FWindowsPlatformMisc::CoInitialize();
+		HRESULT Result = CoCreateInstance(__uuidof(MMDeviceEnumerator), nullptr, CLSCTX_INPROC_SERVER, IID_PPV_ARGS(&DeviceEnumerator));
+		if (Result == S_OK)
+		{
+			DeviceEnumerator->RegisterEndpointNotificationCallback(this);
+		}
+
+		// Register for session events from default endpoint.
+		if (DeviceEnumerator)
+		{
+			TComPtr<IMMDevice> DefaultDevice;
+			if (SUCCEEDED(DeviceEnumerator->GetDefaultAudioEndpoint(eRender, eConsole, &DefaultDevice)))
+			{
+				RegisterForSessionNotifications(DefaultDevice);
+			}
+		}
+	}
+
+}// namespace Audio
+
+#endif //PLATFORM_WINDOWS
diff --git a/Engine/Source/Runtime/Windows/AudioMixerXAudio2/Private/WindowsMMNotificationClient.h b/Engine/Source/Runtime/Windows/AudioMixerXAudio2/Private/WindowsMMNotificationClient.h
new file mode 100644
--- /dev/null
+++ b/Engine/Source/Runtime/Windows/AudioMixerXAudio2/Private/WindowsMMNotificationClient.h
@@ -0,0 +1,118 @@
+// Copyright Epic Games, Inc. All Rights Reserved.
+
+#pragma once
+
+#if PLATFORM_WINDOWS
+
+#include <atomic>
+#include "AudioMixer.h"	
+#include "Misc/ScopeRWLock.h"
+
+#include "Windows/AllowWindowsPlatformTypes.h"
+#include "Windows/AllowWindowsPlatformAtomics.h"
+
+#include "Windows/MinWindows.h"
+#include "Windows/COMPointer.h"
+
+THIRD_PARTY_INCLUDES_START
+#include <mmdeviceapi.h>			// IMMNotificationClient
+#include <audiopolicy.h>			// IAudioSessionEvents
+THIRD_PARTY_INCLUDES_END
+
+#include "Windows/HideWindowsPlatformAtomics.h"
+#include "Windows/HideWindowsPlatformTypes.h"
+
+namespace Audio
+{
+	class IAudioMixerDeviceChangedListener;
+
+	class FWindowsMMNotificationClient final : public IMMNotificationClient, public IAudioSessionEvents
+	{
+	public:
+		FWindowsMMNotificationClient();
+
+		bool RegisterForSessionNotifications(const TComPtr<IMMDevice>& InDevice);
+
+		bool RegisterForSessionNotifications(const FString& InDeviceId);
+
+		void UnregisterForSessionNotifications();
+
+		~FWindowsMMNotificationClient();
+
+		HRESULT STDMETHODCALLTYPE OnDefaultDeviceChanged(EDataFlow InFlow, ERole InRole, LPCWSTR pwstrDeviceId) override;
+
+		// TODO: Ideally we'd use the cache instead of ask for this.
+		bool IsRenderDevice(const FString& InDeviceId) const;
+
+		HRESULT STDMETHODCALLTYPE OnDeviceAdded(LPCWSTR pwstrDeviceId) override;;
+
+		HRESULT STDMETHODCALLTYPE OnDeviceRemoved(LPCWSTR pwstrDeviceId) override;
+
+		HRESULT STDMETHODCALLTYPE OnDeviceStateChanged(LPCWSTR pwstrDeviceId, DWORD dwNewState) override;
+
+		FString GetFriendlyName(const FString InDeviceID);
+		FString GetFriendlyName(const TComPtr<IMMDevice>& InDevice);
+
+		TComPtr<IMMDevice> GetDevice(const FString InDeviceID) const;
+
+		HRESULT STDMETHODCALLTYPE OnPropertyValueChanged(LPCWSTR pwstrDeviceId, const PROPERTYKEY key);
+
+		HRESULT STDMETHODCALLTYPE QueryInterface(const IID& IId, void** UnknownPtrPtr) override;
+
+		ULONG STDMETHODCALLTYPE AddRef() override;
+
+		ULONG STDMETHODCALLTYPE Release() override;
+
+		void RegisterDeviceChangedListener(Audio::IAudioMixerDeviceChangedListener* DeviceChangedListener);
+
+		void UnRegisterDeviceDeviceChangedListener(Audio::IAudioMixerDeviceChangedListener* DeviceChangedListener);
+
+		// Begin IAudioSessionEvents
+		HRESULT STDMETHODCALLTYPE OnDisplayNameChanged(
+			LPCWSTR NewDisplayName,
+			LPCGUID EventContext) override;
+
+		HRESULT STDMETHODCALLTYPE OnIconPathChanged(
+			LPCWSTR NewIconPath,
+			LPCGUID EventContext) override;
+
+		HRESULT STDMETHODCALLTYPE OnSimpleVolumeChanged(
+			float NewVolume,
+			BOOL NewMute,
+			LPCGUID EventContext) override;
+
+		HRESULT STDMETHODCALLTYPE OnChannelVolumeChanged(
+			DWORD ChannelCount,
+			float NewChannelVolumeArray[],
+			DWORD ChangedChannel,
+			LPCGUID EventContext) override;
+
+		HRESULT STDMETHODCALLTYPE OnGroupingParamChanged(
+			LPCGUID NewGroupingParam,
+			LPCGUID EventContext) override;
+
+		HRESULT STDMETHODCALLTYPE OnStateChanged(
+			AudioSessionState NewState) override;
+
+		HRESULT STDMETHODCALLTYPE OnSessionDisconnected(
+			AudioSessionDisconnectReason InDisconnectReason);
+		// End IAudioSessionEvents
+
+	private:
+		LONG Ref;
+		TSet<Audio::IAudioMixerDeviceChangedListener*> Listeners;
+		FRWLock ListenersSetRwLock;
+
+		TComPtr<IMMDeviceEnumerator> DeviceEnumerator;
+
+		FCriticalSection SessionRegistrationCS;
+		TComPtr<IAudioSessionManager> SessionManager;
+		TComPtr<IAudioSessionControl> SessionControls;
+		TComPtr<IMMDevice> DeviceListeningToSessionEvents;
+
+		bool bComInitialized;
+		std::atomic<bool> bHasDisconnectSessionHappened;
+	};
+}// namespace Audio
+
+#endif //PLATFORM_WINDOWS
