Fix CIS non-windows, nounity
#rnx

diff --git a/Engine/Source/Runtime/Windows/AudioMixerXAudio2/Private/WindowsMMDeviceInfoCache.h b/Engine/Source/Runtime/Windows/AudioMixerXAudio2/Private/WindowsMMDeviceInfoCache.h
--- a/Engine/Source/Runtime/Windows/AudioMixerXAudio2/Private/WindowsMMDeviceInfoCache.h
+++ b/Engine/Source/Runtime/Windows/AudioMixerXAudio2/Private/WindowsMMDeviceInfoCache.h
@@ -2,6 +2,8 @@
 
 #pragma once
 
+#include "HAL/Platform.h"
+
 #if PLATFORM_WINDOWS
 
 #include "AudioMixer.h"
diff --git a/Engine/Source/Runtime/Windows/AudioMixerXAudio2/Private/WindowsMMNotificationClient.h b/Engine/Source/Runtime/Windows/AudioMixerXAudio2/Private/WindowsMMNotificationClient.h
--- a/Engine/Source/Runtime/Windows/AudioMixerXAudio2/Private/WindowsMMNotificationClient.h
+++ b/Engine/Source/Runtime/Windows/AudioMixerXAudio2/Private/WindowsMMNotificationClient.h
@@ -2,6 +2,8 @@
 
 #pragma once
 
+#include "HAL/Platform.h"
+
 #if PLATFORM_WINDOWS
 
 #include <atomic>
