Another no unity, non-windows CIS fix.
#rnx

diff --git a/Engine/Source/Runtime/Windows/AudioMixerXAudio2/Private/WindowsMMNotificationClient.cpp b/Engine/Source/Runtime/Windows/AudioMixerXAudio2/Private/WindowsMMNotificationClient.cpp
--- a/Engine/Source/Runtime/Windows/AudioMixerXAudio2/Private/WindowsMMNotificationClient.cpp
+++ b/Engine/Source/Runtime/Windows/AudioMixerXAudio2/Private/WindowsMMNotificationClient.cpp
@@ -1,11 +1,9 @@
 // Copyright Epic Games, Inc. All Rights Reserved.
 
-#pragma once
+#include "WindowsMMNotificationClient.h"
 
 #if PLATFORM_WINDOWS
 
-#include "WindowsMMNotificationClient.h"
-
 #include "Windows/AllowWindowsPlatformTypes.h"
 
 #include "Windows/MinWindows.h"
