Workaround for static analysis warning C6385, if we pass in SupportedExtensionsVersionCount to INTC_D3D11_GetSupportedVersions the second time, the static analyzer thinks it may be different from the first call.

diff --git a/Engine/Source/Runtime/Windows/D3D11RHI/Private/Windows/WindowsD3D11Device.cpp b/Engine/Source/Runtime/Windows/D3D11RHI/Private/Windows/WindowsD3D11Device.cpp
--- a/Engine/Source/Runtime/Windows/D3D11RHI/Private/Windows/WindowsD3D11Device.cpp
+++ b/Engine/Source/Runtime/Windows/D3D11RHI/Private/Windows/WindowsD3D11Device.cpp
@@ -1377,8 +1377,11 @@ void FD3D11DynamicRHI::StartIntelExtensions()
 		SupportedExtensionsVersions = new INTCExtensionVersion[SupportedExtensionsVersionCount]{};
 	}
 
-	if (SUCCEEDED(INTC_D3D11_GetSupportedVersions(Direct3DDevice, SupportedExtensionsVersions, &SupportedExtensionsVersionCount)) && SupportedExtensionsVersions != nullptr)
+	// Workaround for C6385, if we pass in SupportedExtensionsVersionCount again, the static analyzer thinks it may be different from the first call
+	uint32_t DummyCount = 0;
+	if (SUCCEEDED(INTC_D3D11_GetSupportedVersions(Direct3DDevice, SupportedExtensionsVersions, &DummyCount)) && SupportedExtensionsVersions != nullptr)
 	{
+		check(SupportedExtensionsVersionCount == DummyCount);
 		for (uint32_t i = 0; i < SupportedExtensionsVersionCount; i++)
 		{
 			if ((SupportedExtensionsVersions[i].HWFeatureLevel >= AtomicsRequiredVersion.HWFeatureLevel) &&
