Fix AudioMixer XAudio2 to recover after failing to init, allowing for device swap after a open failture. Also strip out DLL fallback code as was causing issues with init recovery.

diff --git a/Engine/Source/Runtime/Windows/AudioMixerXAudio2/Private/AudioMixerPlatformXAudio2.cpp b/Engine/Source/Runtime/Windows/AudioMixerXAudio2/Private/AudioMixerPlatformXAudio2.cpp
--- a/Engine/Source/Runtime/Windows/AudioMixerXAudio2/Private/AudioMixerPlatformXAudio2.cpp
+++ b/Engine/Source/Runtime/Windows/AudioMixerXAudio2/Private/AudioMixerPlatformXAudio2.cpp
@@ -172,25 +172,14 @@ static XAUDIO2_PROCESSOR GetXAudio2ProcessorsToUse()
 }
 
 #if PLATFORM_WINDOWS || PLATFORM_HOLOLENS
-FName GetNextDllToTry(FName Current = NAME_None) 
+FName GetDllName(FName Current = NAME_None) 
 {
 #if PLATFORM_64BITS
 	static const FString XAudio2_9Redist = FPaths::EngineDir() / TEXT("Binaries/ThirdParty/Windows/XAudio2_9/x64/xaudio2_9redist.dll");
 #else
 	static const FString XAudio2_9Redist = FPaths::EngineDir() / TEXT("Binaries/ThirdParty/Windows/XAudio2_9/x86/xaudio2_9redist.dll");
 #endif
-	static const TArray<FName> kDllsToTry = { *XAudio2_9Redist, TEXT("XAudio2_7.dll") };
-	int32 Index = kDllsToTry.IndexOfByKey(Current);
-	
-	// This behaves like a link list of names. 
-	// Current=Name_NONE will have an index of -1, thus yield Index 0	
-	// until we reach end of the list and return Name_NONE.
-
-	if (kDllsToTry.IsValidIndex(Index + 1))			
-	{
-		return kDllsToTry[Index + 1];
-	}
-	return NAME_None;
+	return *XAudio2_9Redist;
 }
 #endif //#if PLATFORM_WINDOWS || PLATFORM_HOLOLENS
 
@@ -257,7 +246,7 @@ namespace Audio
 	{
 #if PLATFORM_WINDOWS || PLATFORM_HOLOLENS
 		FPlatformMisc::CoInitialize();
-		DllName = GetNextDllToTry();
+		DllName = GetDllName();
 #endif // #if PLATFORM_WINDOWS || PLATFORM_HOLOLENS
 	}
 
@@ -1363,7 +1352,7 @@ namespace Audio
 		AudioStreamInfo.NumBuffers = OpenStreamParams.NumBuffers;
 		AudioStreamInfo.AudioMixer = OpenStreamParams.AudioMixer;
 
-		uint32 NumOutputDevices;
+		uint32 NumOutputDevices = 0;
 		HRESULT Result = ERROR_SUCCESS;
 
 		if (GetNumOutputDevices(NumOutputDevices) && NumOutputDevices > 0)
@@ -1410,7 +1399,7 @@ namespace Audio
 				OriginalAudioDeviceId = AudioStreamInfo.DeviceInfo.DeviceId;
 			}
 
-#if PLATFORM_WINDOWS
+#if PLATFORM_WINDOWS			
 			Result = XAudio2System->CreateMasteringVoice(
 				&OutputAudioStreamMasteringVoice,
 				AudioStreamInfo.DeviceInfo.NumChannels,
@@ -1456,10 +1445,24 @@ namespace Audio
 			Result = XAudio2System->CreateSourceVoice(&OutputAudioStreamSourceVoice, &Format, XAUDIO2_VOICE_NOPITCH, 2.0f, &OutputVoiceCallback);
 			XAUDIO2_RETURN_ON_FAIL(Result);
 		}
-		else
+	Cleanup:
+
+		bool bXAudioOpenSuccessfully = OutputAudioStreamSourceVoice && OutputAudioStreamMasteringVoice;
+		if (!bXAudioOpenSuccessfully)
 		{
-			check(!bIsUsingNullDevice);
+			// Undo anything we created.
+			if (OutputAudioStreamSourceVoice)
+			{
+				OutputAudioStreamSourceVoice->DestroyVoice();
+				OutputAudioStreamSourceVoice = nullptr;
+			}		
+			if (OutputAudioStreamMasteringVoice)
+			{
+				OutputAudioStreamMasteringVoice->DestroyVoice();
+				OutputAudioStreamMasteringVoice = nullptr;
+			}
 
+			// Setup for running null device.
 			AudioStreamInfo.NumOutputFrames = OpenStreamParams.NumFrames;
 			AudioStreamInfo.DeviceInfo.OutputChannelArray = { EAudioMixerChannel::FrontLeft, EAudioMixerChannel::FrontRight };
 			AudioStreamInfo.DeviceInfo.NumChannels = 2;
@@ -1467,38 +1470,24 @@ namespace Audio
 			AudioStreamInfo.DeviceInfo.Format = EAudioMixerStreamDataFormat::Float;
 		}
 
-		AudioStreamInfo.StreamState = EAudioOutputStreamState::Open;
-		bIsDeviceOpen = true;
-
-	Cleanup:
-		if (FAILED(Result))
-		{		
-			bool bFailedToCreateMasterVoice = OutputAudioStreamMasteringVoice == nullptr;
-			CloseAudioStream();
-
 #if PLATFORM_WINDOWS
-			if (bFailedToCreateMasterVoice)
-			{
-				// Try another DLL.
-				DllName = GetNextDllToTry(DllName);
-				if (DllName != NAME_None)
-				{
-					UE_LOG(LogAudioMixer, Warning, TEXT("Xaudio 2.9: Failed to create Master Voice. Attempting fallback DLL"));
-					TeardownHardware();
-					if (InitializeHardware())
-					{
-						return OpenAudioStream(Params);
-					}
-				}
-				else
-				{
-					UE_LOG(LogAudioMixer, Warning, TEXT("Xaudio 2.9: Failed to create Master Voice. Out of DLL fallbacks"));
-				}
-			}
-#endif //PLATFORM_WINDOWS
-
+		if(!bXAudioOpenSuccessfully)
+		{
+			// On Windows where we can have audio devices unplugged/removed/hot-swapped:
+			// We must mark ourselves open, even if we failed to open. This will allow the device-swap logic to run.
+			// StartAudioStream will happily use the null renderer path if there's no real stream open.
+			bXAudioOpenSuccessfully = true;
 		}
-		return SUCCEEDED(Result);
+#endif //PLATFORM_WINDOWS		
+
+		// If we opened, mark the stream as open.
+		if(bXAudioOpenSuccessfully)
+		{
+			AudioStreamInfo.StreamState = EAudioOutputStreamState::Open;
+			bIsDeviceOpen = true;
+		}	
+
+		return bXAudioOpenSuccessfully;
 	}
 
 	FAudioPlatformDeviceInfo FMixerPlatformXAudio2::GetPlatformDeviceInfo() const
