Incremental Refactor XAudio2, move stringify error codes into StringHelper

diff --git a/Engine/Source/Runtime/Windows/AudioMixerXAudio2/Private/AudioMixerPlatformXAudio2.cpp b/Engine/Source/Runtime/Windows/AudioMixerXAudio2/Private/AudioMixerPlatformXAudio2.cpp
--- a/Engine/Source/Runtime/Windows/AudioMixerXAudio2/Private/AudioMixerPlatformXAudio2.cpp
+++ b/Engine/Source/Runtime/Windows/AudioMixerXAudio2/Private/AudioMixerPlatformXAudio2.cpp
@@ -37,6 +37,7 @@
 #include "HAL/Event.h"
 #include "CoreMinimal.h"
 #include "Logging/LogMacros.h"
+#include "ToStringHelpers.h"
 
 #if PLATFORM_WINDOWS
 THIRD_PARTY_INCLUDES_START
@@ -52,7 +53,7 @@ THIRD_PARTY_INCLUDES_END
 
 #define XAUDIO2_LOG_RESULT(FunctionName, Result) \
 	{ \
-		FString ErrorString = FString::Printf(TEXT("%s -> 0x%X: %s (line: %d)"), TEXT( FunctionName ), Result, *GetErrorString(Result), __LINE__); \
+		FString ErrorString = FString::Printf(TEXT("%s -> 0x%X: %s (line: %d)"), TEXT( FunctionName ), Result, *Audio::ToErrorFString(Result), __LINE__); \
 		UE_LOG(LogAudioMixer, Error, TEXT("XAudio2 Error: %s"), *ErrorString);																		\
 	}
 
@@ -60,7 +61,7 @@ THIRD_PARTY_INCLUDES_END
 #define XAUDIO2_GOTO_CLEANUP_ON_FAIL(Result)																										 \
 	if (FAILED(Result))																														 \
 	{																																		 \
-		FString ErrorString = FString::Printf(TEXT("%s -> 0x%X: %s (line: %d)"), TEXT( #Result ), Result, *GetErrorString(Result), __LINE__);\
+		FString ErrorString = FString::Printf(TEXT("%s -> 0x%X: %s (line: %d)"), TEXT( #Result ), Result, *Audio::ToErrorFString(Result), __LINE__);\
 		UE_LOG(LogAudioMixer, Error, TEXT("XAudio2 Error: %s"), *ErrorString);																 \
 		goto Cleanup;																														 \
 	}
@@ -69,89 +70,11 @@ THIRD_PARTY_INCLUDES_END
 #define XAUDIO2_RETURN_ON_FAIL(Result)																										 \
 	if (FAILED(Result))																														 \
 	{																																		 \
-		FString ErrorString = FString::Printf(TEXT("%s -> 0x%X: %s (line: %d)"), TEXT( #Result ), Result, *GetErrorString(Result), __LINE__);\
+		FString ErrorString = FString::Printf(TEXT("%s -> 0x%X: %s (line: %d)"), TEXT( #Result ), Result, *Audio::ToErrorFString(Result), __LINE__);\
 		UE_LOG(LogAudioMixer, Error, TEXT("XAudio2 Error: %s"), *ErrorString);																 \
 		return false;																														 \
 	}
 
-
-
-static FString GetErrorString(HRESULT Result)
-{
-#define CASE_AND_STRING(RESULT) case HRESULT(RESULT): return TEXT(#RESULT)
-
-	switch (Result)
-	{
-		case HRESULT(XAUDIO2_E_INVALID_CALL):			return TEXT("XAUDIO2_E_INVALID_CALL");
-		case HRESULT(XAUDIO2_E_XMA_DECODER_ERROR):		return TEXT("XAUDIO2_E_XMA_DECODER_ERROR");
-		case HRESULT(XAUDIO2_E_XAPO_CREATION_FAILED):	return TEXT("XAUDIO2_E_XAPO_CREATION_FAILED");
-		case HRESULT(XAUDIO2_E_DEVICE_INVALIDATED):		return TEXT("XAUDIO2_E_DEVICE_INVALIDATED");
-	#if PLATFORM_WINDOWS
-		case HRESULT(REGDB_E_CLASSNOTREG):				return TEXT("REGDB_E_CLASSNOTREG");
-		case HRESULT(CLASS_E_NOAGGREGATION):			return TEXT("CLASS_E_NOAGGREGATION");
-		case HRESULT(E_NOINTERFACE):					return TEXT("E_NOINTERFACE");
-		case HRESULT(E_POINTER):						return TEXT("E_POINTER");
-		case HRESULT(E_INVALIDARG):						return TEXT("E_INVALIDARG");
-		case HRESULT(E_OUTOFMEMORY):					return TEXT("E_OUTOFMEMORY");
-
-		// AudioClient.h
-		CASE_AND_STRING(AUDCLNT_E_NOT_INITIALIZED				 );
-		CASE_AND_STRING(AUDCLNT_E_ALREADY_INITIALIZED			 );
-		CASE_AND_STRING(AUDCLNT_E_WRONG_ENDPOINT_TYPE			 );
-		CASE_AND_STRING(AUDCLNT_E_DEVICE_INVALIDATED			 );
-		CASE_AND_STRING(AUDCLNT_E_NOT_STOPPED					 );
-		CASE_AND_STRING(AUDCLNT_E_BUFFER_TOO_LARGE				 );
-		CASE_AND_STRING(AUDCLNT_E_OUT_OF_ORDER					 );
-		CASE_AND_STRING(AUDCLNT_E_UNSUPPORTED_FORMAT			 );
-		CASE_AND_STRING(AUDCLNT_E_INVALID_SIZE					 );
-		CASE_AND_STRING(AUDCLNT_E_DEVICE_IN_USE					 );
-		CASE_AND_STRING(AUDCLNT_E_BUFFER_OPERATION_PENDING		 );
-		CASE_AND_STRING(AUDCLNT_E_THREAD_NOT_REGISTERED			 );
-		CASE_AND_STRING(AUDCLNT_E_EXCLUSIVE_MODE_NOT_ALLOWED	 );
-		CASE_AND_STRING(AUDCLNT_E_ENDPOINT_CREATE_FAILED		 );
-		CASE_AND_STRING(AUDCLNT_E_SERVICE_NOT_RUNNING			 );
-		CASE_AND_STRING(AUDCLNT_E_EVENTHANDLE_NOT_EXPECTED		 );
-		CASE_AND_STRING(AUDCLNT_E_EXCLUSIVE_MODE_ONLY			 );
-		CASE_AND_STRING(AUDCLNT_E_BUFDURATION_PERIOD_NOT_EQUAL	 );
-		CASE_AND_STRING(AUDCLNT_E_EVENTHANDLE_NOT_SET			 );
-		CASE_AND_STRING(AUDCLNT_E_INCORRECT_BUFFER_SIZE			 );
-		CASE_AND_STRING(AUDCLNT_E_BUFFER_SIZE_ERROR				 );
-		CASE_AND_STRING(AUDCLNT_E_CPUUSAGE_EXCEEDED				 );
-		CASE_AND_STRING(AUDCLNT_E_BUFFER_ERROR					 );
-		CASE_AND_STRING(AUDCLNT_E_BUFFER_SIZE_NOT_ALIGNED		 );
-		CASE_AND_STRING(AUDCLNT_E_INVALID_DEVICE_PERIOD			 );
-		CASE_AND_STRING(AUDCLNT_E_INVALID_STREAM_FLAG			 );
-		CASE_AND_STRING(AUDCLNT_E_ENDPOINT_OFFLOAD_NOT_CAPABLE	 );
-		CASE_AND_STRING(AUDCLNT_E_OUT_OF_OFFLOAD_RESOURCES		 );
-		CASE_AND_STRING(AUDCLNT_E_OFFLOAD_MODE_ONLY				 );
-		CASE_AND_STRING(AUDCLNT_E_NONOFFLOAD_MODE_ONLY			 );
-		CASE_AND_STRING(AUDCLNT_E_RESOURCES_INVALIDATED			 );
-		CASE_AND_STRING(AUDCLNT_E_RAW_MODE_UNSUPPORTED			 );
-		CASE_AND_STRING(AUDCLNT_E_ENGINE_PERIODICITY_LOCKED		 );
-		CASE_AND_STRING(AUDCLNT_E_ENGINE_FORMAT_LOCKED			 );
-		CASE_AND_STRING(AUDCLNT_S_BUFFER_EMPTY					 );
-		CASE_AND_STRING(AUDCLNT_S_THREAD_ALREADY_REGISTERED		 );
-		CASE_AND_STRING(AUDCLNT_S_POSITION_STALLED				 );
-
-		#undef CASE_AND_STRING
-
-	#endif
-		case HRESULT(0xe000020b):						return TEXT("ERROR_NO_SUCH_DEVINST");
-
-		default:
-		{
-			// We don't know this error, ask this system if it does.
-			TCHAR Buffer[1024] = { 0 };
-			FString Msg(FPlatformMisc::GetSystemErrorMessage(Buffer, UE_ARRAY_COUNT(Buffer), Result));
-
-			// Anything to return? Otherwise "UNKNOWN"
-			return !Msg.IsEmpty() ?
-				Msg :
-				TEXT("UNKNOWN");
-		}
-	}
-}
-
 static XAUDIO2_PROCESSOR GetXAudio2ProcessorsToUse()
 {
 	XAUDIO2_PROCESSOR ProcessorsToUse = (XAUDIO2_PROCESSOR)FPlatformAffinity::GetAudioThreadMask();
@@ -419,7 +342,7 @@ namespace Audio
 								if (FAILED(Result))
 								{
 									UE_LOG(LogAudioMixer, Error, TEXT("XAudio2System->CreateMasteringVoice() -> 0x%X: %s (line: %d) with Args (NumChannels=%u, SampleRate=%u, DeviceID=%s Name=%s)"),
-										Result, *GetErrorString(Result), __LINE__, (uint32)NewDevice.NumChannels, (uint32)NewDevice.SampleRate, *NewDevice.DeviceId, *NewDevice.Name);
+										Result, *Audio::ToErrorFString(Result), __LINE__, (uint32)NewDevice.NumChannels, (uint32)NewDevice.SampleRate, *NewDevice.DeviceId, *NewDevice.Name);
 									SAFE_RELEASE(Results.XAudio2System);
 									return {}; // FAIL.
 								}
diff --git a/Engine/Source/Runtime/Windows/AudioMixerXAudio2/Private/ToStringHelpers.cpp b/Engine/Source/Runtime/Windows/AudioMixerXAudio2/Private/ToStringHelpers.cpp
--- a/Engine/Source/Runtime/Windows/AudioMixerXAudio2/Private/ToStringHelpers.cpp
+++ b/Engine/Source/Runtime/Windows/AudioMixerXAudio2/Private/ToStringHelpers.cpp
@@ -11,8 +11,23 @@
 #pragma comment(lib, "Propsys.lib")
 #endif //!NO_LOGGING
 
+#endif //PLATFORM_WINDOWS
+
+#include "Windows/AllowWindowsPlatformTypes.h"
+
+THIRD_PARTY_INCLUDES_START
+#if PLATFORM_WINDOWS
+	#include <xaudio2redist.h>
+#else
+	#include <xaudio2.h>
+#endif
+THIRD_PARTY_INCLUDES_END
+
+#include "Windows/HideWindowsPlatformTypes.h"
+
 namespace Audio
 {
+#if PLATFORM_WINDOWS
 	const TCHAR* ToString(AudioSessionDisconnectReason InDisconnectReason)
 	{
 #if !NO_LOGGING
@@ -70,6 +85,35 @@ namespace Audio
 		return TEXT("Unknown");
 	}
 
+	FString ToFString(const PROPERTYKEY InKey)
+	{
+#if !NO_LOGGING
+#define IF_PROP_STRING(PROP) if(InKey.fmtid == PROP.fmtid) return TEXT(#PROP)
+
+		IF_PROP_STRING(PKEY_AudioEndpoint_PhysicalSpeakers);
+		IF_PROP_STRING(PKEY_AudioEngine_DeviceFormat);
+		IF_PROP_STRING(PKEY_AudioEngine_OEMFormat);
+		IF_PROP_STRING(PKEY_AudioEndpoint_Association);
+		IF_PROP_STRING(PKEY_AudioEndpoint_ControlPanelPageProvider);
+		IF_PROP_STRING(PKEY_AudioEndpoint_Disable_SysFx);
+		IF_PROP_STRING(PKEY_AudioEndpoint_FormFactor);
+		IF_PROP_STRING(PKEY_AudioEndpoint_FullRangeSpeakers);
+		IF_PROP_STRING(PKEY_AudioEndpoint_GUID);
+		IF_PROP_STRING(PKEY_AudioEndpoint_Supports_EventDriven_Mode);
+
+#undef IF_PROP_STRING
+
+		TCHAR KeyString[PKEYSTR_MAX];
+		HRESULT HR = PSStringFromPropertyKey(InKey, KeyString, ARRAYSIZE(KeyString));
+		if (SUCCEEDED(HR))
+		{
+			return FString(KeyString);
+		}
+#endif //!NO_LOGGING
+		return TEXT("Unknown");
+	}
+#endif //PLATFORM_WINDOWS
+
 	const TCHAR* ToString(EAudioDeviceRole InRole)
 	{
 #define CASE_TO_STRING(X) case EAudioDeviceRole::X: return TEXT(#X)
@@ -99,7 +143,84 @@ namespace Audio
 #undef CASE_TO_STRING
 	}
 
-	const FString ToFString(const TArray<EAudioMixerChannel::Type>& InChannels)
+	FString ToErrorFString(HRESULT InResult)
+	{
+#define CASE_AND_STRING(RESULT) case HRESULT(RESULT): return TEXT(#RESULT)
+
+		switch (InResult)
+		{
+		case HRESULT(XAUDIO2_E_INVALID_CALL):			return TEXT("XAUDIO2_E_INVALID_CALL");
+		case HRESULT(XAUDIO2_E_XMA_DECODER_ERROR):		return TEXT("XAUDIO2_E_XMA_DECODER_ERROR");
+		case HRESULT(XAUDIO2_E_XAPO_CREATION_FAILED):	return TEXT("XAUDIO2_E_XAPO_CREATION_FAILED");
+		case HRESULT(XAUDIO2_E_DEVICE_INVALIDATED):		return TEXT("XAUDIO2_E_DEVICE_INVALIDATED");
+#if PLATFORM_WINDOWS
+		case HRESULT(REGDB_E_CLASSNOTREG):				return TEXT("REGDB_E_CLASSNOTREG");
+		case HRESULT(CLASS_E_NOAGGREGATION):			return TEXT("CLASS_E_NOAGGREGATION");
+		case HRESULT(E_NOINTERFACE):					return TEXT("E_NOINTERFACE");
+		case HRESULT(E_POINTER):						return TEXT("E_POINTER");
+		case HRESULT(E_INVALIDARG):						return TEXT("E_INVALIDARG");
+		case HRESULT(E_OUTOFMEMORY):					return TEXT("E_OUTOFMEMORY");
+
+		// AudioClient.h
+		CASE_AND_STRING(AUDCLNT_E_NOT_INITIALIZED);
+		CASE_AND_STRING(AUDCLNT_E_ALREADY_INITIALIZED);
+		CASE_AND_STRING(AUDCLNT_E_WRONG_ENDPOINT_TYPE);
+		CASE_AND_STRING(AUDCLNT_E_DEVICE_INVALIDATED);
+		CASE_AND_STRING(AUDCLNT_E_NOT_STOPPED);
+		CASE_AND_STRING(AUDCLNT_E_BUFFER_TOO_LARGE);
+		CASE_AND_STRING(AUDCLNT_E_OUT_OF_ORDER);
+		CASE_AND_STRING(AUDCLNT_E_UNSUPPORTED_FORMAT);
+		CASE_AND_STRING(AUDCLNT_E_INVALID_SIZE);
+		CASE_AND_STRING(AUDCLNT_E_DEVICE_IN_USE);
+		CASE_AND_STRING(AUDCLNT_E_BUFFER_OPERATION_PENDING);
+		CASE_AND_STRING(AUDCLNT_E_THREAD_NOT_REGISTERED);
+		CASE_AND_STRING(AUDCLNT_E_EXCLUSIVE_MODE_NOT_ALLOWED);
+		CASE_AND_STRING(AUDCLNT_E_ENDPOINT_CREATE_FAILED);
+		CASE_AND_STRING(AUDCLNT_E_SERVICE_NOT_RUNNING);
+		CASE_AND_STRING(AUDCLNT_E_EVENTHANDLE_NOT_EXPECTED);
+		CASE_AND_STRING(AUDCLNT_E_EXCLUSIVE_MODE_ONLY);
+		CASE_AND_STRING(AUDCLNT_E_BUFDURATION_PERIOD_NOT_EQUAL);
+		CASE_AND_STRING(AUDCLNT_E_EVENTHANDLE_NOT_SET);
+		CASE_AND_STRING(AUDCLNT_E_INCORRECT_BUFFER_SIZE);
+		CASE_AND_STRING(AUDCLNT_E_BUFFER_SIZE_ERROR);
+		CASE_AND_STRING(AUDCLNT_E_CPUUSAGE_EXCEEDED);
+		CASE_AND_STRING(AUDCLNT_E_BUFFER_ERROR);
+		CASE_AND_STRING(AUDCLNT_E_BUFFER_SIZE_NOT_ALIGNED);
+		CASE_AND_STRING(AUDCLNT_E_INVALID_DEVICE_PERIOD);
+		CASE_AND_STRING(AUDCLNT_E_INVALID_STREAM_FLAG);
+		CASE_AND_STRING(AUDCLNT_E_ENDPOINT_OFFLOAD_NOT_CAPABLE);
+		CASE_AND_STRING(AUDCLNT_E_OUT_OF_OFFLOAD_RESOURCES);
+		CASE_AND_STRING(AUDCLNT_E_OFFLOAD_MODE_ONLY);
+		CASE_AND_STRING(AUDCLNT_E_NONOFFLOAD_MODE_ONLY);
+		CASE_AND_STRING(AUDCLNT_E_RESOURCES_INVALIDATED);
+		CASE_AND_STRING(AUDCLNT_E_RAW_MODE_UNSUPPORTED);
+		CASE_AND_STRING(AUDCLNT_E_ENGINE_PERIODICITY_LOCKED);
+		CASE_AND_STRING(AUDCLNT_E_ENGINE_FORMAT_LOCKED);
+		CASE_AND_STRING(AUDCLNT_S_BUFFER_EMPTY);
+		CASE_AND_STRING(AUDCLNT_S_THREAD_ALREADY_REGISTERED);
+		CASE_AND_STRING(AUDCLNT_S_POSITION_STALLED);
+
+#undef CASE_AND_STRING
+
+#endif //PLATFORM_WINDOWS
+
+		case HRESULT(0xe000020b):						return TEXT("ERROR_NO_SUCH_DEVINST");
+
+		default:
+		{
+			// We don't know this error, ask this system if it does.
+			TCHAR Buffer[1024] = { 0 };
+			FString Msg(FPlatformMisc::GetSystemErrorMessage(Buffer, UE_ARRAY_COUNT(Buffer), InResult));
+
+			// Anything to return? Otherwise "UNKNOWN"
+			return !Msg.IsEmpty() ?
+				Msg :
+				TEXT("UNKNOWN");
+		}
+		}
+	}
+
+	FString ToFString(const TArray<EAudioMixerChannel::Type>& InChannels)
 	{
 		FString ChannelString;
 		static const int32 ApproxChannelNameLength = 18;
@@ -110,36 +231,6 @@ namespace Audio
 			ChannelString.Append(TEXT("|"));
 		}
 		return ChannelString;
-	}
-
-	FString ToFString(const PROPERTYKEY Key)
-	{
-#if !NO_LOGGING
-#define IF_PROP_STRING(PROP) if(Key.fmtid == PROP.fmtid) return TEXT(#PROP)
-
-		IF_PROP_STRING(PKEY_AudioEndpoint_PhysicalSpeakers);
-		IF_PROP_STRING(PKEY_AudioEngine_DeviceFormat);
-		IF_PROP_STRING(PKEY_AudioEngine_OEMFormat);
-		IF_PROP_STRING(PKEY_AudioEndpoint_Association);
-		IF_PROP_STRING(PKEY_AudioEndpoint_ControlPanelPageProvider);
-		IF_PROP_STRING(PKEY_AudioEndpoint_Disable_SysFx);
-		IF_PROP_STRING(PKEY_AudioEndpoint_FormFactor);
-		IF_PROP_STRING(PKEY_AudioEndpoint_FullRangeSpeakers);
-		IF_PROP_STRING(PKEY_AudioEndpoint_GUID);
-		IF_PROP_STRING(PKEY_AudioEndpoint_Supports_EventDriven_Mode);
-
-#undef IF_PROP_STRING
-
-		TCHAR KeyString[PKEYSTR_MAX];
-		HRESULT HR = PSStringFromPropertyKey(Key, KeyString, ARRAYSIZE(KeyString));
-		if (SUCCEEDED(HR))
-		{
-			return FString(KeyString);
-		}
-#endif //!NO_LOGGING
-		return TEXT("Unknown");
-	}
-
+	}	
 }
 
-#endif //PLATFORM_WINDOWS
\ No newline at end of file
diff --git a/Engine/Source/Runtime/Windows/AudioMixerXAudio2/Private/ToStringHelpers.h b/Engine/Source/Runtime/Windows/AudioMixerXAudio2/Private/ToStringHelpers.h
--- a/Engine/Source/Runtime/Windows/AudioMixerXAudio2/Private/ToStringHelpers.h
+++ b/Engine/Source/Runtime/Windows/AudioMixerXAudio2/Private/ToStringHelpers.h
@@ -27,7 +27,9 @@ namespace Audio
 	FString ToFString(const PROPERTYKEY Key);
 #endif //PLATFORM_WINDOWS
 
-	const FString ToFString(const TArray<EAudioMixerChannel::Type>& InChannels);
+	FString ToFString(const TArray<EAudioMixerChannel::Type>& InChannels);
+	FString ToErrorFString(HRESULT InResult);
+
 	const TCHAR* ToString(EAudioDeviceRole InRole);
-	const TCHAR* ToString(EAudioDeviceState InState);	
+	const TCHAR* ToString(EAudioDeviceState InState);
 }
