Fix ToStringHelper Non-unity CIS issue.

diff --git a/Engine/Source/Runtime/Windows/AudioMixerXAudio2/Private/ToStringHelpers.h b/Engine/Source/Runtime/Windows/AudioMixerXAudio2/Private/ToStringHelpers.h
--- a/Engine/Source/Runtime/Windows/AudioMixerXAudio2/Private/ToStringHelpers.h
+++ b/Engine/Source/Runtime/Windows/AudioMixerXAudio2/Private/ToStringHelpers.h
@@ -4,20 +4,24 @@
 
 #include "AudioMixer.h"
 
-#if PLATFORM_WINDOWS
-
 #include "Windows/AllowWindowsPlatformTypes.h"
+
+#if PLATFORM_WINDOWS
 #include "Windows/MinWindows.h"
+#endif //PLATFORM_WINDOWS
 
 THIRD_PARTY_INCLUDES_START
+#include <winerror.h>
+
+#if PLATFORM_WINDOWS
 #include <mmdeviceapi.h>			// IMMNotificationClient
 #include <audiopolicy.h>			// IAudioSessionEvents
+#endif //PLATFORM_WINDOWS
+
 THIRD_PARTY_INCLUDES_END
 
 #include "Windows/HideWindowsPlatformTypes.h"
 
-#endif //PLATFORM_WINDOWS
-
 namespace Audio
 {
 #if PLATFORM_WINDOWS
