Making sure codecs match up on cook vs runtime

diff --git a/Engine/Source/Runtime/Windows/AudioMixerXAudio2/Private/AudioMixerPlatformXAudio2.cpp b/Engine/Source/Runtime/Windows/AudioMixerXAudio2/Private/AudioMixerPlatformXAudio2.cpp
--- a/Engine/Source/Runtime/Windows/AudioMixerXAudio2/Private/AudioMixerPlatformXAudio2.cpp
+++ b/Engine/Source/Runtime/Windows/AudioMixerXAudio2/Private/AudioMixerPlatformXAudio2.cpp
@@ -1858,14 +1858,14 @@ namespace Audio
 #if WITH_XMA2 && USE_XMA2_FOR_STREAMING
 			if (InSoundWave->NumChannels <= 2)
 			{
-				RuntimeFormat = Audio::NAME_XMA;
+				return Audio::NAME_XMA;
 			}
 #endif // WITH_XMA2 && USE_XMA2_FOR_STREAMING
 
 #if USE_VORBIS_FOR_STREAMING
-			RuntimeFormat = Audio::NAME_OGG;
+			return Audio::NAME_OGG;
 #else
-			RuntimeFormat = Audio::NAME_OPUS;
+			return Audio::NAME_OPUS;
 #endif // USE_VORBIS_FOR_STREAMING
 		}
 
@@ -1899,7 +1899,7 @@ namespace Audio
 #if WITH_BINK_AUDIO
 		else if (InRuntimeFormat == Audio::NAME_BINKA)
 		{
-			return new FBinkAudioInfo();
+			Decoder = new FBinkAudioInfo();
 		}
 #endif // #if WITH_BINK_AUDIO
 		else
