Fixing invalid call to INTC_D3D11_GetSupportedVersions where the pSupportedExtVersionsCount argument needs to contain the correct count. This is preventing Nanite from rendering on discrete Intel video cards.

diff --git a/Engine/Source/Runtime/Windows/D3D11RHI/Private/Windows/WindowsD3D11Device.cpp b/Engine/Source/Runtime/Windows/D3D11RHI/Private/Windows/WindowsD3D11Device.cpp
--- a/Engine/Source/Runtime/Windows/D3D11RHI/Private/Windows/WindowsD3D11Device.cpp
+++ b/Engine/Source/Runtime/Windows/D3D11RHI/Private/Windows/WindowsD3D11Device.cpp
@@ -1382,7 +1382,7 @@ void FD3D11DynamicRHI::StartIntelExtensions()
 	}
 
 	// Workaround for C6385, if we pass in SupportedExtensionsVersionCount again, the static analyzer thinks it may be different from the first call
-	uint32_t DummyCount = 0;
+	uint32_t DummyCount = SupportedExtensionsVersionCount;
 	if (SUCCEEDED(INTC_D3D11_GetSupportedVersions(Direct3DDevice, SupportedExtensionsVersions, &DummyCount)) && SupportedExtensionsVersions != nullptr)
 	{
 		check(SupportedExtensionsVersionCount == DummyCount);
