Copying //UE5/Dev-ParallelRendering to //UE5/Main

diff --git a/Engine/Source/Runtime/Windows/D3D11RHI/Private/D3D11Texture.cpp b/Engine/Source/Runtime/Windows/D3D11RHI/Private/D3D11Texture.cpp
--- a/Engine/Source/Runtime/Windows/D3D11RHI/Private/D3D11Texture.cpp
+++ b/Engine/Source/Runtime/Windows/D3D11RHI/Private/D3D11Texture.cpp
@@ -74,13 +74,31 @@ static TStatId GetD3D11StatEnum(uint32 BindFlags, bool bCubeMap, bool b3D)
 // Note: This function can be called from many different threads
 // @param TextureSize >0 to allocate, <0 to deallocate
 // @param b3D true:3D, false:2D or cube map
-void UpdateD3D11TextureStats(uint32 BindFlags, uint32 MiscFlags, int64 TextureSize, bool b3D)
+void UpdateD3D11TextureStats(FD3D11Texture& Texture, bool bAllocating)
 {
-	if(TextureSize == 0)
+	uint32 BindFlags, MiscFlags;
+	bool b3D = Texture.IsTexture3D();
+	if (b3D)
 	{
-		return;
+		D3D11_TEXTURE3D_DESC Desc;
+		Texture.GetD3D11Texture3D()->GetDesc(&Desc);
+		BindFlags = Desc.BindFlags;
+		MiscFlags = Desc.MiscFlags;
 	}
-	
+	else
+	{
+		D3D11_TEXTURE2D_DESC Desc;
+		Texture.GetD3D11Texture2D()->GetDesc(&Desc);
+		BindFlags = Desc.BindFlags;
+		MiscFlags = Desc.MiscFlags;
+	}
+
+	int64 TextureSize = Texture.GetMemorySize();
+	if (!bAllocating)
+	{
+		TextureSize = -TextureSize;
+	}
+
 	int64 AlignedSize = (TextureSize > 0) ? Align(TextureSize, 1024) / 1024 : -(Align(-TextureSize, 1024) / 1024);
 	if(ShouldCountAsTextureMemory(BindFlags))
 	{
@@ -95,7 +113,7 @@ void UpdateD3D11TextureStats(uint32 BindFlags, uint32 MiscFlags, int64 TextureSi
 
 	INC_MEMORY_STAT_BY_FName(GetD3D11StatEnum(BindFlags, bCubeMap, b3D).GetName(), TextureSize);
 
-	if(TextureSize > 0)
+	if (bAllocating)
 	{
 		INC_DWORD_STAT(STAT_D3D11TexturesAllocated);
 	}
@@ -103,171 +121,31 @@ void UpdateD3D11TextureStats(uint32 BindFlags, uint32 MiscFlags, int64 TextureSi
 	{
 		INC_DWORD_STAT(STAT_D3D11TexturesReleased);
 	}
-}
-
-template<typename BaseResourceType>
-void D3D11TextureAllocated( TD3D11Texture2D<BaseResourceType>& Texture )
-{
-	ID3D11Texture2D* D3D11Texture2D = Texture.GetResource();
-
-	if(D3D11Texture2D)
-	{
-		if (EnumHasAllFlags(Texture.Flags, TexCreate_Virtual) )
-		{
-			Texture.SetMemorySize(0);
-		}
-		else
-		{
-			D3D11_TEXTURE2D_DESC Desc;
-
-			D3D11Texture2D->GetDesc( &Desc );
-			check(Texture.IsCubemap() == ((Desc.MiscFlags & D3D11_RESOURCE_MISC_TEXTURECUBE) != 0));
-
-			int64 TextureSize = CalcTextureSize( Desc.Width, Desc.Height, Texture.GetFormat(), Desc.MipLevels ) * Desc.ArraySize;
-
-			Texture.SetMemorySize( TextureSize );
-			UpdateD3D11TextureStats(Desc.BindFlags, Desc.MiscFlags, TextureSize, false);
 
 #if PLATFORM_WINDOWS
-			// On Windows there is no way to hook into the low level d3d allocations and frees.
-			// This means that we must manually add the tracking here.
-			LLM(FLowLevelMemTracker::Get().OnLowLevelAlloc(ELLMTracker::Platform, Texture.GetResource(), Texture.GetMemorySize(), ELLMTag::GraphicsPlatform));
-			LLM(FLowLevelMemTracker::Get().OnLowLevelAlloc(ELLMTracker::Default, Texture.GetResource(), Texture.GetMemorySize(), ELLMTag::Textures));
-#endif
-		}
-	}
-}
-
-template<typename BaseResourceType>
-void D3D11TextureDeleted( TD3D11Texture2D<BaseResourceType>& Texture )
-{
-	ID3D11Texture2D* D3D11Texture2D = Texture.GetResource();
-
-	if(D3D11Texture2D)
+	// On Windows there is no way to hook into the low level d3d allocations and frees.
+	// This means that we must manually add the tracking here.
+	if (bAllocating)
 	{
-		D3D11_TEXTURE2D_DESC Desc;
-
-		D3D11Texture2D->GetDesc( &Desc );
-		check(Texture.IsCubemap() == ((Desc.MiscFlags & D3D11_RESOURCE_MISC_TEXTURECUBE) != 0));
-
-		// When using virtual textures use the current memory size, which is the number of physical pages allocated, not virtual
-		int64 TextureSize = 0;
-		if (EnumHasAllFlags(Texture.GetFlags(), TexCreate_Virtual))
-		{
-			TextureSize = Texture.GetMemorySize();
-		}
-		else
-		{
-			TextureSize = CalcTextureSize( Desc.Width, Desc.Height, Texture.GetFormat(), Desc.MipLevels ) * Desc.ArraySize;
-		}
-
-		UpdateD3D11TextureStats(Desc.BindFlags, Desc.MiscFlags, -TextureSize, false);
-
-#if PLATFORM_WINDOWS
-		// On Windows there is no way to hook into the low level d3d allocations and frees.
-		// This means that we must manually add the tracking here.
-		LLM(FLowLevelMemTracker::Get().OnLowLevelFree(ELLMTracker::Platform, Texture.GetResource()));
-		LLM(FLowLevelMemTracker::Get().OnLowLevelFree(ELLMTracker::Default, Texture.GetResource()));
-#endif
+		LLM(FLowLevelMemTracker::Get().OnLowLevelAlloc(ELLMTracker::Platform, Texture.GetResource(), TextureSize, ELLMTag::GraphicsPlatform));
+		LLM(FLowLevelMemTracker::Get().OnLowLevelAlloc(ELLMTracker::Default , Texture.GetResource(), TextureSize, ELLMTag::Textures));
 	}
-}
-
-void D3D11TextureAllocated2D( FD3D11Texture2D& Texture )
-{
-	D3D11TextureAllocated(Texture);
-}
-
-void D3D11TextureAllocated( FD3D11Texture3D& Texture )
-{
-	ID3D11Texture3D* D3D11Texture3D = Texture.GetResource();
-
-	if(D3D11Texture3D)
-	{
-		D3D11_TEXTURE3D_DESC Desc;
-
-		D3D11Texture3D->GetDesc( &Desc );
-
-		int64 TextureSize = CalcTextureSize3D( Desc.Width, Desc.Height, Desc.Depth, Texture.GetFormat(), Desc.MipLevels );
-
-		Texture.SetMemorySize( TextureSize );
-
-		UpdateD3D11TextureStats(Desc.BindFlags, Desc.MiscFlags, TextureSize, true);
-
-#if PLATFORM_WINDOWS
-		// On Windows there is no way to hook into the low level d3d allocations and frees.
-		// This means that we must manually add the tracking here.
-		LLM(FLowLevelMemTracker::Get().OnLowLevelAlloc(ELLMTracker::Platform, Texture.GetResource(), Texture.GetMemorySize(), ELLMTag::GraphicsPlatform));
-		LLM(FLowLevelMemTracker::Get().OnLowLevelAlloc(ELLMTracker::Default, Texture.GetResource(), Texture.GetMemorySize(), ELLMTag::Textures));
-#endif
-	}
-}
-
-void D3D11TextureDeleted( FD3D11Texture3D& Texture )
-{
-	ID3D11Texture3D* D3D11Texture3D = Texture.GetResource();
-
-	if(D3D11Texture3D)
+	else
 	{
-		D3D11_TEXTURE3D_DESC Desc;
-
-		D3D11Texture3D->GetDesc( &Desc );
-
-		int64 TextureSize = CalcTextureSize3D( Desc.Width, Desc.Height, Desc.Depth, Texture.GetFormat(), Desc.MipLevels );
-
-		UpdateD3D11TextureStats(Desc.BindFlags, Desc.MiscFlags, -TextureSize, true);
-
-#if PLATFORM_WINDOWS
-		// On Windows there is no way to hook into the low level d3d allocations and frees.
-		// This means that we must manually add the tracking here.
 		LLM(FLowLevelMemTracker::Get().OnLowLevelFree(ELLMTracker::Platform, Texture.GetResource()));
-		LLM(FLowLevelMemTracker::Get().OnLowLevelFree(ELLMTracker::Default, Texture.GetResource()));
-#endif
-	}
-}
-
-template<typename BaseResourceType>
-TD3D11Texture2D<BaseResourceType>::~TD3D11Texture2D()
-{
-	D3D11TextureDeleted(*this);
-	if (bPooled)
-	{
-		ReturnPooledTexture2D(this->GetNumMips(), this->GetFormat(), this->GetResource());
+		LLM(FLowLevelMemTracker::Get().OnLowLevelFree(ELLMTracker::Default , Texture.GetResource()));
 	}
-
-#if PLATFORM_SUPPORTS_VIRTUAL_TEXTURES
-	D3DRHI->DestroyVirtualTexture(GetFlags(), GetRawTextureMemory());
 #endif
 }
 
-template TD3D11Texture2D<FD3D11BaseTexture2D>::~TD3D11Texture2D();
-
-FD3D11Texture3D::~FD3D11Texture3D()
+FDynamicRHI::FRHICalcTextureSizeResult FD3D11DynamicRHI::RHICalcTexturePlatformSize(FRHITextureDesc const& Desc, uint32 FirstMipIndex)
 {
-	D3D11TextureDeleted( *this );
-}
-
-uint64 FD3D11DynamicRHI::RHICalcTexture2DPlatformSize(uint32 SizeX, uint32 SizeY, uint8 Format, uint32 NumMips, uint32 NumSamples, ETextureCreateFlags Flags, const FRHIResourceCreateInfo& CreateInfo, uint32& OutAlign)
-{
-	OutAlign = 0;
-	return CalcTextureSize(SizeX, SizeY, (EPixelFormat)Format, NumMips);
-}
-
-uint64 FD3D11DynamicRHI::RHICalcTexture2DArrayPlatformSize(uint32 SizeX, uint32 SizeY, uint32 ArraySize, uint8 Format, uint32 NumMips, uint32 NumSamples, ETextureCreateFlags Flags, const FRHIResourceCreateInfo& CreateInfo, uint32& OutAlign)
-{
-	OutAlign = 0;
-	return CalcTextureSize(SizeX, SizeY, (EPixelFormat)Format, NumMips) * ArraySize;
-}
-
-uint64 FD3D11DynamicRHI::RHICalcTexture3DPlatformSize(uint32 SizeX, uint32 SizeY, uint32 SizeZ, uint8 Format, uint32 NumMips, ETextureCreateFlags Flags, const FRHIResourceCreateInfo& CreateInfo, uint32& OutAlign)
-{
-	OutAlign = 0;
-	return CalcTextureSize3D(SizeX, SizeY, SizeZ, (EPixelFormat)Format, NumMips);
-}
-
-uint64 FD3D11DynamicRHI::RHICalcTextureCubePlatformSize(uint32 Size, uint8 Format, uint32 NumMips, ETextureCreateFlags Flags, const FRHIResourceCreateInfo& CreateInfo, uint32& OutAlign)
-{
-	OutAlign = 0;
-	return CalcTextureSize(Size, Size, (EPixelFormat)Format, NumMips) * 6;
+	// D3D11 does not provide a way to compute the actual driver/GPU specific in-memory size of a texture.
+	// Fallback to the estimate based on the texture's dimensions / format etc.
+	FDynamicRHI::FRHICalcTextureSizeResult Result;
+	Result.Size = Desc.CalcMemorySizeEstimate(FirstMipIndex);
+	Result.Align = 1;
+	return Result;
 }
 
 /**
@@ -303,145 +181,6 @@ bool FD3D11DynamicRHI::RHIGetTextureMemoryVisualizeData( FColor* /*TextureData*/
 	return false;
 }
 
-/*------------------------------------------------------------------------------
-	Texture pooling.
-------------------------------------------------------------------------------*/
-
-/** Define to 1 to enable the pooling of 2D texture resources. */
-#define USE_TEXTURE_POOLING 0
-
-/** A texture resource stored in the pool. */
-struct FPooledTexture2D
-{
-	/** The texture resource. */
-	TRefCountPtr<ID3D11Texture2D> Resource;
-};
-
-/** A pool of D3D texture resources. */
-struct FTexturePool
-{
-	TArray<FPooledTexture2D> Textures;
-};
-
-/** The global texture pool. */
-struct FGlobalTexturePool
-{
-	/** Formats stored in the pool. */
-	enum EInternalFormat
-	{
-		IF_DXT1,
-		IF_DXT5,
-		IF_BC5,
-		IF_Max
-	};
-
-	enum
-	{
-		/** Minimum mip count for which to pool textures. */
-		MinMipCount = 7,
-		/** Maximum mip count for which to pool textures. */
-		MaxMipCount = 13,
-		/** The number of pools based on mip levels. */
-		MipPoolCount = MaxMipCount - MinMipCount,
-	};
-
-	/** The individual texture pools. */
-	FTexturePool Pools[MipPoolCount][IF_Max];
-};
-FGlobalTexturePool GTexturePool;
-
-/**
- * Releases all pooled textures.
- */
-void ReleasePooledTextures()
-{
-	for (int32 MipPoolIndex = 0; MipPoolIndex < FGlobalTexturePool::MipPoolCount; ++MipPoolIndex)
-	{
-		for (int32 FormatPoolIndex = 0; FormatPoolIndex < FGlobalTexturePool::IF_Max; ++FormatPoolIndex)
-		{
-			GTexturePool.Pools[MipPoolIndex][FormatPoolIndex].Textures.Empty();
-		}
-	}
-}
-
-/**
- * Retrieves the texture pool for the specified mip count and format.
- */
-FTexturePool* GetTexturePool(int32 MipCount, EPixelFormat PixelFormat)
-{
-	FTexturePool* Pool = NULL;
-	int32 MipPool = MipCount - FGlobalTexturePool::MinMipCount;
-	if (MipPool >= 0 && MipPool < FGlobalTexturePool::MipPoolCount)
-	{
-		int32 FormatPool = -1;
-		switch (PixelFormat)
-		{
-			case PF_DXT1: FormatPool = FGlobalTexturePool::IF_DXT1; break;
-			case PF_DXT5: FormatPool = FGlobalTexturePool::IF_DXT5; break;
-			case PF_BC5: FormatPool = FGlobalTexturePool::IF_BC5; break;
-		}
-		if (FormatPool >= 0 && FormatPool < FGlobalTexturePool::IF_Max)
-		{
-			Pool = &GTexturePool.Pools[MipPool][FormatPool];
-		}
-	}
-	return Pool;
-}
-
-/**
- * Retrieves a texture from the pool if one exists.
- */
-bool GetPooledTexture2D(int32 MipCount, EPixelFormat PixelFormat, FPooledTexture2D* OutTexture)
-{
-#if USE_TEXTURE_POOLING
-	FTexturePool* Pool = GetTexturePool(MipCount,PixelFormat);
-	if (Pool && Pool->Textures.Num() > 0)
-	{
-		*OutTexture = Pool->Textures.Last();
-
-		{
-			D3D11_TEXTURE2D_DESC Desc;
-			OutTexture->Resource->GetDesc(&Desc);
-			check(Desc.Format == GPixelFormats[PixelFormat].PlatformFormat);
-			check(MipCount == Desc.MipLevels);
-			check(Desc.Width == Desc.Height);
-			check(Desc.Width == (1 << (MipCount-1)));
-			int32 TextureSize = CalcTextureSize(Desc.Width, Desc.Height, PixelFormat, Desc.MipLevels);
-			DEC_MEMORY_STAT_BY(STAT_D3D11TexturePoolMemory,TextureSize);
-		}
-
-		Pool->Textures.RemoveAt(Pool->Textures.Num() - 1);
-		return true;
-	}
-#endif // #if USE_TEXTURE_POOLING
-	return false;
-}
-
-/**
- * Returns a texture to its pool.
- */
-void ReturnPooledTexture2D(int32 MipCount, EPixelFormat PixelFormat, ID3D11Texture2D* InResource)
-{
-#if USE_TEXTURE_POOLING
-	FTexturePool* Pool = GetTexturePool(MipCount,PixelFormat);
-	if (Pool)
-	{
-		FPooledTexture2D* PooledTexture = new(Pool->Textures) FPooledTexture2D;
-		PooledTexture->Resource = InResource;
-		{
-			D3D11_TEXTURE2D_DESC Desc;
-			PooledTexture->Resource->GetDesc(&Desc);
-			check(Desc.Format == GPixelFormats[PixelFormat].PlatformFormat);
-			check(MipCount == Desc.MipLevels);
-			check(Desc.Width == Desc.Height);
-			check(Desc.Width == (1 << (MipCount-1)));
-			int32 TextureSize = CalcTextureSize(Desc.Width, Desc.Height, PixelFormat, Desc.MipLevels);
-			INC_MEMORY_STAT_BY(STAT_D3D11TexturePoolMemory,TextureSize);
-		}
-	}
-#endif // #if USE_TEXTURE_POOLING
-}
-
 /** Find an appropriate DXGI format for the input format and SRGB setting. */
 inline DXGI_FORMAT FindSharedResourceDXGIFormat(DXGI_FORMAT InFormat, bool bSRGB)
 {
@@ -596,10 +335,20 @@ void SafeCreateTexture2D(ID3D11Device* Direct3DDevice, int32 UEFormat, const D3D
 #endif // #if GUARDED_TEXTURE_CREATES
 }
 
-template<typename BaseResourceType>
-TD3D11Texture2D<BaseResourceType>* FD3D11DynamicRHI::CreateD3D11Texture2D(uint32 SizeX,uint32 SizeY,uint32 SizeZ,bool bTextureArray,bool bCubeTexture,uint8 Format,
-	uint32 NumMips,uint32 NumSamples,ETextureCreateFlags Flags, FRHIResourceCreateInfo& CreateInfo)
+FD3D11Texture* FD3D11DynamicRHI::CreateD3D11Texture2D(FRHITextureCreateDesc const& CreateDesc, TConstArrayView<D3D11_SUBRESOURCE_DATA> InitialData)
 {
+	check(!CreateDesc.IsTexture3D());
+
+	const bool                bTextureArray = CreateDesc.IsTextureArray();
+	const bool                bCubeTexture  = CreateDesc.IsTextureCube();
+	const uint32              SizeX         = CreateDesc.Extent.X;
+	const uint32              SizeY         = CreateDesc.Extent.Y;
+	const uint32              SizeZ         = bCubeTexture ? CreateDesc.ArraySize * 6 : CreateDesc.ArraySize;
+	const EPixelFormat        Format        = CreateDesc.Format;
+	const uint32              NumMips       = CreateDesc.NumMips;
+	const uint32              NumSamples    = CreateDesc.NumSamples;
+	const ETextureCreateFlags Flags         = CreateDesc.Flags;
+
 	check(SizeX > 0 && SizeY > 0 && NumMips > 0);
 
 	if (bCubeTexture)
@@ -620,8 +369,6 @@ TD3D11Texture2D<BaseResourceType>* FD3D11DynamicRHI::CreateD3D11Texture2D(uint32
 
 	SCOPE_CYCLE_COUNTER(STAT_D3D11CreateTextureTime);
 
-	bool bPooledTexture = true;
-
 	const bool bSRGB = EnumHasAnyFlags(Flags, TexCreate_SRGB);
 
 	const DXGI_FORMAT PlatformResourceFormat = FD3D11DynamicRHI::GetPlatformTextureResourceFormat((DXGI_FORMAT)GPixelFormats[Format].PlatformFormat, Flags);
@@ -632,29 +379,11 @@ TD3D11Texture2D<BaseResourceType>* FD3D11DynamicRHI::CreateD3D11Texture2D(uint32
 	D3D11_USAGE TextureUsage = D3D11_USAGE_DEFAULT;
 	bool bCreateShaderResource = true;
 
-	uint32 ActualMSAACount = NumSamples;
+	uint32 ActualMSAAQuality = GetMaxMSAAQuality(NumSamples);
+	check(ActualMSAAQuality != 0xffffffff);
+	check(NumSamples == 1 || !EnumHasAnyFlags(Flags, TexCreate_Shared));
 
-	uint32 ActualMSAAQuality = GetMaxMSAAQuality(ActualMSAACount);
-
-	// 0xffffffff means not supported
-	if (ActualMSAAQuality == 0xffffffff || EnumHasAnyFlags(Flags, TexCreate_Shared))
-	{
-		// no MSAA
-		ActualMSAACount = 1;
-		ActualMSAAQuality = 0;
-	}
-
-	const bool bIsMultisampled = ActualMSAACount > 1;
-
-	if (bIsMultisampled)
-	{
-		bPooledTexture = false;
-	}
-
-	if (NumMips < 1 || SizeX != SizeY || (1 << (NumMips - 1)) != SizeX || EnumHasAnyFlags(Flags, TexCreate_Shared))
-	{
-		bPooledTexture = false;
-	}
+	const bool bIsMultisampled = NumSamples > 1;
 
 	if (EnumHasAnyFlags(Flags, TexCreate_CPUReadback))
 	{
@@ -680,7 +409,7 @@ TD3D11Texture2D<BaseResourceType>* FD3D11DynamicRHI::CreateD3D11Texture2D(uint32
 	TextureDesc.MipLevels = NumMips;
 	TextureDesc.ArraySize = SizeZ;
 	TextureDesc.Format = PlatformResourceFormat;
-	TextureDesc.SampleDesc.Count = ActualMSAACount;
+	TextureDesc.SampleDesc.Count = NumSamples;
 	TextureDesc.SampleDesc.Quality = ActualMSAAQuality;
 	TextureDesc.Usage = TextureUsage;
 	TextureDesc.BindFlags = bCreateShaderResource? D3D11_BIND_SHADER_RESOURCE : 0;
@@ -719,7 +448,6 @@ TD3D11Texture2D<BaseResourceType>* FD3D11DynamicRHI::CreateD3D11Texture2D(uint32
 	{
 		// Set the flag that allows us to call GenerateMips on this texture later
 		TextureDesc.MiscFlags |= D3D11_RESOURCE_MISC_GENERATE_MIPS;
-		bPooledTexture = false;
 	}
 
 	// Set up the texture bind flags.
@@ -763,7 +491,6 @@ TD3D11Texture2D<BaseResourceType>* FD3D11DynamicRHI::CreateD3D11Texture2D(uint32
 	if (EnumHasAnyFlags(Flags, TexCreate_UAV))
 	{
 		TextureDesc.BindFlags |= D3D11_BIND_UNORDERED_ACCESS;
-		bPooledTexture = false;
 	}
 
 	if (bCreateDSV && !EnumHasAnyFlags(Flags, TexCreate_ShaderResource))
@@ -772,161 +499,94 @@ TD3D11Texture2D<BaseResourceType>* FD3D11DynamicRHI::CreateD3D11Texture2D(uint32
 		bCreateShaderResource = false;
 	}
 
-	if (bCreateDSV || bCreateRTV || bCubeTexture || bTextureArray)
-	{
-		bPooledTexture = false;
-	}
-
 	TRefCountPtr<ID3D11Texture2D> TextureResource;
 	TRefCountPtr<ID3D11ShaderResourceView> ShaderResourceView;
 	TArray<TRefCountPtr<ID3D11RenderTargetView> > RenderTargetViews;
 	TRefCountPtr<ID3D11DepthStencilView> DepthStencilViews[FExclusiveDepthStencil::MaxIndex];
-	
-#if PLATFORM_SUPPORTS_VIRTUAL_TEXTURES
-	// Turn off pooling when we are using virtual textures or the texture is offline processed as we control when the memory is released
-	if (EnumHasAnyFlags(Flags, TexCreate_Virtual | TexCreate_OfflineProcessed))
-	{
-		bPooledTexture = false;
-	}
-	void* RawTextureMemory = nullptr;
-#else
-	Flags &= ~TexCreate_Virtual;
-#endif
 
-	if (bPooledTexture)
+	TArray<D3D11_SUBRESOURCE_DATA> SubResourceData;
+	D3D11_SUBRESOURCE_DATA const* pSubresourceData = nullptr;
+
+	if (InitialData.Num())
 	{
-		FPooledTexture2D PooledTexture;
-		if (GetPooledTexture2D(NumMips, (EPixelFormat)Format, &PooledTexture))
-		{
-			TextureResource = PooledTexture.Resource;
-		}
+		// Caller provided initla data structs.
+		check(InitialData.Num() == NumMips * SizeZ);
+		pSubresourceData = InitialData.GetData();
 	}
-
-	if (!IsValidRef(TextureResource))
+	else if (CreateDesc.BulkData)
 	{
-		TArray<D3D11_SUBRESOURCE_DATA> SubResourceData;
-
-		if (CreateInfo.BulkData)
-		{
-			uint8* Data = (uint8*)CreateInfo.BulkData->GetResourceBulkData();
+		uint8* Data = (uint8*)CreateDesc.BulkData->GetResourceBulkData();
 
-			// each mip of each array slice counts as a subresource
-			SubResourceData.AddZeroed(NumMips * SizeZ);
+		// each mip of each array slice counts as a subresource
+		SubResourceData.AddZeroed(NumMips * SizeZ);
 
-			uint32 SliceOffset = 0;
-			for (uint32 ArraySliceIndex = 0; ArraySliceIndex < SizeZ; ++ArraySliceIndex)
-			{			
-				uint32 MipOffset = 0;
-				for(uint32 MipIndex = 0;MipIndex < NumMips;++MipIndex)
-				{
-					uint32 DataOffset = SliceOffset + MipOffset;
-					uint32 SubResourceIndex = ArraySliceIndex * NumMips + MipIndex;
+		uint32 SliceOffset = 0;
+		for (uint32 ArraySliceIndex = 0; ArraySliceIndex < SizeZ; ++ArraySliceIndex)
+		{			
+			uint32 MipOffset = 0;
+			for(uint32 MipIndex = 0;MipIndex < NumMips;++MipIndex)
+			{
+				uint32 DataOffset = SliceOffset + MipOffset;
+				uint32 SubResourceIndex = ArraySliceIndex * NumMips + MipIndex;
 
-					uint32 NumBlocksX = FMath::Max<uint32>(1,(SizeX >> MipIndex) / GPixelFormats[Format].BlockSizeX);
-					uint32 NumBlocksY = FMath::Max<uint32>(1,(SizeY >> MipIndex) / GPixelFormats[Format].BlockSizeY);
+				uint32 NumBlocksX = FMath::Max<uint32>(1,(SizeX >> MipIndex) / GPixelFormats[Format].BlockSizeX);
+				uint32 NumBlocksY = FMath::Max<uint32>(1,(SizeY >> MipIndex) / GPixelFormats[Format].BlockSizeY);
 
-					SubResourceData[SubResourceIndex].pSysMem = &Data[DataOffset];
-					SubResourceData[SubResourceIndex].SysMemPitch      =  NumBlocksX * GPixelFormats[Format].BlockBytes;
-					SubResourceData[SubResourceIndex].SysMemSlicePitch =  NumBlocksX * NumBlocksY * SubResourceData[MipIndex].SysMemPitch;
+				SubResourceData[SubResourceIndex].pSysMem = &Data[DataOffset];
+				SubResourceData[SubResourceIndex].SysMemPitch      =  NumBlocksX * GPixelFormats[Format].BlockBytes;
+				SubResourceData[SubResourceIndex].SysMemSlicePitch =  NumBlocksX * NumBlocksY * SubResourceData[MipIndex].SysMemPitch;
 
-					MipOffset                                  += NumBlocksY * SubResourceData[MipIndex].SysMemPitch;
-				}
-				SliceOffset += MipOffset;
+				MipOffset                                  += NumBlocksY * SubResourceData[MipIndex].SysMemPitch;
 			}
+			SliceOffset += MipOffset;
 		}
 
-#if PLATFORM_SUPPORTS_VIRTUAL_TEXTURES
-		if (EnumHasAnyFlags(Flags, TexCreate_Virtual | TexCreate_OfflineProcessed))
-		{
-			RawTextureMemory = CreateVirtualTexture(SizeX, SizeY, SizeZ, NumMips, bCubeTexture, Flags, &TextureDesc, &TextureResource);
-		}
-		else
-#endif
+		pSubresourceData = SubResourceData.GetData();
+	}
+
 #if INTEL_EXTENSIONS
-		if (EnumHasAnyFlags(Flags, ETextureCreateFlags::Atomic64Compatible) && IsRHIDeviceIntel() && GRHISupportsAtomicUInt64)
-		{
-			INTC_D3D11_TEXTURE2D_DESC IntelDesc{};
-			IntelDesc.EmulatedTyped64bitAtomics = true;
-			IntelDesc.pD3D11Desc = &TextureDesc;
+	if (EnumHasAnyFlags(Flags, ETextureCreateFlags::Atomic64Compatible) && IsRHIDeviceIntel() && GRHISupportsAtomicUInt64)
+	{
+		INTC_D3D11_TEXTURE2D_DESC IntelDesc{};
+		IntelDesc.EmulatedTyped64bitAtomics = true;
+		IntelDesc.pD3D11Desc = &TextureDesc;
 
-			VERIFYD3D11RESULT(INTC_D3D11_CreateTexture2D(IntelExtensionContext, &IntelDesc, CreateInfo.BulkData != nullptr ? (const D3D11_SUBRESOURCE_DATA*)SubResourceData.GetData() : nullptr, TextureResource.GetInitReference()));
-		}
-		else
+		VERIFYD3D11RESULT(INTC_D3D11_CreateTexture2D(IntelExtensionContext, &IntelDesc, pSubresourceData, TextureResource.GetInitReference()));
+	}
+	else
 #endif
-		{
-			SafeCreateTexture2D(Direct3DDevice, Format, &TextureDesc, CreateInfo.BulkData != NULL ? (const D3D11_SUBRESOURCE_DATA*)SubResourceData.GetData() : NULL, TextureResource.GetInitReference());
-		}
+	{
+		SafeCreateTexture2D(Direct3DDevice, Format, &TextureDesc, pSubresourceData, TextureResource.GetInitReference());
+	}
 
-		if(bCreateRTV)
+	if (bCreateRTV)
+	{
+		// Create a render target view for each mip
+		for (uint32 MipIndex = 0; MipIndex < NumMips; MipIndex++)
 		{
-			// Create a render target view for each mip
-			for (uint32 MipIndex = 0; MipIndex < NumMips; MipIndex++)
+			if (EnumHasAnyFlags(Flags, TexCreate_TargetArraySlicesIndependently) && (bTextureArray || bCubeTexture))
 			{
-				if (EnumHasAnyFlags(Flags, TexCreate_TargetArraySlicesIndependently) && (bTextureArray || bCubeTexture))
-				{
-					bCreatedRTVPerSlice = true;
+				bCreatedRTVPerSlice = true;
 
-					for (uint32 SliceIndex = 0; SliceIndex < TextureDesc.ArraySize; SliceIndex++)
-					{
-						D3D11_RENDER_TARGET_VIEW_DESC RTVDesc;
-						FMemory::Memzero(RTVDesc);
-
-						RTVDesc.Format = PlatformRenderTargetFormat;
-
-						if (bIsMultisampled)
-						{
-							RTVDesc.ViewDimension = D3D11_RTV_DIMENSION_TEXTURE2DMSARRAY;
-							RTVDesc.Texture2DMSArray.FirstArraySlice = SliceIndex;
-							RTVDesc.Texture2DMSArray.ArraySize = 1;
-						}
-						else
-						{
-							RTVDesc.ViewDimension = D3D11_RTV_DIMENSION_TEXTURE2DARRAY;
-							RTVDesc.Texture2DArray.FirstArraySlice = SliceIndex;
-							RTVDesc.Texture2DArray.ArraySize = 1;
-							RTVDesc.Texture2DArray.MipSlice = MipIndex;
-						}
-
-						TRefCountPtr<ID3D11RenderTargetView> RenderTargetView;
-						VERIFYD3D11RESULT_EX(Direct3DDevice->CreateRenderTargetView(TextureResource,&RTVDesc,RenderTargetView.GetInitReference()), Direct3DDevice);
-						RenderTargetViews.Add(RenderTargetView);
-					}
-				}
-				else
+				for (uint32 SliceIndex = 0; SliceIndex < TextureDesc.ArraySize; SliceIndex++)
 				{
 					D3D11_RENDER_TARGET_VIEW_DESC RTVDesc;
 					FMemory::Memzero(RTVDesc);
 
 					RTVDesc.Format = PlatformRenderTargetFormat;
 
-					if (bTextureArray || bCubeTexture)
+					if (bIsMultisampled)
 					{
-						if (bIsMultisampled)
-						{
-							RTVDesc.ViewDimension = D3D11_RTV_DIMENSION_TEXTURE2DMSARRAY;
-							RTVDesc.Texture2DMSArray.FirstArraySlice = 0;
-							RTVDesc.Texture2DMSArray.ArraySize = TextureDesc.ArraySize;
-						}
-						else
-						{
-							RTVDesc.ViewDimension = D3D11_RTV_DIMENSION_TEXTURE2DARRAY;
-							RTVDesc.Texture2DArray.FirstArraySlice = 0;
-							RTVDesc.Texture2DArray.ArraySize = TextureDesc.ArraySize;
-							RTVDesc.Texture2DArray.MipSlice = MipIndex;
-						}
+						RTVDesc.ViewDimension = D3D11_RTV_DIMENSION_TEXTURE2DMSARRAY;
+						RTVDesc.Texture2DMSArray.FirstArraySlice = SliceIndex;
+						RTVDesc.Texture2DMSArray.ArraySize = 1;
 					}
 					else
 					{
-						if (bIsMultisampled)
-						{
-							RTVDesc.ViewDimension = D3D11_RTV_DIMENSION_TEXTURE2DMS;
-							// Nothing to set
-						}
-						else
-						{
-							RTVDesc.ViewDimension = D3D11_RTV_DIMENSION_TEXTURE2D;
-							RTVDesc.Texture2D.MipSlice = MipIndex;
-						}
+						RTVDesc.ViewDimension = D3D11_RTV_DIMENSION_TEXTURE2DARRAY;
+						RTVDesc.Texture2DArray.FirstArraySlice = SliceIndex;
+						RTVDesc.Texture2DArray.ArraySize = 1;
+						RTVDesc.Texture2DArray.MipSlice = MipIndex;
 					}
 
 					TRefCountPtr<ID3D11RenderTargetView> RenderTargetView;
@@ -934,58 +594,99 @@ TD3D11Texture2D<BaseResourceType>* FD3D11DynamicRHI::CreateD3D11Texture2D(uint32
 					RenderTargetViews.Add(RenderTargetView);
 				}
 			}
-		}
-	
-		if (bCreateDSV)
-		{
-			// Create a depth-stencil-view for the texture.
-			D3D11_DEPTH_STENCIL_VIEW_DESC DSVDesc;
-			FMemory::Memzero(DSVDesc);
+			else
+			{
+				D3D11_RENDER_TARGET_VIEW_DESC RTVDesc;
+				FMemory::Memzero(RTVDesc);
 
-			DSVDesc.Format = FindDepthStencilDXGIFormat(PlatformResourceFormat);
+				RTVDesc.Format = PlatformRenderTargetFormat;
 
-			if (bTextureArray || bCubeTexture)
-			{
-				if (bIsMultisampled)
+				if (bTextureArray || bCubeTexture)
 				{
-					DSVDesc.ViewDimension = D3D11_DSV_DIMENSION_TEXTURE2DMSARRAY;
-					DSVDesc.Texture2DMSArray.FirstArraySlice = 0;
-					DSVDesc.Texture2DMSArray.ArraySize = TextureDesc.ArraySize;
+					if (bIsMultisampled)
+					{
+						RTVDesc.ViewDimension = D3D11_RTV_DIMENSION_TEXTURE2DMSARRAY;
+						RTVDesc.Texture2DMSArray.FirstArraySlice = 0;
+						RTVDesc.Texture2DMSArray.ArraySize = TextureDesc.ArraySize;
+					}
+					else
+					{
+						RTVDesc.ViewDimension = D3D11_RTV_DIMENSION_TEXTURE2DARRAY;
+						RTVDesc.Texture2DArray.FirstArraySlice = 0;
+						RTVDesc.Texture2DArray.ArraySize = TextureDesc.ArraySize;
+						RTVDesc.Texture2DArray.MipSlice = MipIndex;
+					}
 				}
 				else
 				{
-					DSVDesc.ViewDimension = D3D11_DSV_DIMENSION_TEXTURE2DARRAY;
-					DSVDesc.Texture2DArray.FirstArraySlice = 0;
-					DSVDesc.Texture2DArray.ArraySize = TextureDesc.ArraySize;
-					DSVDesc.Texture2DArray.MipSlice = 0;
+					if (bIsMultisampled)
+					{
+						RTVDesc.ViewDimension = D3D11_RTV_DIMENSION_TEXTURE2DMS;
+						// Nothing to set
+					}
+					else
+					{
+						RTVDesc.ViewDimension = D3D11_RTV_DIMENSION_TEXTURE2D;
+						RTVDesc.Texture2D.MipSlice = MipIndex;
+					}
 				}
+
+				TRefCountPtr<ID3D11RenderTargetView> RenderTargetView;
+				VERIFYD3D11RESULT_EX(Direct3DDevice->CreateRenderTargetView(TextureResource,&RTVDesc,RenderTargetView.GetInitReference()), Direct3DDevice);
+				RenderTargetViews.Add(RenderTargetView);
+			}
+		}
+	}
+	
+	if (bCreateDSV)
+	{
+		// Create a depth-stencil-view for the texture.
+		D3D11_DEPTH_STENCIL_VIEW_DESC DSVDesc;
+		FMemory::Memzero(DSVDesc);
+
+		DSVDesc.Format = FindDepthStencilDXGIFormat(PlatformResourceFormat);
+
+		if (bTextureArray || bCubeTexture)
+		{
+			if (bIsMultisampled)
+			{
+				DSVDesc.ViewDimension = D3D11_DSV_DIMENSION_TEXTURE2DMSARRAY;
+				DSVDesc.Texture2DMSArray.FirstArraySlice = 0;
+				DSVDesc.Texture2DMSArray.ArraySize = TextureDesc.ArraySize;
 			}
 			else
 			{
-				if (bIsMultisampled)
-				{
-					DSVDesc.ViewDimension = D3D11_DSV_DIMENSION_TEXTURE2DMS;
-					// Nothing to set
-				}
-				else
-				{
-					DSVDesc.ViewDimension = D3D11_DSV_DIMENSION_TEXTURE2D;
-					DSVDesc.Texture2D.MipSlice = 0;
-				}
+				DSVDesc.ViewDimension = D3D11_DSV_DIMENSION_TEXTURE2DARRAY;
+				DSVDesc.Texture2DArray.FirstArraySlice = 0;
+				DSVDesc.Texture2DArray.ArraySize = TextureDesc.ArraySize;
+				DSVDesc.Texture2DArray.MipSlice = 0;
 			}
+		}
+		else
+		{
+			if (bIsMultisampled)
+			{
+				DSVDesc.ViewDimension = D3D11_DSV_DIMENSION_TEXTURE2DMS;
+				// Nothing to set
+			}
+			else
+			{
+				DSVDesc.ViewDimension = D3D11_DSV_DIMENSION_TEXTURE2D;
+				DSVDesc.Texture2D.MipSlice = 0;
+			}
+		}
 
-			for (uint32 AccessType = 0; AccessType < FExclusiveDepthStencil::MaxIndex; ++AccessType)
+		for (uint32 AccessType = 0; AccessType < FExclusiveDepthStencil::MaxIndex; ++AccessType)
+		{
+			// Create a read-only access views for the texture.
+			// Read-only DSVs are not supported in Feature Level 10 so 
+			// a dummy DSV is created in order reduce logic complexity at a higher-level.
+			DSVDesc.Flags = (AccessType & FExclusiveDepthStencil::DepthRead_StencilWrite) ? D3D11_DSV_READ_ONLY_DEPTH : 0;
+			if(HasStencilBits(DSVDesc.Format))
 			{
-				// Create a read-only access views for the texture.
-				// Read-only DSVs are not supported in Feature Level 10 so 
-				// a dummy DSV is created in order reduce logic complexity at a higher-level.
-				DSVDesc.Flags = (AccessType & FExclusiveDepthStencil::DepthRead_StencilWrite) ? D3D11_DSV_READ_ONLY_DEPTH : 0;
-				if(HasStencilBits(DSVDesc.Format))
-				{
-					DSVDesc.Flags |= (AccessType & FExclusiveDepthStencil::DepthWrite_StencilRead) ? D3D11_DSV_READ_ONLY_STENCIL : 0;
-				}
-				VERIFYD3D11RESULT_EX(Direct3DDevice->CreateDepthStencilView(TextureResource,&DSVDesc,DepthStencilViews[AccessType].GetInitReference()), Direct3DDevice);
+				DSVDesc.Flags |= (AccessType & FExclusiveDepthStencil::DepthWrite_StencilRead) ? D3D11_DSV_READ_ONLY_STENCIL : 0;
 			}
+			VERIFYD3D11RESULT_EX(Direct3DDevice->CreateDepthStencilView(TextureResource,&DSVDesc,DepthStencilViews[AccessType].GetInitReference()), Direct3DDevice);
 		}
 	}
 	check(IsValidRef(TextureResource));
@@ -1050,35 +751,15 @@ TD3D11Texture2D<BaseResourceType>* FD3D11DynamicRHI::CreateD3D11Texture2D(uint32
 		check(IsValidRef(ShaderResourceView));
 	}
 
-	TD3D11Texture2D<BaseResourceType>* Texture2D = new TD3D11Texture2D<BaseResourceType>(
-		this,
+	FD3D11Texture* Texture2D = new FD3D11Texture(
+		CreateDesc,
 		TextureResource,
 		ShaderResourceView,
-		bCreatedRTVPerSlice,
 		TextureDesc.ArraySize,
+		bCreatedRTVPerSlice,
 		RenderTargetViews,
-		DepthStencilViews,
-		SizeX,
-		SizeY,
-		SizeZ,
-		NumMips,
-		ActualMSAACount,
-		(EPixelFormat)Format,
-		bCubeTexture,
-		Flags,
-		bPooledTexture,
-		CreateInfo.ClearValueBinding
-#if PLATFORM_SUPPORTS_VIRTUAL_TEXTURES
-		, RawTextureMemory
-#endif
-		);
-
-	if (CreateInfo.DebugName)
-	{
-		Texture2D->SetName(CreateInfo.DebugName);
-	}
-
-	D3D11TextureAllocated(*Texture2D);
+		DepthStencilViews
+	);
 	
 #if !PLATFORM_HOLOLENS
 	if (IsRHIDeviceNVIDIA() && EnumHasAnyFlags(Flags, TexCreate_AFRManual))
@@ -1093,16 +774,25 @@ TD3D11Texture2D<BaseResourceType>* FD3D11DynamicRHI::CreateD3D11Texture2D(uint32
 		NvAPI_D3D_SetResourceHint(Direct3DDevice, (NVDX_ObjectHandle)IHVHandle, NVAPI_D3D_SRH_CATEGORY_SLI, NVAPI_D3D_SRH_SLI_APP_CONTROLLED_INTERFRAME_CONTENT_SYNC, &ManualAFR);
 	}
 #endif
-	if (CreateInfo.BulkData)
+	if (CreateDesc.BulkData)
 	{
-		CreateInfo.BulkData->Discard();
+		CreateDesc.BulkData->Discard();
 	}
 
 	return Texture2D;
 }
 
-FD3D11Texture3D* FD3D11DynamicRHI::CreateD3D11Texture3D(uint32 SizeX,uint32 SizeY,uint32 SizeZ,uint8 Format,uint32 NumMips,ETextureCreateFlags Flags,FRHIResourceCreateInfo& CreateInfo)
+FD3D11Texture* FD3D11DynamicRHI::CreateD3D11Texture3D(FRHITextureCreateDesc const& CreateDesc)
 {
+	check(CreateDesc.IsTexture3D());
+	check(CreateDesc.ArraySize == 1);
+	const uint32              SizeX   = CreateDesc.Extent.X;
+	const uint32              SizeY   = CreateDesc.Extent.Y;
+	const uint32              SizeZ   = CreateDesc.Depth;
+	const EPixelFormat        Format  = CreateDesc.Format;
+	const uint32              NumMips = CreateDesc.NumMips;
+	const ETextureCreateFlags Flags   = CreateDesc.Flags;
+
 	SCOPE_CYCLE_COUNTER(STAT_D3D11CreateTextureTime);
 	
 	const bool bSRGB = EnumHasAnyFlags(Flags, TexCreate_SRGB);
@@ -1149,9 +839,9 @@ FD3D11Texture3D* FD3D11DynamicRHI::CreateD3D11Texture3D(uint32 SizeX,uint32 Size
 
 	TArray<D3D11_SUBRESOURCE_DATA> SubResourceData;
 
-	if (CreateInfo.BulkData)
+	if (CreateDesc.BulkData)
 	{
-		uint8* Data = (uint8*)CreateInfo.BulkData->GetResourceBulkData();
+		uint8* Data = (uint8*)CreateDesc.BulkData->GetResourceBulkData();
 		SubResourceData.AddZeroed(NumMips);
 		uint32 MipOffset = 0;
 		for(uint32 MipIndex = 0;MipIndex < NumMips;++MipIndex)
@@ -1164,7 +854,7 @@ FD3D11Texture3D* FD3D11DynamicRHI::CreateD3D11Texture3D(uint32 SizeX,uint32 Size
 	}
 
 	TRefCountPtr<ID3D11Texture3D> TextureResource;
-	const D3D11_SUBRESOURCE_DATA* SubResData = CreateInfo.BulkData != nullptr ? (const D3D11_SUBRESOURCE_DATA*)SubResourceData.GetData() : nullptr;
+	const D3D11_SUBRESOURCE_DATA* SubResData = CreateDesc.BulkData != nullptr ? (const D3D11_SUBRESOURCE_DATA*)SubResourceData.GetData() : nullptr;
 	VERIFYD3D11CREATETEXTURERESULT(
 		Direct3DDevice->CreateTexture3D(&TextureDesc, SubResData,TextureResource.GetInitReference()),
 		Format,
@@ -1211,16 +901,16 @@ FD3D11Texture3D* FD3D11DynamicRHI::CreateD3D11Texture3D(uint32 SizeX,uint32 Size
 		VERIFYD3D11RESULT_EX(Direct3DDevice->CreateRenderTargetView(TextureResource,&RTVDesc,RenderTargetView.GetInitReference()), Direct3DDevice);
 	}
 
-	TArray<TRefCountPtr<ID3D11RenderTargetView> > RenderTargetViews;
-	RenderTargetViews.Add(RenderTargetView);
-	FD3D11Texture3D* Texture3D = new FD3D11Texture3D(this,TextureResource,ShaderResourceView,RenderTargetViews,SizeX,SizeY,SizeZ,NumMips,(EPixelFormat)Format,Flags, CreateInfo.ClearValueBinding);
-
-	if (CreateInfo.DebugName)
-	{
-		Texture3D->SetName(CreateInfo.DebugName);
-	}
+	FD3D11Texture* Texture3D = new FD3D11Texture(
+		CreateDesc,
+		TextureResource,
+		ShaderResourceView,
+		1,     // InRTVArraySize
+		false, // bInCreatedRTVsPerSlice
+		{ RenderTargetView },
+		{}
+	);
 
-	D3D11TextureAllocated(*Texture3D);
 #if !PLATFORM_HOLOLENS
 	if (IsRHIDeviceNVIDIA() && EnumHasAnyFlags(Flags, TexCreate_AFRManual))
 	{
@@ -1234,9 +924,9 @@ FD3D11Texture3D* FD3D11DynamicRHI::CreateD3D11Texture3D(uint32 SizeX,uint32 Size
 		NvAPI_D3D_SetResourceHint(Direct3DDevice, (NVDX_ObjectHandle)IHVHandle, NVAPI_D3D_SRH_CATEGORY_SLI, NVAPI_D3D_SRH_SLI_APP_CONTROLLED_INTERFRAME_CONTENT_SYNC, &ManualAFR);
 	}
 #endif
-	if (CreateInfo.BulkData)
+	if (CreateDesc.BulkData)
 	{
-		CreateInfo.BulkData->Discard();
+		CreateDesc.BulkData->Discard();
 	}
 
 	return Texture3D;
@@ -1246,147 +936,79 @@ FD3D11Texture3D* FD3D11DynamicRHI::CreateD3D11Texture3D(uint32 SizeX,uint32 Size
 	2D texture support.
 -----------------------------------------------------------------------------*/
 
-FTexture2DRHIRef FD3D11DynamicRHI::RHICreateTexture2D(uint32 SizeX,uint32 SizeY,uint8 Format,uint32 NumMips,uint32 NumSamples,ETextureCreateFlags Flags, ERHIAccess InResourceState,FRHIResourceCreateInfo& CreateInfo)
+FTextureRHIRef FD3D11DynamicRHI::RHICreateTexture(const FRHITextureCreateDesc& CreateDesc)
 {
-	return CreateD3D11Texture2D<FD3D11BaseTexture2D>(SizeX,SizeY,1,false,false,Format,NumMips,NumSamples,Flags,CreateInfo);
+	return CreateDesc.IsTexture3D()
+		? CreateD3D11Texture3D(CreateDesc)
+		: CreateD3D11Texture2D(CreateDesc);
 }
 
-FTexture2DRHIRef FD3D11DynamicRHI::RHICreateTexture2D_RenderThread(
-	class FRHICommandListImmediate& RHICmdList,
-	uint32 SizeX,
-	uint32 SizeY,
-	uint8 Format,
-	uint32 NumMips,
-	uint32 NumSamples,
-	ETextureCreateFlags Flags,
-	ERHIAccess InResourceState,
-	FRHIResourceCreateInfo& CreateInfo)
+FTextureRHIRef FD3D11DynamicRHI::RHICreateTexture_RenderThread(class FRHICommandListImmediate& RHICmdList, const FRHITextureCreateDesc& CreateDesc)
 {
-	return RHICreateTexture2D(SizeX, SizeY, Format, NumMips, NumSamples, Flags, InResourceState, CreateInfo);
+	return RHICreateTexture(CreateDesc);
 }
 
-FTexture2DRHIRef FD3D11DynamicRHI::RHIAsyncCreateTexture2D(uint32 SizeX,uint32 SizeY,uint8 Format,uint32 NumMips,ETextureCreateFlags Flags, ERHIAccess InResourceState, void** InitialMipData,uint32 NumInitialMips)
+FTextureRHIRef FD3D11DynamicRHI::RHIAsyncCreateTexture2D(uint32 SizeX, uint32 SizeY, uint8 Format, uint32 NumMips, ETextureCreateFlags Flags, ERHIAccess InResourceState, void** InitialMipData, uint32 NumInitialMips)
 {
-	FD3D11Texture2D* NewTexture = NULL;
-	TRefCountPtr<ID3D11Texture2D> TextureResource;
-	TRefCountPtr<ID3D11ShaderResourceView> ShaderResourceView;
-	D3D11_TEXTURE2D_DESC TextureDesc = {0};
-	D3D11_SHADER_RESOURCE_VIEW_DESC SRVDesc;
-
-	D3D11_SUBRESOURCE_DATA SubResourceData[ MAX_TEXTURE_MIP_COUNT ];
-	FPlatformMemory::Memzero( SubResourceData, sizeof( D3D11_SUBRESOURCE_DATA ) * MAX_TEXTURE_MIP_COUNT );
-
-	ETextureCreateFlags InvalidFlags = TexCreate_RenderTargetable | TexCreate_ResolveTargetable | TexCreate_DepthStencilTargetable | TexCreate_GenerateMipCapable | TexCreate_UAV | TexCreate_Presentable | TexCreate_CPUReadback;
-	check(!EnumHasAnyFlags(Flags, InvalidFlags));
-
-	TArray<TRefCountPtr<ID3D11RenderTargetView> > RenderTargetViews;
-	
-	check(GRHISupportsAsyncTextureCreation);
-
-	const DXGI_FORMAT PlatformResourceFormat = (DXGI_FORMAT)GPixelFormats[Format].PlatformFormat;
-	const DXGI_FORMAT PlatformShaderResourceFormat = FindShaderResourceDXGIFormat(PlatformResourceFormat, EnumHasAnyFlags(Flags, TexCreate_SRGB));
-
-	TextureDesc.Width = SizeX;
-	TextureDesc.Height = SizeY;
-	TextureDesc.MipLevels = NumMips;
-	TextureDesc.ArraySize = 1;
-	TextureDesc.Format = PlatformResourceFormat;
-	TextureDesc.SampleDesc.Count = 1;
-	TextureDesc.SampleDesc.Quality = 0;
-	TextureDesc.Usage = D3D11_USAGE_DEFAULT;
-	TextureDesc.BindFlags = D3D11_BIND_SHADER_RESOURCE;
-	TextureDesc.CPUAccessFlags = 0;
-	TextureDesc.MiscFlags = 0;
-
-	ApplyBC7SoftwareAdapterWorkaround(Adapter.bSoftwareAdapter, TextureDesc);
+	TArray<D3D11_SUBRESOURCE_DATA, TInlineAllocator<12>> SubresourceData;
+	SubresourceData.SetNumUninitialized(NumMips);
 
 	for (uint32 MipIndex = 0; MipIndex < NumInitialMips; ++MipIndex)
 	{
-		uint32 NumBlocksX = FMath::Max<uint32>(1,(SizeX >> MipIndex) / GPixelFormats[Format].BlockSizeX);
-		uint32 NumBlocksY = FMath::Max<uint32>(1,(SizeY >> MipIndex) / GPixelFormats[Format].BlockSizeY);
+		uint32 NumBlocksX = FMath::Max<uint32>(1, (SizeX >> MipIndex) / GPixelFormats[Format].BlockSizeX);
+		uint32 NumBlocksY = FMath::Max<uint32>(1, (SizeY >> MipIndex) / GPixelFormats[Format].BlockSizeY);
 
-		SubResourceData[MipIndex].pSysMem = InitialMipData[MipIndex];
-		SubResourceData[MipIndex].SysMemPitch = NumBlocksX * GPixelFormats[Format].BlockBytes;
-		SubResourceData[MipIndex].SysMemSlicePitch = NumBlocksX * NumBlocksY * GPixelFormats[Format].BlockBytes;
+		SubresourceData[MipIndex].pSysMem = InitialMipData[MipIndex];
+		SubresourceData[MipIndex].SysMemPitch = NumBlocksX * GPixelFormats[Format].BlockBytes;
+		SubresourceData[MipIndex].SysMemSlicePitch = NumBlocksX * NumBlocksY * GPixelFormats[Format].BlockBytes;
 	}
 
 	void* TempBuffer = ZeroBuffer;
 	uint32 TempBufferSize = ZeroBufferSize;
 	for (uint32 MipIndex = NumInitialMips; MipIndex < NumMips; ++MipIndex)
 	{
-		uint32 NumBlocksX = FMath::Max<uint32>(1,(SizeX >> MipIndex) / GPixelFormats[Format].BlockSizeX);
-		uint32 NumBlocksY = FMath::Max<uint32>(1,(SizeY >> MipIndex) / GPixelFormats[Format].BlockSizeY);
+		uint32 NumBlocksX = FMath::Max<uint32>(1, (SizeX >> MipIndex) / GPixelFormats[Format].BlockSizeX);
+		uint32 NumBlocksY = FMath::Max<uint32>(1, (SizeY >> MipIndex) / GPixelFormats[Format].BlockSizeY);
 		uint32 MipSize = NumBlocksX * NumBlocksY * GPixelFormats[Format].BlockBytes;
 
 		if (MipSize > TempBufferSize)
 		{
-			UE_LOG(LogD3D11RHI, Display,TEXT("Temp texture streaming buffer not large enough, needed %d bytes"),MipSize);
+			UE_LOG(LogD3D11RHI, Display, TEXT("Temp texture streaming buffer not large enough, needed %d bytes"), MipSize);
 			check(TempBufferSize == ZeroBufferSize);
 			TempBufferSize = MipSize;
 			TempBuffer = FMemory::Malloc(TempBufferSize);
-			FMemory::Memzero(TempBuffer,TempBufferSize);
+			FMemory::Memzero(TempBuffer, TempBufferSize);
 		}
 
-		SubResourceData[MipIndex].pSysMem = TempBuffer;
-		SubResourceData[MipIndex].SysMemPitch = NumBlocksX * GPixelFormats[Format].BlockBytes;
-		SubResourceData[MipIndex].SysMemSlicePitch = MipSize;
+		SubresourceData[MipIndex].pSysMem = TempBuffer;
+		SubresourceData[MipIndex].SysMemPitch = NumBlocksX * GPixelFormats[Format].BlockBytes;
+		SubresourceData[MipIndex].SysMemSlicePitch = MipSize;
 	}
 
-#if INTEL_EXTENSIONS
-	if (EnumHasAnyFlags(Flags, ETextureCreateFlags::Atomic64Compatible) && IsRHIDeviceIntel() && GRHISupportsAtomicUInt64)
-	{
-		INTC_D3D11_TEXTURE2D_DESC IntelDesc{};
-		IntelDesc.EmulatedTyped64bitAtomics = true;
-		IntelDesc.pD3D11Desc = &TextureDesc;
-
-		INTC_D3D11_CreateTexture2D(IntelExtensionContext, &IntelDesc, SubResourceData, TextureResource.GetInitReference());
-	}
-	else
-#endif
-	{
-		SafeCreateTexture2D(Direct3DDevice, Format, &TextureDesc, SubResourceData, TextureResource.GetInitReference());
-	}
+	FTextureRHIRef Texture = CreateD3D11Texture2D(
+		FRHITextureCreateDesc::Create2D(
+			  TEXT("RHIAsyncCreateTexture2D")
+			, { (int32)SizeX, (int32)SizeY }
+			, (EPixelFormat)Format
+			, FClearValueBinding::None
+			, Flags
+			, NumMips
+		)
+		, SubresourceData
+	);
 
 	if (TempBufferSize != ZeroBufferSize)
 	{
 		FMemory::Free(TempBuffer);
 	}
 
-	SRVDesc.Format = PlatformShaderResourceFormat;
-	SRVDesc.ViewDimension = D3D11_SRV_DIMENSION_TEXTURE2D;
-	SRVDesc.Texture2D.MostDetailedMip = 0;
-	SRVDesc.Texture2D.MipLevels = NumMips;
-	VERIFYD3D11RESULT_EX(Direct3DDevice->CreateShaderResourceView(TextureResource,&SRVDesc,ShaderResourceView.GetInitReference()), Direct3DDevice);
-
-	NewTexture = new FD3D11Texture2D(
-		this,
-		TextureResource,
-		ShaderResourceView,
-		false,
-		1,
-		RenderTargetViews,
-		/*DepthStencilViews=*/ NULL,
-		SizeX,
-		SizeY,
-		0,
-		NumMips,
-		/*ActualMSAACount=*/ 1,
-		(EPixelFormat)Format,
-		/*bInCubemap=*/ false,
-		Flags,
-		/*bPooledTexture=*/ false,
-		FClearValueBinding()
-		);
-
-	D3D11TextureAllocated(*NewTexture);
-	
-	return NewTexture;
+	return Texture;
 }
 
 void FD3D11DynamicRHI::RHICopySharedMips(FRHITexture2D* DestTexture2DRHI, FRHITexture2D* SrcTexture2DRHI)
 {
-	FD3D11Texture2D* DestTexture2D = ResourceCast(DestTexture2DRHI);
-	FD3D11Texture2D* SrcTexture2D = ResourceCast(SrcTexture2DRHI);
+	FD3D11Texture* DestTexture2D = ResourceCast(DestTexture2DRHI);
+	FD3D11Texture* SrcTexture2D = ResourceCast(SrcTexture2DRHI);
 
 	// Use the GPU to asynchronously copy the old mip-maps into the new texture.
 	const uint32 NumSharedMips = FMath::Min(DestTexture2D->GetNumMips(),SrcTexture2D->GetNumMips());
@@ -1408,50 +1030,9 @@ void FD3D11DynamicRHI::RHICopySharedMips(FRHITexture2D* DestTexture2DRHI, FRHITe
 	}
 }
 
-FTexture2DArrayRHIRef FD3D11DynamicRHI::RHICreateTexture2DArray(uint32 SizeX,uint32 SizeY,uint32 SizeZ,uint8 Format,uint32 NumMips,uint32 NumSamples,ETextureCreateFlags Flags, ERHIAccess InResourceState,FRHIResourceCreateInfo& CreateInfo)
-{
-	check(SizeZ >= 1);
-	return CreateD3D11Texture2D<FD3D11BaseTexture2DArray>(SizeX,SizeY,SizeZ,true,false,Format,NumMips,NumSamples,Flags,CreateInfo);
-}
-
-FTexture2DArrayRHIRef FD3D11DynamicRHI::RHICreateTexture2DArray_RenderThread(
-	class FRHICommandListImmediate& RHICmdList,
-	uint32 SizeX,
-	uint32 SizeY,
-	uint32 SizeZ,
-	uint8 Format,
-	uint32 NumMips,
-	uint32 NumSamples,
-	ETextureCreateFlags Flags,
-	ERHIAccess InResourceState,
-	FRHIResourceCreateInfo& CreateInfo)
-{
-	return RHICreateTexture2DArray(SizeX, SizeY, SizeZ, Format, NumMips, NumSamples, Flags, InResourceState, CreateInfo);
-}
-
-FTexture3DRHIRef FD3D11DynamicRHI::RHICreateTexture3D(uint32 SizeX,uint32 SizeY,uint32 SizeZ,uint8 Format,uint32 NumMips,ETextureCreateFlags Flags, ERHIAccess InResourceState,FRHIResourceCreateInfo& CreateInfo)
-{
-	check(SizeZ >= 1);
-	return CreateD3D11Texture3D(SizeX,SizeY,SizeZ,Format,NumMips,Flags,CreateInfo);
-}
-
-FTexture3DRHIRef FD3D11DynamicRHI::RHICreateTexture3D_RenderThread(
-	class FRHICommandListImmediate& RHICmdList,
-	uint32 SizeX,
-	uint32 SizeY,
-	uint32 SizeZ,
-	uint8 Format,
-	uint32 NumMips,
-	ETextureCreateFlags Flags,
-	ERHIAccess InResourceState,
-	FRHIResourceCreateInfo& CreateInfo)
-{
-	return RHICreateTexture3D(SizeX, SizeY, SizeZ, Format, NumMips, Flags, InResourceState, CreateInfo);
-}
-
 FShaderResourceViewRHIRef FD3D11DynamicRHI::RHICreateShaderResourceView(FRHITexture* TextureRHI, const FRHITextureSRVCreateInfo& CreateInfo)
 {
-	FD3D11TextureBase* Texture = GetD3D11TextureFromRHITexture(TextureRHI);
+	FD3D11Texture* Texture = GetD3D11TextureFromRHITexture(TextureRHI);
 
 	// Create a Shader Resource View
 	D3D11_SHADER_RESOURCE_VIEW_DESC SRVDesc;
@@ -1459,10 +1040,8 @@ FShaderResourceViewRHIRef FD3D11DynamicRHI::RHICreateShaderResourceView(FRHIText
 
 	if (TextureRHI->GetTexture3D() != NULL)
 	{
-		FD3D11Texture3D* Texture3D = static_cast<FD3D11Texture3D*>(Texture);
-
 		D3D11_TEXTURE3D_DESC TextureDesc;
-		Texture3D->GetResource()->GetDesc(&TextureDesc);
+		Texture->GetD3D11Texture3D()->GetDesc(&TextureDesc);
 		BaseTextureFormat = TextureDesc.Format;
 
 		SRVDesc.ViewDimension = D3D11_SRV_DIMENSION_TEXTURE3D;
@@ -1471,10 +1050,8 @@ FShaderResourceViewRHIRef FD3D11DynamicRHI::RHICreateShaderResourceView(FRHIText
 	}
 	else if (TextureRHI->GetTexture2DArray() != NULL)
 	{
-		FD3D11Texture2DArray* Texture2DArray = static_cast<FD3D11Texture2DArray*>(Texture);
-
 		D3D11_TEXTURE2D_DESC TextureDesc;
-		Texture2DArray->GetResource()->GetDesc(&TextureDesc);
+		Texture->GetD3D11Texture2D()->GetDesc(&TextureDesc);
 		BaseTextureFormat = TextureDesc.Format;
 
 		SRVDesc.ViewDimension = D3D11_SRV_DIMENSION_TEXTURE2DARRAY;
@@ -1485,10 +1062,8 @@ FShaderResourceViewRHIRef FD3D11DynamicRHI::RHICreateShaderResourceView(FRHIText
 	}
 	else if (TextureRHI->GetTextureCube() != NULL)
 	{
-		FD3D11TextureCube* TextureCube = static_cast<FD3D11TextureCube*>(Texture);
-
 		D3D11_TEXTURE2D_DESC TextureDesc;
-		TextureCube->GetResource()->GetDesc(&TextureDesc);
+		Texture->GetD3D11Texture2D()->GetDesc(&TextureDesc);
 		BaseTextureFormat = TextureDesc.Format;
 
 		SRVDesc.ViewDimension = D3D11_SRV_DIMENSION_TEXTURECUBE;
@@ -1497,10 +1072,8 @@ FShaderResourceViewRHIRef FD3D11DynamicRHI::RHICreateShaderResourceView(FRHIText
 	}
 	else
 	{
-		FD3D11Texture2D* Texture2D = static_cast<FD3D11Texture2D*>(Texture);
-
 		D3D11_TEXTURE2D_DESC TextureDesc;
-		Texture2D->GetResource()->GetDesc(&TextureDesc);
+		Texture->GetD3D11Texture2D()->GetDesc(&TextureDesc);
 		BaseTextureFormat = TextureDesc.Format;
 
 		if (TextureDesc.SampleDesc.Count > 1)
@@ -1540,7 +1113,7 @@ FShaderResourceViewRHIRef FD3D11DynamicRHI::RHICreateShaderResourceView_RenderTh
 /** Generates mip maps for the surface. */
 void FD3D11DynamicRHI::RHIGenerateMips(FRHITexture* TextureRHI)
 {
-	FD3D11TextureBase* Texture = GetD3D11TextureFromRHITexture(TextureRHI);
+	FD3D11Texture* Texture = GetD3D11TextureFromRHITexture(TextureRHI);
 	// Surface must have been created with D3D11_BIND_RENDER_TARGET for GenerateMips to work
 	check(Texture->GetShaderResourceView() && Texture->GetRenderTargetView(0, -1));
 	Direct3DDeviceIMContext->GenerateMips(Texture->GetShaderResourceView());
@@ -1561,7 +1134,7 @@ uint32 FD3D11DynamicRHI::RHIComputeMemorySize(FRHITexture* TextureRHI)
 		return 0;
 	}
 	
-	FD3D11TextureBase* Texture = GetD3D11TextureFromRHITexture(TextureRHI);
+	FD3D11Texture* Texture = GetD3D11TextureFromRHITexture(TextureRHI);
 	return Texture->GetMemorySize();
 }
 
@@ -1578,8 +1151,8 @@ uint32 FD3D11DynamicRHI::RHIComputeMemorySize(FRHITexture* TextureRHI)
  */
 void FD3D11DynamicRHI::RHIAsyncCopyTexture2DCopy(FRHITexture2D* NewTexture2DRHI, FRHITexture2D* Texture2DRHI, int32 NewMipCount, int32 NewSizeX, int32 NewSizeY, FThreadSafeCounter* RequestStatus)
 {
-	FD3D11Texture2D* Texture2D = ResourceCast(Texture2DRHI);
-	FD3D11Texture2D* NewTexture2D = ResourceCast(NewTexture2DRHI);
+	FD3D11Texture* Texture2D = ResourceCast(Texture2DRHI);
+	FD3D11Texture* NewTexture2D = ResourceCast(NewTexture2DRHI);
 
 	// Use the GPU to asynchronously copy the old mip-maps into the new texture.
 	const uint32 NumSharedMips = FMath::Min(Texture2D->GetNumMips(), NewTexture2D->GetNumMips());
@@ -1624,11 +1197,18 @@ void FD3D11DynamicRHI::RHIAsyncCopyTexture2DCopy(FRHITexture2D* NewTexture2DRHI,
  */
 FTexture2DRHIRef FD3D11DynamicRHI::RHIAsyncReallocateTexture2D(FRHITexture2D* Texture2DRHI, int32 NewMipCount, int32 NewSizeX, int32 NewSizeY, FThreadSafeCounter* RequestStatus)
 {
-	FD3D11Texture2D* Texture2D = ResourceCast(Texture2DRHI);
+	FD3D11Texture* Texture2D = ResourceCast(Texture2DRHI);
+
+	FRHITextureCreateDesc CreateDesc(
+		Texture2D->GetDesc(),
+		RHIGetDefaultResourceState(Texture2D->GetDesc().Flags, false),
+		TEXT("RHIAsyncReallocateTexture2D")
+	);
+	CreateDesc.Extent = FIntPoint(NewSizeX, NewSizeY);
+	CreateDesc.NumMips = NewMipCount;
 
 	// Allocate a new texture.
-	FRHIResourceCreateInfo CreateInfo(TEXT("RHIAsyncReallocateTexture2D"));
-	FD3D11Texture2D* NewTexture2D = CreateD3D11Texture2D<FD3D11BaseTexture2D>(NewSizeX,NewSizeY,1,false,false,Texture2D->GetFormat(),NewMipCount,1,Texture2D->GetFlags(),CreateInfo);
+	FD3D11Texture* NewTexture2D = CreateD3D11Texture2D(CreateDesc);
 
 	RHIAsyncCopyTexture2DCopy(NewTexture2D, Texture2D, NewMipCount, NewSizeX, NewSizeY, RequestStatus);
 
@@ -1652,9 +1232,15 @@ FTexture2DRHIRef FD3D11DynamicRHI::AsyncReallocateTexture2D_RenderThread(
 	else
 	{
 		// Allocate a new texture.
-		FRHIResourceCreateInfo CreateInfo(TEXT("AsyncReallocateTexture2D_RenderThread"));
-		FD3D11Texture2D* NewTexture2DPointer = CreateD3D11Texture2D<FD3D11BaseTexture2D>(NewSizeX, NewSizeY, 1, false, false, Texture2D->GetFormat(), NewMipCount, 1, Texture2D->GetFlags(), CreateInfo);
-		NewTexture2D = NewTexture2DPointer;
+		FRHITextureCreateDesc CreateDesc(
+			Texture2D->GetDesc(),
+			RHIGetDefaultResourceState(Texture2D->GetDesc().Flags, false),
+			TEXT("AsyncReallocateTexture2D_RenderThread")
+		);
+		CreateDesc.Extent = FIntPoint(NewSizeX, NewSizeY);
+		CreateDesc.NumMips = NewMipCount;
+
+		NewTexture2D = CreateD3D11Texture2D(CreateDesc);
 
 		RunOnRHIThread([this, NewTexture2D, Texture2D, NewMipCount, NewSizeX, NewSizeY, RequestStatus]()
 		{
@@ -1707,9 +1293,10 @@ ETextureReallocationStatus FD3D11DynamicRHI::CancelAsyncReallocateTexture2D_Rend
 	return RHICancelAsyncReallocateTexture2D(Texture2D, bBlockUntilCompleted);
 }
 
-template<typename RHIResourceType>
-void* TD3D11Texture2D<RHIResourceType>::Lock(uint32 MipIndex,uint32 ArrayIndex,EResourceLockMode LockMode,uint32& DestStride,bool bForceLockDeferred)
+void* FD3D11Texture::Lock(FD3D11DynamicRHI* D3DRHI, uint32 MipIndex, uint32 ArrayIndex, EResourceLockMode LockMode, uint32& DestStride, bool bForceLockDeferred)
 {
+	check(!IsTexture3D()); // Only 2D texture locks are implemented
+
 	SCOPE_CYCLE_COUNTER(STAT_D3D11LockTextureTime);
 
 	// Calculate the subresource index corresponding to the specified mip-map.
@@ -1726,16 +1313,9 @@ void* TD3D11Texture2D<RHIResourceType>::Lock(uint32 MipIndex,uint32 ArrayIndex,E
 	const uint32 MipBytes = NumBlocksX * NumBlocksY * BlockBytes;
 
 	FD3D11LockedData LockedData;
-#if	PLATFORM_SUPPORTS_VIRTUAL_TEXTURES
-	if (D3DRHI->HandleSpecialLock(LockedData, MipIndex, ArrayIndex, GetFlags(), LockMode, GetResource(), RawTextureMemory, GetNumMips(), DestStride))
-	{
-		// nothing left to do...
-	}
-	else
-#endif
 	if( LockMode == RLM_WriteOnly )
 	{
-		if (!bForceLockDeferred && EnumHasAnyFlags(Flags, TexCreate_CPUWritable))
+		if (!bForceLockDeferred && EnumHasAnyFlags(GetDesc().Flags, TexCreate_CPUWritable))
 		{
 			D3D11_MAPPED_SUBRESOURCE MappedTexture;
 			VERIFYD3D11RESULT_EX(D3DRHI->GetDeviceContext()->Map(GetResource(), Subresource, D3D11_MAP_WRITE, 0, &MappedTexture), D3DRHI->GetDevice());
@@ -1757,7 +1337,7 @@ void* TD3D11Texture2D<RHIResourceType>::Lock(uint32 MipIndex,uint32 ArrayIndex,E
 
 		// Create the staging texture.
 		D3D11_TEXTURE2D_DESC StagingTextureDesc;
-		GetResource()->GetDesc(&StagingTextureDesc);
+		GetD3D11Texture2D()->GetDesc(&StagingTextureDesc);
 		StagingTextureDesc.Width = MipSizeX;
 		StagingTextureDesc.Height = MipSizeY;
 		StagingTextureDesc.MipLevels = 1;
@@ -1769,7 +1349,7 @@ void* TD3D11Texture2D<RHIResourceType>::Lock(uint32 MipIndex,uint32 ArrayIndex,E
 		TRefCountPtr<ID3D11Texture2D> StagingTexture;
 		VERIFYD3D11CREATETEXTURERESULT(
 			D3DRHI->GetDevice()->CreateTexture2D(&StagingTextureDesc,NULL,StagingTexture.GetInitReference()),
-			RHIResourceType::GetFormat(),
+			GetDesc().Format,
 			this->GetSizeX(),
 			this->GetSizeY(),
 			this->GetSizeZ(),
@@ -1805,7 +1385,7 @@ void* TD3D11Texture2D<RHIResourceType>::Lock(uint32 MipIndex,uint32 ArrayIndex,E
 	}
 	else
 	{
-		RunOnRHIThread([this, Subresource, LockedData]()
+		RunOnRHIThread([this, Subresource, LockedData, D3DRHI]()
 		{
 			D3DRHI->AddLockedData(FD3D11LockedKey(GetResource(), Subresource), LockedData);
 		});
@@ -1814,9 +1394,10 @@ void* TD3D11Texture2D<RHIResourceType>::Lock(uint32 MipIndex,uint32 ArrayIndex,E
 	return (void*)LockedData.GetData();
 }
 
-template<typename RHIResourceType>
-void TD3D11Texture2D<RHIResourceType>::Unlock(uint32 MipIndex,uint32 ArrayIndex)
+void FD3D11Texture::Unlock(FD3D11DynamicRHI* D3DRHI, uint32 MipIndex, uint32 ArrayIndex)
 {
+	check(!IsTexture3D()); // Only 2D texture locks are implemented
+
 	SCOPE_CYCLE_COUNTER(STAT_D3D11UnlockTextureTime);
 
 	// Calculate the subresource index corresponding to the specified mip-map.
@@ -1826,14 +1407,7 @@ void TD3D11Texture2D<RHIResourceType>::Unlock(uint32 MipIndex,uint32 ArrayIndex)
 	FD3D11LockedData LockedData;
 	verifyf(D3DRHI->RemoveLockedData(FD3D11LockedKey(GetResource(), Subresource), LockedData), TEXT("Texture is not locked"));
 
-#if PLATFORM_SUPPORTS_VIRTUAL_TEXTURES
-	if (D3DRHI->HandleSpecialUnlock(MipIndex, GetFlags(), GetResource(), RawTextureMemory))
-	{
-		// nothing left to do...
-	}
-	else
-#endif
-	if (!LockedData.bLockDeferred && EnumHasAnyFlags(Flags, TexCreate_CPUWritable))
+	if (!LockedData.bLockDeferred && EnumHasAnyFlags(GetDesc().Flags, TexCreate_CPUWritable))
 	{
 		D3DRHI->GetDeviceContext()->Unmap(GetResource(), 0);
 	}
@@ -1849,12 +1423,12 @@ void TD3D11Texture2D<RHIResourceType>::Unlock(uint32 MipIndex,uint32 ArrayIndex)
 	}
 }
 
-void* FD3D11DynamicRHI::RHILockTexture2D(FRHITexture2D* TextureRHI,uint32 MipIndex,EResourceLockMode LockMode,uint32& DestStride,bool bLockWithinMiptail)
+void* FD3D11DynamicRHI::RHILockTexture2D(FRHITexture2D* TextureRHI, uint32 MipIndex, EResourceLockMode LockMode, uint32& DestStride, bool bLockWithinMiptail)
 {
 	check(TextureRHI);
-	FD3D11Texture2D* Texture = ResourceCast(TextureRHI);
+	FD3D11Texture* Texture = ResourceCast(TextureRHI);
 	ConditionalClearShaderResource(Texture, false);
-	return Texture->Lock(MipIndex,0,LockMode,DestStride);
+	return Texture->Lock(this, MipIndex, 0, LockMode, DestStride);
 }
 
 void* FD3D11DynamicRHI::LockTexture2D_RenderThread(
@@ -1879,17 +1453,17 @@ void* FD3D11DynamicRHI::LockTexture2D_RenderThread(
 	}
 	else
 	{
-		FD3D11Texture2D* TextureD3D11 = ResourceCast(Texture);
-		LockedTexture = TextureD3D11->Lock(MipIndex, 0, LockMode, DestStride, true);
+		FD3D11Texture* TextureD3D11 = ResourceCast(Texture);
+		LockedTexture = TextureD3D11->Lock(this, MipIndex, 0, LockMode, DestStride, true);
 	}
 	return LockedTexture;
 }
 
-void FD3D11DynamicRHI::RHIUnlockTexture2D(FRHITexture2D* TextureRHI,uint32 MipIndex,bool bLockWithinMiptail)
+void FD3D11DynamicRHI::RHIUnlockTexture2D(FRHITexture2D* TextureRHI, uint32 MipIndex, bool bLockWithinMiptail)
 {
 	check(TextureRHI);
-	FD3D11Texture2D* Texture = ResourceCast(TextureRHI);
-	Texture->Unlock(MipIndex,0);
+	FD3D11Texture* Texture = ResourceCast(TextureRHI);
+	Texture->Unlock(this, MipIndex, 0);
 }
 
 void FD3D11DynamicRHI::UnlockTexture2D_RenderThread(
@@ -1912,22 +1486,22 @@ void FD3D11DynamicRHI::UnlockTexture2D_RenderThread(
 	}
 }
 
-void* FD3D11DynamicRHI::RHILockTexture2DArray(FRHITexture2DArray* TextureRHI,uint32 TextureIndex,uint32 MipIndex,EResourceLockMode LockMode,uint32& DestStride,bool bLockWithinMiptail)
+void* FD3D11DynamicRHI::RHILockTexture2DArray(FRHITexture2DArray* TextureRHI, uint32 TextureIndex, uint32 MipIndex, EResourceLockMode LockMode, uint32& DestStride, bool bLockWithinMiptail)
 {
-	FD3D11Texture2DArray* Texture = ResourceCast(TextureRHI);
+	FD3D11Texture* Texture = ResourceCast(TextureRHI);
 	ConditionalClearShaderResource(Texture, false);
-	return Texture->Lock(MipIndex,TextureIndex,LockMode,DestStride);
+	return Texture->Lock(this, MipIndex, TextureIndex, LockMode, DestStride);
 }
 
-void FD3D11DynamicRHI::RHIUnlockTexture2DArray(FRHITexture2DArray* TextureRHI,uint32 TextureIndex,uint32 MipIndex,bool bLockWithinMiptail)
+void FD3D11DynamicRHI::RHIUnlockTexture2DArray(FRHITexture2DArray* TextureRHI, uint32 TextureIndex, uint32 MipIndex, bool bLockWithinMiptail)
 {
-	FD3D11Texture2DArray* Texture = ResourceCast(TextureRHI);
-	Texture->Unlock(MipIndex,TextureIndex);
+	FD3D11Texture* Texture = ResourceCast(TextureRHI);
+	Texture->Unlock(this, MipIndex, TextureIndex);
 }
 
-void FD3D11DynamicRHI::RHIUpdateTexture2D(FRHITexture2D* TextureRHI,uint32 MipIndex,const FUpdateTextureRegion2D& UpdateRegion,uint32 SourcePitch,const uint8* SourceData)
+void FD3D11DynamicRHI::RHIUpdateTexture2D(FRHITexture2D* TextureRHI, uint32 MipIndex, const FUpdateTextureRegion2D& UpdateRegion, uint32 SourcePitch, const uint8* SourceData)
 {
-	FD3D11Texture2D* Texture = ResourceCast(TextureRHI);
+	FD3D11Texture* Texture = ResourceCast(TextureRHI);
 
 	D3D11_BOX DestBox =
 	{
@@ -1974,7 +1548,7 @@ void FD3D11DynamicRHI::UpdateTexture2D_RenderThread(
 
 void FD3D11DynamicRHI::RHIUpdateTexture3D(FRHITexture3D* TextureRHI,uint32 MipIndex,const FUpdateTextureRegion3D& UpdateRegion,uint32 SourceRowPitch,uint32 SourceDepthPitch,const uint8* SourceData)
 {
-	FD3D11Texture3D* Texture = ResourceCast(TextureRHI);
+	FD3D11Texture* Texture = ResourceCast(TextureRHI);
 
 	// The engine calls this with the texture size in the region. 
 	// Some platforms like D3D11 needs that to be rounded up to the block size.
@@ -1984,7 +1558,7 @@ void FD3D11DynamicRHI::RHIUpdateTexture3D(FRHITexture3D* TextureRHI,uint32 MipIn
 
 	D3D11_BOX DestBox =
 	{
-		UpdateRegion.DestX,                      UpdateRegion.DestY,                       UpdateRegion.DestZ,
+		UpdateRegion.DestX,                                 UpdateRegion.DestY,                                 UpdateRegion.DestZ,
 		UpdateRegion.DestX + NumBlockX * Format.BlockSizeX, UpdateRegion.DestY + NumBlockY * Format.BlockSizeY, UpdateRegion.DestZ + UpdateRegion.Depth
 	};
 
@@ -2049,47 +1623,13 @@ void FD3D11DynamicRHI::UpdateTexture3D_RenderThread(
 /*-----------------------------------------------------------------------------
 	Cubemap texture support.
 -----------------------------------------------------------------------------*/
-FTextureCubeRHIRef FD3D11DynamicRHI::RHICreateTextureCube(uint32 Size, uint8 Format, uint32 NumMips, ETextureCreateFlags Flags, ERHIAccess InResourceState, FRHIResourceCreateInfo& CreateInfo)
-{
-	return CreateD3D11Texture2D<FD3D11BaseTextureCube>(Size,Size,6,false,true,Format,NumMips,1,Flags,CreateInfo);
-}
-
-FTextureCubeRHIRef FD3D11DynamicRHI::RHICreateTextureCube_RenderThread(
-	class FRHICommandListImmediate& RHICmdList,
-	uint32 Size,
-	uint8 Format,
-	uint32 NumMips,
-	ETextureCreateFlags Flags,
-	ERHIAccess InResourceState,
-	FRHIResourceCreateInfo& CreateInfo)
-{
-	return RHICreateTextureCube(Size, Format, NumMips, Flags, InResourceState, CreateInfo);
-}
-
-FTextureCubeRHIRef FD3D11DynamicRHI::RHICreateTextureCubeArray(uint32 Size, uint32 ArraySize, uint8 Format, uint32 NumMips, ETextureCreateFlags Flags, ERHIAccess InResourceState, FRHIResourceCreateInfo& CreateInfo)
-{
-	return CreateD3D11Texture2D<FD3D11BaseTextureCube>(Size,Size,6 * ArraySize,true,true,Format,NumMips,1,Flags,CreateInfo);
-}
-
-FTextureCubeRHIRef FD3D11DynamicRHI::RHICreateTextureCubeArray_RenderThread(
-	class FRHICommandListImmediate& RHICmdList,
-	uint32 Size,
-	uint32 ArraySize,
-	uint8 Format,
-	uint32 NumMips,
-	ETextureCreateFlags Flags,
-	ERHIAccess InResourceState,
-	FRHIResourceCreateInfo& CreateInfo)
-{
-	return RHICreateTextureCubeArray(Size, ArraySize, Format, NumMips, Flags, InResourceState, CreateInfo);
-}
 
-void* FD3D11DynamicRHI::RHILockTextureCubeFace(FRHITextureCube* TextureCubeRHI,uint32 FaceIndex,uint32 ArrayIndex,uint32 MipIndex,EResourceLockMode LockMode,uint32& DestStride,bool bLockWithinMiptail)
+void* FD3D11DynamicRHI::RHILockTextureCubeFace(FRHITextureCube* TextureCubeRHI, uint32 FaceIndex, uint32 ArrayIndex, uint32 MipIndex, EResourceLockMode LockMode, uint32& DestStride, bool bLockWithinMiptail)
 {
-	FD3D11TextureCube* TextureCube = ResourceCast(TextureCubeRHI);
+	FD3D11Texture* TextureCube = ResourceCast(TextureCubeRHI);
 	ConditionalClearShaderResource(TextureCube, false);
 	uint32 D3DFace = GetD3D11CubeFace((ECubeFace)FaceIndex);
-	return TextureCube->Lock(MipIndex,D3DFace + ArrayIndex * 6,LockMode,DestStride);
+	return TextureCube->Lock(this, MipIndex, D3DFace + ArrayIndex * 6, LockMode, DestStride);
 }
 
 void* FD3D11DynamicRHI::RHILockTextureCubeFace_RenderThread(
@@ -2115,18 +1655,18 @@ void* FD3D11DynamicRHI::RHILockTextureCubeFace_RenderThread(
 	}
 	else
 	{
-		FD3D11TextureCube* TextureCube = ResourceCast(Texture);
+		FD3D11Texture* TextureCube = ResourceCast(Texture);
 		const uint32 D3DFace = GetD3D11CubeFace((ECubeFace)FaceIndex);
-		LockedTexture = TextureCube->Lock(MipIndex, D3DFace + ArrayIndex * 6, LockMode, DestStride, true);
+		LockedTexture = TextureCube->Lock(this, MipIndex, D3DFace + ArrayIndex * 6, LockMode, DestStride, true);
 	}
 	return LockedTexture;
 }
 
-void FD3D11DynamicRHI::RHIUnlockTextureCubeFace(FRHITextureCube* TextureCubeRHI,uint32 FaceIndex,uint32 ArrayIndex,uint32 MipIndex,bool bLockWithinMiptail)
+void FD3D11DynamicRHI::RHIUnlockTextureCubeFace(FRHITextureCube* TextureCubeRHI, uint32 FaceIndex, uint32 ArrayIndex, uint32 MipIndex, bool bLockWithinMiptail)
 {
-	FD3D11TextureCube* TextureCube = ResourceCast(TextureCubeRHI);
+	FD3D11Texture* TextureCube = ResourceCast(TextureCubeRHI);
 	uint32 D3DFace = GetD3D11CubeFace((ECubeFace)FaceIndex);
-	TextureCube->Unlock(MipIndex,D3DFace + ArrayIndex * 6);
+	TextureCube->Unlock(this, MipIndex, D3DFace + ArrayIndex * 6);
 }
 
 void FD3D11DynamicRHI::RHIUnlockTextureCubeFace_RenderThread(
@@ -2152,22 +1692,13 @@ void FD3D11DynamicRHI::RHIUnlockTextureCubeFace_RenderThread(
 
 void FD3D11DynamicRHI::RHIBindDebugLabelName(FRHITexture* TextureRHI, const TCHAR* Name)
 {
+	FD3D11Texture* Texture = GetD3D11TextureFromRHITexture(TextureRHI);
+
 	//todo: require names at texture creation time.
 	FName DebugName(Name);
-	TextureRHI->SetName(DebugName);
+	Texture->SetName(DebugName);
 #if UE_BUILD_DEBUG || UE_BUILD_DEVELOPMENT
-	if (FD3D11Texture2D* Texture2D = (FD3D11Texture2D*)TextureRHI->GetTexture2D())
-	{
-		Texture2D->GetResource()->SetPrivateData(WKPDID_D3DDebugObjectName, FCString::Strlen(Name) + 1, TCHAR_TO_ANSI(Name));
-	}
-	else if (FD3D11TextureCube* TextureCube = (FD3D11TextureCube*)TextureRHI->GetTextureCube())
-	{
-		TextureCube->GetResource()->SetPrivateData(WKPDID_D3DDebugObjectName, FCString::Strlen(Name) + 1, TCHAR_TO_ANSI(Name));
-	}
-	else if (FD3D11Texture3D* Texture3D = (FD3D11Texture3D*)TextureRHI->GetTexture3D())
-	{
-		Texture3D->GetResource()->SetPrivateData(WKPDID_D3DDebugObjectName, FCString::Strlen(Name) + 1, TCHAR_TO_ANSI(Name));
-	}
+	Texture->GetResource()->SetPrivateData(WKPDID_D3DDebugObjectName, FCString::Strlen(Name) + 1, TCHAR_TO_ANSI(Name));
 #endif
 }
 
@@ -2179,8 +1710,7 @@ void FD3D11DynamicRHI::RHIVirtualTextureSetFirstMipVisible(FRHITexture2D* Textur
 {
 }
 
-template<typename BaseResourceType>
-TD3D11Texture2D<BaseResourceType>* FD3D11DynamicRHI::CreateTextureFromResource(bool bTextureArray, bool bCubeTexture, EPixelFormat Format, ETextureCreateFlags TexCreateFlags, const FClearValueBinding& ClearValueBinding, ID3D11Texture2D* TextureResource)
+FD3D11Texture* FD3D11DynamicRHI::CreateTextureFromResource(bool bTextureArray, bool bCubeTexture, EPixelFormat Format, ETextureCreateFlags TexCreateFlags, const FClearValueBinding& ClearValueBinding, ID3D11Texture2D* TextureResource)
 {
 	D3D11_TEXTURE2D_DESC TextureDesc;
 	TextureResource->GetDesc(&TextureDesc);
@@ -2402,157 +1932,128 @@ TD3D11Texture2D<BaseResourceType>* FD3D11DynamicRHI::CreateTextureFromResource(b
 		check(IsValidRef(ShaderResourceView));
 	}
 
-	TD3D11Texture2D<BaseResourceType>* Texture2D = new TD3D11Texture2D<BaseResourceType>(
-		this,
-		TextureResource,
-		ShaderResourceView,
-		bCreatedRTVPerSlice,
-		TextureDesc.ArraySize,
-		RenderTargetViews,
-		DepthStencilViews,
-		TextureDesc.Width,
-		TextureDesc.Height,
+	ETextureDimension Dimension = bTextureArray
+		? (bCubeTexture ? ETextureDimension::TextureCubeArray : ETextureDimension::Texture2DArray)
+		: (bCubeTexture ? ETextureDimension::TextureCube      : ETextureDimension::Texture2D     );
+
+	FRHITextureCreateDesc RHITextureDesc = FRHITextureCreateDesc(
+		Dimension,
+		TexCreateFlags,
+		(EPixelFormat)Format,
+		ClearValueBinding,
+		{ (int32)TextureDesc.Width, (int32)TextureDesc.Height },
+		1, // Depth
 		TextureDesc.ArraySize,
 		TextureDesc.MipLevels,
 		TextureDesc.SampleDesc.Count,
-		Format,
-		bCubeTexture,
-		TexCreateFlags,
-		/*bPooledTexture=*/ false,
-		ClearValueBinding
-		);
-
-	D3D11TextureAllocated(*Texture2D);
-
-	return Texture2D;
-}
-
-template <typename BaseResourceType>
-FTextureRHIRef FD3D11DynamicRHI::CreateAliasedD3D11Texture2D(TD3D11Texture2D<BaseResourceType>* SourceTexture)
-{
-	D3D11_TEXTURE2D_DESC TextureDesc;
-	SourceTexture->GetResource()->GetDesc(&TextureDesc);
-
-	const bool bSRGB = EnumHasAnyFlags(SourceTexture->Flags, TexCreate_SRGB);
-
-	const DXGI_FORMAT PlatformResourceFormat = TextureDesc.Format;
-	const DXGI_FORMAT PlatformShaderResourceFormat = FindShaderResourceDXGIFormat(PlatformResourceFormat, bSRGB);
-	const DXGI_FORMAT PlatformRenderTargetFormat = FindShaderResourceDXGIFormat(PlatformResourceFormat, bSRGB);
-
-	// Determine the MSAA settings to use for the texture.
-	D3D11_DSV_DIMENSION DepthStencilViewDimension = D3D11_DSV_DIMENSION_TEXTURE2D;
-	D3D11_RTV_DIMENSION RenderTargetViewDimension = D3D11_RTV_DIMENSION_TEXTURE2D;
-	D3D11_SRV_DIMENSION ShaderResourceViewDimension = D3D11_SRV_DIMENSION_TEXTURE2D;
-
-	if (TextureDesc.SampleDesc.Count > 1)
-	{
-		DepthStencilViewDimension = D3D11_DSV_DIMENSION_TEXTURE2DMS;
-		RenderTargetViewDimension = D3D11_RTV_DIMENSION_TEXTURE2DMS;
-		ShaderResourceViewDimension = D3D11_SRV_DIMENSION_TEXTURE2DMS;
-	}
-
-	TRefCountPtr<ID3D11ShaderResourceView> ShaderResourceView;
-	TArray<TRefCountPtr<ID3D11RenderTargetView> > RenderTargetViews;
-	TRefCountPtr<ID3D11DepthStencilView> DepthStencilViews[FExclusiveDepthStencil::MaxIndex];
-
-	bool bCreatedRTVPerSlice = false;
-	const bool bCubeTexture = SourceTexture->IsCubemap();
-	const bool bTextureArray = !bCubeTexture && TextureDesc.ArraySize > 1;
-
-	if (TextureDesc.BindFlags & D3D11_BIND_RENDER_TARGET)
-	{
-		// Create a render target view for each mip
-		for (uint32 MipIndex = 0; MipIndex < TextureDesc.MipLevels; MipIndex++)
-		{
-			// Just add null RTV entries (we'll be aliasing from source shortly).
-			if (EnumHasAnyFlags(SourceTexture->Flags, TexCreate_TargetArraySlicesIndependently) && (bTextureArray || bCubeTexture))
-			{
-				bCreatedRTVPerSlice = true;
-
-				for (uint32 SliceIndex = 0; SliceIndex < TextureDesc.ArraySize; SliceIndex++)
-				{
-					RenderTargetViews.Add(nullptr);
-				}
-			}
-			else
-			{
-				RenderTargetViews.Add(nullptr);
-			}
-		}
-	}
+		0, // ExtData
+		RHIGetDefaultResourceState(TexCreateFlags, false),
+		TEXT("FD3D11DynamicRHI::CreateTextureFromResource")
+	);
 
-	TD3D11Texture2D<BaseResourceType>* Texture2D = new TD3D11Texture2D<BaseResourceType>(
-		this,
-		nullptr,
-		nullptr,
-		bCreatedRTVPerSlice,
+	FD3D11Texture* Texture2D = new FD3D11Texture(
+		RHITextureDesc,
+		TextureResource,
+		ShaderResourceView,
 		TextureDesc.ArraySize,
+		bCreatedRTVPerSlice,
 		RenderTargetViews,
-		nullptr,
-		TextureDesc.Width,
-		TextureDesc.Height,
-		TextureDesc.ArraySize,
-		TextureDesc.MipLevels,
-		TextureDesc.SampleDesc.Count,
-		static_cast<BaseResourceType*>(SourceTexture)->GetFormat(),
-		bCubeTexture,
-		SourceTexture->Flags,
-		/*bPooledTexture=*/ false,
-		static_cast<BaseResourceType*>(SourceTexture)->GetClearBinding()
-		);
-
-	// We'll be the same size, since we're the same thing. Avoid the check in D3D11Resources.h (AliasResources).
-	Texture2D->SetMemorySize(SourceTexture->GetMemorySize());
-
-	Texture2D->AliasResources(SourceTexture);
+		DepthStencilViews
+	);
 
 	return Texture2D;
 }
 
-
 FTexture2DRHIRef FD3D11DynamicRHI::RHICreateTexture2DFromResource(EPixelFormat Format, ETextureCreateFlags TexCreateFlags, const FClearValueBinding& ClearValueBinding, ID3D11Texture2D* TextureResource)
 {
-	return CreateTextureFromResource<FD3D11BaseTexture2D>(false, false, Format, TexCreateFlags, ClearValueBinding, TextureResource);
+	return CreateTextureFromResource(false, false, Format, TexCreateFlags, ClearValueBinding, TextureResource);
 }
 
 FTexture2DArrayRHIRef FD3D11DynamicRHI::RHICreateTexture2DArrayFromResource(EPixelFormat Format, ETextureCreateFlags TexCreateFlags, const FClearValueBinding& ClearValueBinding, ID3D11Texture2D* TextureResource)
 {
-	return CreateTextureFromResource<FD3D11BaseTexture2DArray>(true, false, Format, TexCreateFlags, ClearValueBinding, TextureResource);
+	return CreateTextureFromResource(true, false, Format, TexCreateFlags, ClearValueBinding, TextureResource);
 }
 
 FTextureCubeRHIRef FD3D11DynamicRHI::RHICreateTextureCubeFromResource(EPixelFormat Format, ETextureCreateFlags TexCreateFlags, const FClearValueBinding& ClearValueBinding, ID3D11Texture2D* TextureResource)
 {
-	return CreateTextureFromResource<FD3D11BaseTextureCube>(false, true, Format, TexCreateFlags, ClearValueBinding, TextureResource);
+	return CreateTextureFromResource(false, true, Format, TexCreateFlags, ClearValueBinding, TextureResource);
+}
+
+FD3D11Texture::FD3D11Texture(
+	  const FRHITextureCreateDesc& InDesc
+	, ID3D11Resource* InResource
+	, ID3D11ShaderResourceView* InShaderResourceView
+	, int32 InRTVArraySize
+	, bool bInCreatedRTVsPerSlice
+	, TConstArrayView<TRefCountPtr<ID3D11RenderTargetView>> InRenderTargetViews
+	, TConstArrayView<TRefCountPtr<ID3D11DepthStencilView>> InDepthStencilViews
+	)
+	: FRHITexture         (InDesc)
+	, Resource            (InResource)
+	, ShaderResourceView  (InShaderResourceView)
+	, RenderTargetViews   (InRenderTargetViews)
+	, RTVArraySize        (InRTVArraySize)
+	, bCreatedRTVsPerSlice(bInCreatedRTVsPerSlice)
+{
+	// Set the DSVs for all the access type combinations
+	if (InDepthStencilViews.Num())
+	{
+		check(InDepthStencilViews.Num() == FExclusiveDepthStencil::MaxIndex);
+		for (uint32 Index = 0; Index < FExclusiveDepthStencil::MaxIndex; Index++)
+		{
+			DepthStencilViews[Index] = InDepthStencilViews[Index];
+		}
+	}
+
+	UpdateD3D11TextureStats(*this, true);
 }
 
-void FD3D11DynamicRHI::RHIAliasTextureResources(FTextureRHIRef& DestTextureRHI, FTextureRHIRef& SrcTextureRHI)
+FD3D11Texture::FD3D11Texture(FD3D11Texture const& Other, EAliasResourceParam)
+	: FRHITexture(FRHITextureCreateDesc(Other.GetDesc(), ERHIAccess::SRVMask, TEXT("FD3D11Texture::FD3D11Texture Alias")))
+	, bAlias(true)
 {
-	FD3D11TextureBase* DestTexture = GetD3D11TextureFromRHITexture(DestTextureRHI);
-	FD3D11TextureBase* SrcTexture = GetD3D11TextureFromRHITexture(SrcTextureRHI);
+	AliasResource(Other);
+}
 
-	if (DestTexture && SrcTexture)
+FD3D11Texture::~FD3D11Texture()
+{
+	if (!bAlias)
 	{
-		DestTexture->AliasResources(SrcTexture);
+		UpdateD3D11TextureStats(*this, false);
 	}
 }
 
-FTextureRHIRef FD3D11DynamicRHI::RHICreateAliasedTexture(FTextureRHIRef& SourceTexture)
+void FD3D11Texture::AliasResource(FD3D11Texture const& Other)
 {
-	if (SourceTexture->GetTexture2D() != nullptr)
-	{
-		return CreateAliasedD3D11Texture2D<FD3D11BaseTexture2D>(static_cast<FD3D11Texture2D*>(SourceTexture->GetTexture2D()));
-	}
-	else if (SourceTexture->GetTexture2DArray() != nullptr)
-	{
-		return CreateAliasedD3D11Texture2D<FD3D11BaseTexture2DArray>(static_cast<FD3D11Texture2DArray*>(SourceTexture->GetTexture2DArray()));
-	}
-	else if (SourceTexture->GetTextureCube() != nullptr)
+	check(bAlias);
+	IHVResourceHandle    = Other.IHVResourceHandle;
+	Resource             = Other.Resource;
+	ShaderResourceView   = Other.ShaderResourceView;
+	RenderTargetViews    = Other.RenderTargetViews;
+	bCreatedRTVsPerSlice = Other.bCreatedRTVsPerSlice;
+	RTVArraySize         = Other.RTVArraySize;
+
+	for (int32 Index = 0; Index < UE_ARRAY_COUNT(DepthStencilViews); ++Index)
 	{
-		return CreateAliasedD3D11Texture2D<FD3D11BaseTextureCube>(static_cast<FD3D11TextureCube*>(SourceTexture->GetTextureCube()));
+		DepthStencilViews[Index] = Other.DepthStencilViews[Index];
 	}
+}
+
+void FD3D11DynamicRHI::RHIAliasTextureResources(FTextureRHIRef& DstTextureRHI, FTextureRHIRef& SrcTextureRHI)
+{
+	FD3D11Texture* DstTexture = GetD3D11TextureFromRHITexture(DstTextureRHI);
+	FD3D11Texture* SrcTexture = GetD3D11TextureFromRHITexture(SrcTextureRHI);
+	check(DstTexture && SrcTexture);
+
+	DstTexture->AliasResource(*SrcTexture);
+}
+
+FTextureRHIRef FD3D11DynamicRHI::RHICreateAliasedTexture(FTextureRHIRef& SrcTextureRHI)
+{
+	FD3D11Texture* SrcTexture = GetD3D11TextureFromRHITexture(SrcTextureRHI);
+	check(SrcTexture);
 
-	UE_LOG(LogD3D11RHI, Error, TEXT("Currently FD3D11DynamicRHI::RHICreateAliasedTexture only supports 2D, 2D Array and Cube textures."));
-	return nullptr;
+	return new FD3D11Texture(*SrcTexture, FD3D11Texture::CreateAlias);
 }
 
 void FD3D11DynamicRHI::RHICopyTexture(FRHITexture* SourceTextureRHI, FRHITexture* DestTextureRHI, const FRHICopyTextureInfo& CopyInfo)
@@ -2565,8 +2066,8 @@ void FD3D11DynamicRHI::RHICopyTexture(FRHITexture* SourceTextureRHI, FRHITexture
 
 	FRHICommandList_RecursiveHazardous RHICmdList(this);	
 
-	FD3D11TextureBase* SourceTexture = GetD3D11TextureFromRHITexture(SourceTextureRHI);
-	FD3D11TextureBase* DestTexture = GetD3D11TextureFromRHITexture(DestTextureRHI);
+	FD3D11Texture* SourceTexture = GetD3D11TextureFromRHITexture(SourceTextureRHI);
+	FD3D11Texture* DestTexture = GetD3D11TextureFromRHITexture(DestTextureRHI);
 
 	check(SourceTexture && DestTexture);
 
diff --git a/Engine/Source/Runtime/Windows/D3D11RHI/Private/D3D11UAV.cpp b/Engine/Source/Runtime/Windows/D3D11RHI/Private/D3D11UAV.cpp
--- a/Engine/Source/Runtime/Windows/D3D11RHI/Private/D3D11UAV.cpp
+++ b/Engine/Source/Runtime/Windows/D3D11RHI/Private/D3D11UAV.cpp
@@ -76,33 +76,30 @@ FUnorderedAccessViewRHIRef FD3D11DynamicRHI::RHICreateUnorderedAccessView_Render
 
 FUnorderedAccessViewRHIRef FD3D11DynamicRHI::RHICreateUnorderedAccessView(FRHITexture* TextureRHI, uint32 MipLevel, uint16 FirstArraySlice, uint16 NumArraySlices)
 {
-	FD3D11TextureBase* Texture = GetD3D11TextureFromRHITexture(TextureRHI);
+	FD3D11Texture* Texture = GetD3D11TextureFromRHITexture(TextureRHI);
 	
 	D3D11_UNORDERED_ACCESS_VIEW_DESC UAVDesc{};
 
 	if (TextureRHI->GetTexture3D() != NULL)
 	{
-		FD3D11Texture3D* Texture3D = (FD3D11Texture3D*)Texture;
 		UAVDesc.ViewDimension = D3D11_UAV_DIMENSION_TEXTURE3D;
 		UAVDesc.Texture3D.MipSlice = MipLevel;
 		UAVDesc.Texture3D.FirstWSlice = 0;
-		UAVDesc.Texture3D.WSize = Texture3D->GetSizeZ() >> MipLevel;
+		UAVDesc.Texture3D.WSize = Texture->GetDesc().Depth >> MipLevel;
 	}
 	else if (TextureRHI->GetTexture2DArray() != NULL)
 	{
-		FD3D11Texture2DArray* Texture2DArray = (FD3D11Texture2DArray*)Texture;
 		UAVDesc.ViewDimension = D3D11_UAV_DIMENSION_TEXTURE2DARRAY;
 		UAVDesc.Texture2DArray.MipSlice = MipLevel;
 		UAVDesc.Texture2DArray.FirstArraySlice = NumArraySlices == 0 ? 0 : FirstArraySlice;
-		UAVDesc.Texture2DArray.ArraySize = NumArraySlices == 0 ? Texture2DArray->GetSizeZ() : NumArraySlices;
+		UAVDesc.Texture2DArray.ArraySize = NumArraySlices == 0 ? Texture->GetDesc().ArraySize : NumArraySlices;
 	}
 	else if (TextureRHI->GetTextureCube() != NULL)
 	{
-		FD3D11TextureCube* TextureCube = (FD3D11TextureCube*)Texture;
 		UAVDesc.ViewDimension = D3D11_UAV_DIMENSION_TEXTURE2DARRAY;
 		UAVDesc.Texture2DArray.MipSlice = MipLevel;
 		UAVDesc.Texture2DArray.FirstArraySlice = 0;
-		UAVDesc.Texture2DArray.ArraySize = TextureCube->GetSizeZ();
+		UAVDesc.Texture2DArray.ArraySize = Texture->GetSizeZ();
 	}
 	else
 	{
@@ -433,8 +430,8 @@ void FD3D11DynamicRHI::ClearUAV(TRHICommandList_RecursiveHazardous<FD3D11Dynamic
 	{
 		if (UAVDesc.ViewDimension == D3D11_UAV_DIMENSION_TEXTURE2D)
 		{
-			FD3D11Texture2D* Texture2D = static_cast<FD3D11Texture2D*>(UnorderedAccessView->Resource.GetReference());
-			FIntVector Size = Texture2D->GetSizeXYZ();
+			FD3D11Texture* Texture = static_cast<FD3D11Texture*>(UnorderedAccessView->Resource.GetReference());
+			FIntVector Size = Texture->GetSizeXYZ();
 
 			uint32 Width  = Size.X >> UAVDesc.Texture2D.MipSlice;
 			uint32 Height = Size.Y >> UAVDesc.Texture2D.MipSlice;
@@ -442,8 +439,8 @@ void FD3D11DynamicRHI::ClearUAV(TRHICommandList_RecursiveHazardous<FD3D11Dynamic
 		}
 		else if (UAVDesc.ViewDimension == D3D11_UAV_DIMENSION_TEXTURE2DARRAY)
 		{
-			FD3D11Texture2DArray* Texture2DArray = static_cast<FD3D11Texture2DArray*>(UnorderedAccessView->Resource.GetReference());
-			FIntVector Size = Texture2DArray->GetSizeXYZ();
+			FD3D11Texture* Texture = static_cast<FD3D11Texture*>(UnorderedAccessView->Resource.GetReference());
+			FIntVector Size = Texture->GetSizeXYZ();
 
 			uint32 Width = Size.X >> UAVDesc.Texture2DArray.MipSlice;
 			uint32 Height = Size.Y >> UAVDesc.Texture2DArray.MipSlice;
@@ -451,8 +448,8 @@ void FD3D11DynamicRHI::ClearUAV(TRHICommandList_RecursiveHazardous<FD3D11Dynamic
 		}
 		else if (UAVDesc.ViewDimension == D3D11_UAV_DIMENSION_TEXTURE3D)
 		{
-			FD3D11Texture3D* Texture3D = static_cast<FD3D11Texture3D*>(UnorderedAccessView->Resource.GetReference());
-			FIntVector Size = Texture3D->GetSizeXYZ();
+			FD3D11Texture* Texture = static_cast<FD3D11Texture*>(UnorderedAccessView->Resource.GetReference());
+			FIntVector Size = Texture->GetSizeXYZ();
 
 			// @todo - is WSize / mip index handling here correct?
 			uint32 Width = Size.X >> UAVDesc.Texture2DArray.MipSlice;
diff --git a/Engine/Source/Runtime/Windows/D3D11RHI/Private/D3D11Viewport.cpp b/Engine/Source/Runtime/Windows/D3D11RHI/Private/D3D11Viewport.cpp
--- a/Engine/Source/Runtime/Windows/D3D11RHI/Private/D3D11Viewport.cpp
+++ b/Engine/Source/Runtime/Windows/D3D11RHI/Private/D3D11Viewport.cpp
@@ -89,8 +89,6 @@ namespace RHIConsoleVariables
 		);
 };
 
-extern void D3D11TextureAllocated2D( FD3D11Texture2D& Texture );
-
 /**
  * Returns the current swap chain flags but with the same tearing policy used during construction.
  */
@@ -111,7 +109,7 @@ uint32 FD3D11Viewport::GetSwapChainFlags()
 /**
  * Creates a FD3D11Surface to represent a swap chain's back buffer.
  */
-FD3D11Texture2D* FD3D11Viewport::GetSwapChainSurface(FD3D11DynamicRHI* D3DRHI, EPixelFormat PixelFormat, uint32 SizeX, uint32 SizeY, IDXGISwapChain* SwapChain)
+FD3D11Texture* FD3D11Viewport::GetSwapChainSurface(FD3D11DynamicRHI* D3DRHI, EPixelFormat PixelFormat, uint32 SizeX, uint32 SizeY, IDXGISwapChain* SwapChain)
 {
 	// Grab the back buffer
 	TRefCountPtr<ID3D11Texture2D> BackBufferResource;
@@ -182,27 +180,30 @@ FD3D11Texture2D* FD3D11Viewport::GetSwapChainSurface(FD3D11DynamicRHI* D3DRHI, E
 	SRVDesc.Texture2D.MipLevels = 1;
 	VERIFYD3D11RESULT_EX(D3DRHI->GetDevice()->CreateShaderResourceView(BackBufferResource,&SRVDesc,BackBufferShaderResourceView.GetInitReference()), D3DRHI->GetDevice());
 
-	FD3D11Texture2D* NewTexture = new FD3D11Texture2D(
-		D3DRHI,
+	FRHITextureCreateDesc CreateDesc(
+		ETextureDimension::Texture2D,
+		ETextureCreateFlags::RenderTargetable,
+		PixelFormat,
+		FClearValueBinding(),
+		{ (int32)TextureDesc.Width, (int32)TextureDesc.Height },
+		1, // Depth
+		1, // ArraySize,
+		1, // NumMips,
+		1, // NumSamples
+		0, // ExtData
+		RHIGetDefaultResourceState(ETextureCreateFlags::RenderTargetable, false),
+		TEXT("FD3D11Viewport::GetSwapChainSurface")
+	);
+
+	FD3D11Texture* NewTexture = new FD3D11Texture(
+		CreateDesc,
 		BackBufferResource,
 		BackBufferShaderResourceView,
-		false,
-		1,
-		RenderTargetViews,
-		NULL,
-		TextureDesc.Width,
-		TextureDesc.Height,
-		1,
 		1,
-		1,
-		PixelFormat,
-		false,
-		TexCreate_RenderTargetable,
 		false,
-		FClearValueBinding()
-		);
-
-	D3D11TextureAllocated2D(*NewTexture);
+		RenderTargetViews,
+		{}
+	);
 
 	return NewTexture;
 }
diff --git a/Engine/Source/Runtime/Windows/D3D11RHI/Private/HoloLens/HoloLensD3D11Viewport.cpp b/Engine/Source/Runtime/Windows/D3D11RHI/Private/HoloLens/HoloLensD3D11Viewport.cpp
--- a/Engine/Source/Runtime/Windows/D3D11RHI/Private/HoloLens/HoloLensD3D11Viewport.cpp
+++ b/Engine/Source/Runtime/Windows/D3D11RHI/Private/HoloLens/HoloLensD3D11Viewport.cpp
@@ -37,8 +37,6 @@ DXGI_FORMAT GetSupportedSwapChainBufferFormat(DXGI_FORMAT InPreferredDXGIFormat)
 //	return DXGI_FORMAT_B8G8R8A8_UNORM;
 //}
 
-extern void D3D11TextureAllocated2D( FD3D11Texture2D& Texture );
-
 FD3D11Viewport::FD3D11Viewport(FD3D11DynamicRHI* InD3DRHI,HWND InWindowHandle,uint32 InSizeX,uint32 InSizeY,bool bInIsFullscreen, EPixelFormat InPreferredPixelFormat):
 	D3DRHI(InD3DRHI),
 	LastFlipTime(0),
diff --git a/Engine/Source/Runtime/Windows/D3D11RHI/Public/D3D11Resources.h b/Engine/Source/Runtime/Windows/D3D11RHI/Public/D3D11Resources.h
--- a/Engine/Source/Runtime/Windows/D3D11RHI/Public/D3D11Resources.h
+++ b/Engine/Source/Runtime/Windows/D3D11RHI/Public/D3D11Resources.h
@@ -139,62 +139,89 @@ public:
 };
 
 /** Texture base class. */
-class D3D11RHI_API FD3D11TextureBase : public FD3D11BaseShaderResource
+class D3D11RHI_API FD3D11Texture final : public FRHITexture, public FD3D11BaseShaderResource
 {
 public:
+	virtual uint32 AddRef     () const override	{ return FRHITexture::AddRef();      }
+	virtual uint32 Release    () const override	{ return FRHITexture::Release();     }
+	virtual uint32 GetRefCount() const override	{ return FRHITexture::GetRefCount(); }
 
-	FD3D11TextureBase(
-		class FD3D11DynamicRHI* InD3DRHI,
+	explicit FD3D11Texture(
+		const FRHITextureCreateDesc& InDesc,
 		ID3D11Resource* InResource,
 		ID3D11ShaderResourceView* InShaderResourceView,
 		int32 InRTVArraySize,
 		bool bInCreatedRTVsPerSlice,
-		const TArray<TRefCountPtr<ID3D11RenderTargetView> >& InRenderTargetViews,
-		TRefCountPtr<ID3D11DepthStencilView>* InDepthStencilViews
-		) 
-	: D3DRHI(InD3DRHI)
-	, IHVResourceHandle(nullptr)
-	, MemorySize(0)
-	, BaseShaderResource(this)
-	, Resource(InResource)
-	, ShaderResourceView(InShaderResourceView)
-	, RenderTargetViews(InRenderTargetViews)
-	, bCreatedRTVsPerSlice(bInCreatedRTVsPerSlice)
-	, RTVArraySize(InRTVArraySize)
+		TConstArrayView<TRefCountPtr<ID3D11RenderTargetView>> InRenderTargetViews,
+		TConstArrayView<TRefCountPtr<ID3D11DepthStencilView>> InDepthStencilViews
+	);
+
+	enum EAliasResourceParam { CreateAlias };
+	explicit FD3D11Texture(FD3D11Texture const& Other, EAliasResourceParam);
+	void AliasResource(FD3D11Texture const& Other);
+
+	virtual ~FD3D11Texture();
+
+	inline uint64 GetMemorySize() const
 	{
-		// Set the DSVs for all the access type combinations
-		if ( InDepthStencilViews != nullptr )
-		{
-			for (uint32 Index = 0; Index < FExclusiveDepthStencil::MaxIndex; Index++)
-			{
-				DepthStencilViews[Index] = InDepthStencilViews[Index];
-			}
-		}
+		return RHICalcTexturePlatformSize(GetDesc()).Size;
 	}
 
-	virtual ~FD3D11TextureBase() {}
+	// Accessors.
+	inline ID3D11Resource* GetResource() const { return Resource; }
+	inline ID3D11ShaderResourceView* GetShaderResourceView() const { return ShaderResourceView; }
 
-	int32 GetMemorySize() const
+	inline bool IsCubemap() const
 	{
-		return MemorySize;
+		FRHITextureDesc const& Desc = GetDesc();
+		return Desc.Dimension == ETextureDimension::TextureCube || Desc.Dimension == ETextureDimension::TextureCubeArray;
 	}
 
-	void SetMemorySize( int32 InMemorySize )
+	inline ID3D11Texture2D* GetD3D11Texture2D() const
 	{
-		MemorySize = InMemorySize;
+		check(Resource);
+		check(GetDesc().Dimension == ETextureDimension::Texture2D
+		   || GetDesc().Dimension == ETextureDimension::Texture2DArray
+		   || GetDesc().Dimension == ETextureDimension::TextureCube
+		   || GetDesc().Dimension == ETextureDimension::TextureCubeArray);
+
+		return static_cast<ID3D11Texture2D*>(Resource.GetReference());
 	}
 
-	// Accessors.
-	ID3D11Resource* GetResource() const { return Resource; }
-	ID3D11ShaderResourceView* GetShaderResourceView() const { return ShaderResourceView; }
-	FD3D11BaseShaderResource* GetBaseShaderResource() const { return BaseShaderResource; }
+	inline ID3D11Texture3D* GetD3D11Texture3D() const
+	{
+		check(Resource);
+		check(GetDesc().Dimension == ETextureDimension::Texture3D);
+
+		return static_cast<ID3D11Texture3D*>(Resource.GetReference());
+	}
+
+	inline bool IsTexture3D() const
+	{
+		return GetDesc().Dimension == ETextureDimension::Texture3D;
+	}
+
+	virtual inline void* GetNativeResource() const override
+	{
+		return GetResource();
+	}
+
+	virtual inline void* GetNativeShaderResourceView() const override
+	{
+		return GetShaderResourceView();
+	}
+
+	virtual inline void* GetTextureBaseRHI() override
+	{
+		return this;
+	}
 
-	void SetIHVResourceHandle(void* InHandle)
+	inline void SetIHVResourceHandle(void* InHandle)
 	{
 		IHVResourceHandle = InHandle;
 	}
 
-	void* GetIHVResourceHandle() const
+	inline void* GetIHVResourceHandle() const
 	{
 		return IHVResourceHandle;
 	}
@@ -203,7 +230,7 @@ public:
 	 * Get the render target view for the specified mip and array slice.
 	 * An array slice of -1 is used to indicate that no array slice should be required. 
 	 */
-	ID3D11RenderTargetView* GetRenderTargetView(int32 MipIndex, int32 ArraySliceIndex) const
+	inline ID3D11RenderTargetView* GetRenderTargetView(int32 MipIndex, int32 ArraySliceIndex) const
 	{
 		int32 ArrayIndex = MipIndex;
 
@@ -224,42 +251,35 @@ public:
 		}
 		return 0;
 	}
-	ID3D11DepthStencilView* GetDepthStencilView(FExclusiveDepthStencil AccessType) const
+
+	inline ID3D11DepthStencilView* GetDepthStencilView(FExclusiveDepthStencil AccessType) const
 	{ 
 		return DepthStencilViews[AccessType.GetIndex()]; 
 	}
 
-	void AliasResources(FD3D11TextureBase* Texture)
+#if RHI_ENABLE_RESOURCE_INFO
+	virtual bool GetResourceInfo(FRHIResourceInfo& OutResourceInfo) const override
 	{
-		check(MemorySize == Texture->MemorySize);
-		check(bCreatedRTVsPerSlice == Texture->bCreatedRTVsPerSlice);
-		check(RTVArraySize == Texture->RTVArraySize);
-
-		// Do not copy the BaseShaderResource from the source texture (this is initialized correctly here, and is used for
-		// state caching logic).
-		Resource = Texture->Resource;
-		ShaderResourceView = Texture->ShaderResourceView;
-		RenderTargetViews = Texture->RenderTargetViews;
-
-		for (uint32 Index = 0; Index < FExclusiveDepthStencil::MaxIndex; Index++)
-		{
-			DepthStencilViews[Index] = Texture->DepthStencilViews[Index];
-		}
+		OutResourceInfo = FRHIResourceInfo{};
+		OutResourceInfo.Name = GetName();
+		OutResourceInfo.Type = GetType();
+		OutResourceInfo.VRamAllocation.AllocationSize = GetMemorySize();
+		return true;
 	}
+#endif
 
-protected:
+	/**
+	* Locks one of the texture's mip-maps.
+	* @return A pointer to the specified texture data.
+	*/
+	void* Lock(class FD3D11DynamicRHI* D3DRHI, uint32 MipIndex, uint32 ArrayIndex, EResourceLockMode LockMode, uint32& DestStride, bool bForceLockDeferred = false);
 
-	/** The D3D11 RHI that created this texture. */
-	FD3D11DynamicRHI* D3DRHI;
+	/** Unlocks a previously locked mip-map. */
+	void Unlock(class FD3D11DynamicRHI* D3DRHI, uint32 MipIndex, uint32 ArrayIndex);
 
+private:
 	//Resource handle for use by IHVs for SLI and other purposes.
-	void* IHVResourceHandle;
-
-	/** Amount of memory allocated by this texture, in bytes. */
-	int32 MemorySize;
-
-	/** Pointer to the base shader resource. Usually the object itself, but not for texture references. */
-	FD3D11BaseShaderResource* BaseShaderResource;
+	void* IHVResourceHandle = nullptr;
 
 	/** The texture resource. */
 	TRefCountPtr<ID3D11Resource> Resource;
@@ -270,261 +290,24 @@ protected:
 	/** A render targetable view of the texture. */
 	TArray<TRefCountPtr<ID3D11RenderTargetView> > RenderTargetViews;
 
-	bool bCreatedRTVsPerSlice;
-
-	int32 RTVArraySize;
-
 	/** A depth-stencil targetable view of the texture. */
 	TRefCountPtr<ID3D11DepthStencilView> DepthStencilViews[FExclusiveDepthStencil::MaxIndex];
-};
-
-/** 2D texture (vanilla, cubemap or 2D array) */
-template<typename BaseResourceType>
-class D3D11RHI_API TD3D11Texture2D : public BaseResourceType, public FD3D11TextureBase
-{
-public:
-
-	/** Flags used when the texture was created */
-	ETextureCreateFlags Flags;
-
-	/** Initialization constructor. */
-	TD3D11Texture2D(
-		class FD3D11DynamicRHI* InD3DRHI,
-		ID3D11Texture2D* InResource,
-		ID3D11ShaderResourceView* InShaderResourceView,
-		bool bInCreatedRTVsPerSlice,
-		int32 InRTVArraySize,
-		const TArray<TRefCountPtr<ID3D11RenderTargetView> >& InRenderTargetViews,
-		TRefCountPtr<ID3D11DepthStencilView>* InDepthStencilViews,
-		uint32 InSizeX,
-		uint32 InSizeY,
-		uint32 InSizeZ,
-		uint32 InNumMips,
-		uint32 InNumSamples,
-		EPixelFormat InFormat,
-		bool bInCubemap,
-		ETextureCreateFlags InFlags,
-		bool bInPooled,
-		const FClearValueBinding& InClearValue
-#if PLATFORM_SUPPORTS_VIRTUAL_TEXTURES
-		, void* InRawTextureMemory = nullptr
-#endif
-		)
-	: BaseResourceType(
-		InSizeX,
-		InSizeY,
-		InSizeZ,
-		InNumMips,
-		InNumSamples,
-		InFormat,
-		InFlags,
-		InClearValue
-		)
-	, FD3D11TextureBase(
-		InD3DRHI,
-		InResource,
-		InShaderResourceView, 
-		InRTVArraySize,
-		bInCreatedRTVsPerSlice,
-		InRenderTargetViews,
-		InDepthStencilViews
-		)
-	, Flags(InFlags)
-	, bCubemap(bInCubemap)
-	, bPooled(bInPooled)
-#if PLATFORM_SUPPORTS_VIRTUAL_TEXTURES
-	, RawTextureMemory(InRawTextureMemory)
-#endif
-	{
-	}
-
-	virtual ~TD3D11Texture2D();
-
-	// FRHIResource overrides
-#if RHI_ENABLE_RESOURCE_INFO
-	bool GetResourceInfo(FRHIResourceInfo& OutResourceInfo) const override
-	{
-		OutResourceInfo = FRHIResourceInfo{};
-		OutResourceInfo.Name = this->GetName();
-		OutResourceInfo.Type = this->GetType();
-		OutResourceInfo.VRamAllocation.AllocationSize = GetMemorySize();
-		return true;
-	}
-#endif
-
-	/**
-	 * Locks one of the texture's mip-maps.
-	 * @return A pointer to the specified texture data.
-	 */
-	void* Lock(uint32 MipIndex,uint32 ArrayIndex,EResourceLockMode LockMode,uint32& DestStride,bool bForceLockDeferred = false);
-
-	/** Unlocks a previously locked mip-map. */
-	void Unlock(uint32 MipIndex,uint32 ArrayIndex);
-
-	// Accessors.
-	ID3D11Texture2D* GetResource() const { return (ID3D11Texture2D*)FD3D11TextureBase::GetResource(); }
-	bool IsCubemap() const { return bCubemap; }
-
-	/** FRHITexture override.  See FRHITexture::GetNativeResource() */
-	virtual void* GetNativeResource() const override final
-	{ 
-		return GetResource();
-	}
-	virtual void* GetNativeShaderResourceView() const override final
-	{
-		return GetShaderResourceView();
-	}
-	virtual void* GetTextureBaseRHI() override final
-	{
-		return static_cast<FD3D11TextureBase*>(this);
-	}
-
 
-	// IRefCountedObject interface.
-	virtual uint32 AddRef() const
-	{
-		return FRHIResource::AddRef();
-	}
-	virtual uint32 Release() const
-	{
-		return FRHIResource::Release();
-	}
-	virtual uint32 GetRefCount() const
-	{
-		return FRHIResource::GetRefCount();
-	}
-#if PLATFORM_SUPPORTS_VIRTUAL_TEXTURES
-	void* GetRawTextureMemory() const
-	{
-		return RawTextureMemory;
-	}
-#endif
-
-private:
-
-	/** Whether the texture is a cube-map. */
-	const uint32 bCubemap : 1;
-	/** Whether the texture can be pooled. */
-	const uint32 bPooled : 1;
-#if PLATFORM_SUPPORTS_VIRTUAL_TEXTURES
-	void* RawTextureMemory;
-#endif
-};
-
-/** 3D Texture */
-class FD3D11Texture3D : public FRHITexture3D, public FD3D11TextureBase
-{
-public:
-
-	/** Initialization constructor. */
-	FD3D11Texture3D(
-		class FD3D11DynamicRHI* InD3DRHI,
-		ID3D11Texture3D* InResource,
-		ID3D11ShaderResourceView* InShaderResourceView,
-		const TArray<TRefCountPtr<ID3D11RenderTargetView> >& InRenderTargetViews,
-		uint32 InSizeX,
-		uint32 InSizeY,
-		uint32 InSizeZ,
-		uint32 InNumMips,
-		EPixelFormat InFormat,
-		ETextureCreateFlags InFlags,
-		const FClearValueBinding& InClearValue
-		)
-	: FRHITexture3D(InSizeX,InSizeY,InSizeZ,InNumMips,InFormat,InFlags,InClearValue)
-	, FD3D11TextureBase(
-		InD3DRHI,
-		InResource,
-		InShaderResourceView,
-		1,
-		false,
-		InRenderTargetViews,
-		NULL
-		)
-	{
-	}
-
-	virtual ~FD3D11Texture3D();
-
-	// FRHIResource overrides
-#if RHI_ENABLE_RESOURCE_INFO
-	bool GetResourceInfo(FRHIResourceInfo& OutResourceInfo) const override
-	{
-		OutResourceInfo = FRHIResourceInfo{};
-		OutResourceInfo.Name = GetName();
-		OutResourceInfo.Type = GetType();
-		OutResourceInfo.VRamAllocation.AllocationSize = GetMemorySize();
-		return true;
-	}
-#endif
-	
-	// Accessors.
-	ID3D11Texture3D* GetResource() const { return (ID3D11Texture3D*)FD3D11TextureBase::GetResource(); }
-	
-	virtual void* GetTextureBaseRHI() override final
-	{
-		return static_cast<FD3D11TextureBase*>(this);
-	}
-
-	// IRefCountedObject interface.
-	virtual uint32 AddRef() const
-	{
-		return FRHIResource::AddRef();
-	}
-	virtual uint32 Release() const
-	{
-		return FRHIResource::Release();
-	}
-	virtual uint32 GetRefCount() const
-	{
-		return FRHIResource::GetRefCount();
-	}
-};
-
-class FD3D11BaseTexture2D : public FRHITexture2D
-{
-public:
-	FD3D11BaseTexture2D(uint32 InSizeX, uint32 InSizeY, uint32 InSizeZ, uint32 InNumMips, uint32 InNumSamples, EPixelFormat InFormat, ETextureCreateFlags InFlags, const FClearValueBinding& InClearValue)
-	: FRHITexture2D(InSizeX,InSizeY,InNumMips,InNumSamples,InFormat,InFlags, InClearValue)
-	{}
-	uint32 GetSizeZ() const { return 0; }
-};
-
-class FD3D11BaseTexture2DArray : public FRHITexture2DArray
-{
-public:
-	FD3D11BaseTexture2DArray(uint32 InSizeX, uint32 InSizeY, uint32 InSizeZ, uint32 InNumMips, uint32 InNumSamples, EPixelFormat InFormat, ETextureCreateFlags InFlags, const FClearValueBinding& InClearValue)
-	: FRHITexture2DArray(InSizeX,InSizeY,InSizeZ,InNumMips,InNumSamples, InFormat,InFlags,InClearValue)
-	{}
-};
-
-class FD3D11BaseTextureCube : public FRHITextureCube
-{
-public:
-	FD3D11BaseTextureCube(uint32 InSizeX, uint32 InSizeY, uint32 InSizeZ, uint32 InNumMips, uint32 InNumSamples, EPixelFormat InFormat, ETextureCreateFlags InFlags, const FClearValueBinding& InClearValue)
-	: FRHITextureCube(InSizeX,InNumMips,InFormat,InFlags,InClearValue)
-	, SliceCount(InSizeZ)
-	{ check(InNumSamples == 1); }
-	uint32 GetSizeX() const { return GetSize(); }
-	uint32 GetSizeY() const { return GetSize(); } //-V524
-	uint32 GetSizeZ() const { return SliceCount; }
+	int32 RTVArraySize;
 
-private:
-	uint32 SliceCount;
+	uint8 bCreatedRTVsPerSlice : 1;
+	uint8 bAlias : 1;
 };
 
-typedef TD3D11Texture2D<FRHITexture>              FD3D11Texture;
-typedef TD3D11Texture2D<FD3D11BaseTexture2D>      FD3D11Texture2D;
-typedef TD3D11Texture2D<FD3D11BaseTexture2DArray> FD3D11Texture2DArray;
-typedef TD3D11Texture2D<FD3D11BaseTextureCube>    FD3D11TextureCube;
-
-/** Given a pointer to a RHI texture that was created by the D3D11 RHI, returns a pointer to the FD3D11TextureBase it encapsulates. */
-FORCEINLINE FD3D11TextureBase* GetD3D11TextureFromRHITexture(FRHITexture* Texture)
+/** Given a pointer to a RHI texture that was created by the D3D11 RHI, returns a pointer to the FD3D11Texture it encapsulates. */
+FORCEINLINE FD3D11Texture* GetD3D11TextureFromRHITexture(FRHITexture* Texture)
 {
 	if (!Texture)
 	{
-		return NULL;
+		return nullptr;
 	}
-	FD3D11TextureBase* Result((FD3D11TextureBase*)Texture->GetTextureBaseRHI());
+
+	FD3D11Texture* Result = static_cast<FD3D11Texture*>(Texture->GetTextureBaseRHI());
 	check(Result);
 	return Result;
 }
@@ -722,9 +505,6 @@ public:
 	{}
 };
 
-void ReturnPooledTexture2D(int32 MipCount, EPixelFormat PixelFormat, ID3D11Texture2D* InResource);
-void ReleasePooledTextures();
-
 template<class T>
 struct TD3D11ResourceTraits
 {
@@ -760,31 +540,11 @@ struct TD3D11ResourceTraits<FRHIBoundShaderState>
 	typedef FD3D11BoundShaderState TConcreteType;
 };
 template<>
-struct TD3D11ResourceTraits<FRHITexture3D>
-{
-	typedef FD3D11Texture3D TConcreteType;
-};
-template<>
 struct TD3D11ResourceTraits<FRHITexture>
 {
 	typedef FD3D11Texture TConcreteType;
 };
 template<>
-struct TD3D11ResourceTraits<FRHITexture2D>
-{
-	typedef FD3D11Texture2D TConcreteType;
-};
-template<>
-struct TD3D11ResourceTraits<FRHITexture2DArray>
-{
-	typedef FD3D11Texture2DArray TConcreteType;
-};
-template<>
-struct TD3D11ResourceTraits<FRHITextureCube>
-{
-	typedef FD3D11TextureCube TConcreteType;
-};
-template<>
 struct TD3D11ResourceTraits<FRHIRenderQuery>
 {
 	typedef FD3D11RenderQuery TConcreteType;
