Copying //UE5/Dev-ParallelRendering to //UE5/Main

diff --git a/Engine/Source/Runtime/Windows/D3D11RHI/Public/D3D11Util.h b/Engine/Source/Runtime/Windows/D3D11RHI/Public/D3D11Util.h
--- a/Engine/Source/Runtime/Windows/D3D11RHI/Public/D3D11Util.h
+++ b/Engine/Source/Runtime/Windows/D3D11RHI/Public/D3D11Util.h
@@ -120,43 +120,40 @@ FORCEINLINE uint32 GetD3D11CubeFace(ECubeFace Face)
 class FD3D11LockedKey
 {
 public:
-	void* SourceObject;
-	uint32 Subresource;
+	ID3D11Resource* SourceObject = nullptr;
+	uint32 Subresource = 0;
 
-public:
-	FD3D11LockedKey() : SourceObject(NULL)
-		, Subresource(0)
-	{}
-	FD3D11LockedKey(ID3D11Texture2D* source, uint32 subres=0) : SourceObject((void*)source)
-		, Subresource(subres)
-	{}
-	FD3D11LockedKey(ID3D11Texture3D* source, uint32 subres=0) : SourceObject((void*)source)
-		, Subresource(subres)
-	{}
-	FD3D11LockedKey(ID3D11Buffer* source, uint32 subres=0) : SourceObject((void*)source)
-		, Subresource(subres)
+	FD3D11LockedKey() = default;
+
+	FD3D11LockedKey(ID3D11Resource* InSource, uint32 InSubresource = 0)
+		: SourceObject(InSource)
+		, Subresource (InSubresource)
 	{}
-	bool operator==( const FD3D11LockedKey& Other ) const
+
+	bool operator == (const FD3D11LockedKey& Other) const
 	{
 		return SourceObject == Other.SourceObject && Subresource == Other.Subresource;
 	}
-	bool operator!=( const FD3D11LockedKey& Other ) const
+
+	bool operator != (const FD3D11LockedKey& Other) const
 	{
 		return SourceObject != Other.SourceObject || Subresource != Other.Subresource;
 	}
-	FD3D11LockedKey& operator=( const FD3D11LockedKey& Other )
+
+	FD3D11LockedKey& operator = (const FD3D11LockedKey& Other)
 	{
 		SourceObject = Other.SourceObject;
 		Subresource = Other.Subresource;
 		return *this;
 	}
+
 	uint32 GetHash() const
 	{
-		return PointerHash( SourceObject );
+		return PointerHash(SourceObject);
 	}
 
 	/** Hashing function. */
-	friend uint32 GetTypeHash( const FD3D11LockedKey& K )
+	friend uint32 GetTypeHash(const FD3D11LockedKey& K)
 	{
 		return K.GetHash();
 	}
diff --git a/Engine/Source/Runtime/Windows/D3D11RHI/Public/D3D11Viewport.h b/Engine/Source/Runtime/Windows/D3D11RHI/Public/D3D11Viewport.h
--- a/Engine/Source/Runtime/Windows/D3D11RHI/Public/D3D11Viewport.h
+++ b/Engine/Source/Runtime/Windows/D3D11RHI/Public/D3D11Viewport.h
@@ -71,7 +71,7 @@ public:
 
 	// Accessors.
 	FIntPoint GetSizeXY() const { return FIntPoint(SizeX, SizeY); }
-	FD3D11Texture2D* GetBackBuffer() const { return BackBuffer; }
+	FD3D11Texture* GetBackBuffer() const { return BackBuffer; }
 	EColorSpaceAndEOTF GetPixelColorSpace() const { return PixelColorSpace; }
 
 	virtual void WaitForFrameEventCompletion() override
@@ -91,7 +91,7 @@ public:
 #endif
 
 	virtual void* GetNativeSwapChain() const override { return GetSwapChain(); }
-	virtual void* GetNativeBackBufferTexture() const override { return GetBackBuffer()->GetResource(); }
+	virtual void* GetNativeBackBufferTexture() const override { return GetBackBuffer()->GetD3D11Texture2D(); }
 	virtual void* GetNativeBackBufferRT() const override { return GetBackBuffer()->GetRenderTargetView(0, 0); }
 
 	virtual void SetCustomPresent(FRHICustomPresent* InCustomPresent) override
@@ -101,7 +101,7 @@ public:
 	virtual FRHICustomPresent* GetCustomPresent() const { return CustomPresent; }
 
 	virtual void* GetNativeWindow(void** AddParam = nullptr) const override { return (void*)WindowHandle; }
-	static FD3D11Texture2D* GetSwapChainSurface(FD3D11DynamicRHI* D3DRHI, EPixelFormat PixelFormat, uint32 SizeX, uint32 SizeY, IDXGISwapChain* SwapChain);
+	static FD3D11Texture* GetSwapChainSurface(FD3D11DynamicRHI* D3DRHI, EPixelFormat PixelFormat, uint32 SizeX, uint32 SizeY, IDXGISwapChain* SwapChain);
 
 	static DXGI_FORMAT GetRenderTargetFormat(EPixelFormat PixelFormat)
 	{
@@ -159,7 +159,7 @@ protected:
 #else
 	TRefCountPtr<IDXGISwapChain> SwapChain;
 #endif
-	TRefCountPtr<FD3D11Texture2D> BackBuffer;
+	TRefCountPtr<FD3D11Texture> BackBuffer;
 
 	// Support for selecting non-default output for display in fullscreen exclusive
 	TRefCountPtr<IDXGIOutput>	ForcedFullscreenOutput;
