Cleaning up dxgi includes in D3D11RHI

diff --git a/Engine/Source/Runtime/Windows/D3D11RHI/Private/D3D11Device.cpp b/Engine/Source/Runtime/Windows/D3D11RHI/Private/D3D11Device.cpp
--- a/Engine/Source/Runtime/Windows/D3D11RHI/Private/D3D11Device.cpp
+++ b/Engine/Source/Runtime/Windows/D3D11RHI/Private/D3D11Device.cpp
@@ -15,9 +15,6 @@
 	#endif
 #include "Windows/HideWindowsPlatformTypes.h"
 
-#include "dxgi1_5.h"
-
-
 bool D3D11RHI_ShouldCreateWithD3DDebug()
 {
 	// Use a debug device if specified on the command line.
diff --git a/Engine/Source/Runtime/Windows/D3D11RHI/Private/D3D11Util.cpp b/Engine/Source/Runtime/Windows/D3D11RHI/Private/D3D11Util.cpp
--- a/Engine/Source/Runtime/Windows/D3D11RHI/Private/D3D11Util.cpp
+++ b/Engine/Source/Runtime/Windows/D3D11RHI/Private/D3D11Util.cpp
@@ -8,10 +8,6 @@
 #include "EngineModule.h"
 #include "RendererInterface.h"
 
-THIRD_PARTY_INCLUDES_START
-#include "dxgi1_4.h"
-THIRD_PARTY_INCLUDES_END
-
 #define D3DERR(x) case x: ErrorCodeText = TEXT(#x); break;
 #define LOCTEXT_NAMESPACE "Developer.MessageLog"
 
diff --git a/Engine/Source/Runtime/Windows/D3D11RHI/Private/HoloLens/HoloLensD3D11Viewport.cpp b/Engine/Source/Runtime/Windows/D3D11RHI/Private/HoloLens/HoloLensD3D11Viewport.cpp
--- a/Engine/Source/Runtime/Windows/D3D11RHI/Private/HoloLens/HoloLensD3D11Viewport.cpp
+++ b/Engine/Source/Runtime/Windows/D3D11RHI/Private/HoloLens/HoloLensD3D11Viewport.cpp
@@ -5,7 +5,6 @@
 
 #include "AllowWindowsPlatformTypes.h"
 	#include <DirectXMath.h>
-	#include <dxgi1_2.h>
 #include "HideWindowsPlatformTypes.h"
 
 #pragma warning(disable : 4946)	// reinterpret_cast used between related classes: 'Platform::Object' and ...
diff --git a/Engine/Source/Runtime/Windows/D3D11RHI/Private/Windows/WindowsD3D11Device.cpp b/Engine/Source/Runtime/Windows/D3D11RHI/Private/Windows/WindowsD3D11Device.cpp
--- a/Engine/Source/Runtime/Windows/D3D11RHI/Private/Windows/WindowsD3D11Device.cpp
+++ b/Engine/Source/Runtime/Windows/D3D11RHI/Private/Windows/WindowsD3D11Device.cpp
@@ -21,11 +21,6 @@
 #include "Runtime/HeadMountedDisplay/Public/IHeadMountedDisplayModule.h"
 #include "GenericPlatform/GenericPlatformDriver.h"			// FGPUDriverInfo
 #include "GenericPlatform/GenericPlatformCrashContext.h"
-THIRD_PARTY_INCLUDES_START
-#include "dxgi1_3.h"
-#include "dxgi1_4.h"
-#include "dxgi1_6.h"
-THIRD_PARTY_INCLUDES_END
 #include "RHIValidation.h"
 #include "HAL/ExceptionHandling.h"
 
diff --git a/Engine/Source/Runtime/Windows/D3D11RHI/Private/Windows/WindowsD3D11Viewport.cpp b/Engine/Source/Runtime/Windows/D3D11RHI/Private/Windows/WindowsD3D11Viewport.cpp
--- a/Engine/Source/Runtime/Windows/D3D11RHI/Private/Windows/WindowsD3D11Viewport.cpp
+++ b/Engine/Source/Runtime/Windows/D3D11RHI/Private/Windows/WindowsD3D11Viewport.cpp
@@ -11,10 +11,6 @@
 #include "Windows/AllowWindowsPlatformTypes.h"
 #include <dwmapi.h>
 
-THIRD_PARTY_INCLUDES_START
-#include "dxgi1_6.h"
-THIRD_PARTY_INCLUDES_END
-
 #ifndef DXGI_SWAP_CHAIN_FLAG_ALLOW_TEARING
 #define DXGI_SWAP_CHAIN_FLAG_ALLOW_TEARING  2048
 #endif
diff --git a/Engine/Source/Runtime/Windows/D3D11RHI/Public/D3D11Viewport.h b/Engine/Source/Runtime/Windows/D3D11RHI/Public/D3D11Viewport.h
--- a/Engine/Source/Runtime/Windows/D3D11RHI/Public/D3D11Viewport.h
+++ b/Engine/Source/Runtime/Windows/D3D11RHI/Public/D3D11Viewport.h
@@ -10,12 +10,6 @@
 #include "RenderResource.h"
 #include "RenderUtils.h"
 
-#if PLATFORM_HOLOLENS
-#include "AllowWindowsPlatformTypes.h"
-#include <dxgi1_2.h>
-#include "HideWindowsPlatformTypes.h"
-#endif
-
 /** A D3D event query resource. */
 class D3D11RHI_API FD3D11EventQuery
 {
diff --git a/Engine/Source/Runtime/Windows/D3D11RHI/Public/HoloLens/D3D11ThirdParty.h b/Engine/Source/Runtime/Windows/D3D11RHI/Public/HoloLens/D3D11ThirdParty.h
--- a/Engine/Source/Runtime/Windows/D3D11RHI/Public/HoloLens/D3D11ThirdParty.h
+++ b/Engine/Source/Runtime/Windows/D3D11RHI/Public/HoloLens/D3D11ThirdParty.h
@@ -15,6 +15,7 @@
 #include "Microsoft/AllowMicrosoftPlatformTypes.h"
 THIRD_PARTY_INCLUDES_START
 	#include <d3d11_2.h>
+	#include <dxgi1_6.h>
 	#include <dxgidebug.h>
 THIRD_PARTY_INCLUDES_END
 #include "Microsoft/HideMicrosoftPlatformTypes.h"
diff --git a/Engine/Source/Runtime/Windows/D3D11RHI/Public/ID3D11DynamicRHI.h b/Engine/Source/Runtime/Windows/D3D11RHI/Public/ID3D11DynamicRHI.h
--- a/Engine/Source/Runtime/Windows/D3D11RHI/Public/ID3D11DynamicRHI.h
+++ b/Engine/Source/Runtime/Windows/D3D11RHI/Public/ID3D11DynamicRHI.h
@@ -1,7 +1,7 @@
 // Copyright Epic Games, Inc. All Rights Reserved.
 #pragma once
 
-#include "DynamicRHI.h"
+#include "RHI.h"
 
 #if PLATFORM_WINDOWS
 	#include "Windows/D3D11ThirdParty.h"
diff --git a/Engine/Source/Runtime/Windows/D3D11RHI/Public/Windows/D3D11ThirdParty.h b/Engine/Source/Runtime/Windows/D3D11RHI/Public/Windows/D3D11ThirdParty.h
--- a/Engine/Source/Runtime/Windows/D3D11RHI/Public/Windows/D3D11ThirdParty.h
+++ b/Engine/Source/Runtime/Windows/D3D11RHI/Public/Windows/D3D11ThirdParty.h
@@ -15,6 +15,7 @@
 #include "Microsoft/AllowMicrosoftPlatformTypes.h"
 THIRD_PARTY_INCLUDES_START
 	#include <d3d11_2.h>
+	#include <dxgi1_6.h>
 	#include <dxgidebug.h>
 THIRD_PARTY_INCLUDES_END
 #include "Microsoft/HideMicrosoftPlatformTypes.h"
