- Fix for incorrect UAVs on Cube / Cube Arrays

#rnx

diff --git a/Engine/Source/Runtime/Windows/D3D11RHI/Private/D3D11UAV.cpp b/Engine/Source/Runtime/Windows/D3D11RHI/Private/D3D11UAV.cpp
--- a/Engine/Source/Runtime/Windows/D3D11RHI/Private/D3D11UAV.cpp
+++ b/Engine/Source/Runtime/Windows/D3D11RHI/Private/D3D11UAV.cpp
@@ -98,8 +98,8 @@ FUnorderedAccessViewRHIRef FD3D11DynamicRHI::RHICreateUnorderedAccessView(FRHITe
 	{
 		UAVDesc.ViewDimension = D3D11_UAV_DIMENSION_TEXTURE2DARRAY;
 		UAVDesc.Texture2DArray.MipSlice = MipLevel;
-		UAVDesc.Texture2DArray.FirstArraySlice = 0;
-		UAVDesc.Texture2DArray.ArraySize = Texture->GetSizeZ();
+		UAVDesc.Texture2DArray.FirstArraySlice = NumArraySlices == 0 ? 0 : FirstArraySlice * 6;
+		UAVDesc.Texture2DArray.ArraySize = NumArraySlices == 0 ? Texture->GetDesc().ArraySize * 6 : NumArraySlices * 6;
 	}
 	else
 	{
