Fix for rare ensure on DeviceStateChanges in WindowMMDeviceInfoCache

diff --git a/Engine/Source/Runtime/Windows/AudioMixerXAudio2/Private/WindowsMMDeviceInfoCache.cpp b/Engine/Source/Runtime/Windows/AudioMixerXAudio2/Private/WindowsMMDeviceInfoCache.cpp
--- a/Engine/Source/Runtime/Windows/AudioMixerXAudio2/Private/WindowsMMDeviceInfoCache.cpp
+++ b/Engine/Source/Runtime/Windows/AudioMixerXAudio2/Private/WindowsMMDeviceInfoCache.cpp
@@ -387,19 +387,15 @@ void FWindowsMMDeviceCache::OnDefaultRenderDeviceChanged(const EAudioDeviceRole
 
 void FWindowsMMDeviceCache::OnDeviceAdded(const FString& DeviceId, bool bIsRender)
 {
-	if (ensure(DeviceEnumerator))
+	if (TOptional<FCacheEntry> NewDeviceEntry = BuildCacheEntry(DeviceId))
 	{
-		TComPtr<IMMDevice> Device;
-		if (SUCCEEDED(DeviceEnumerator->GetDevice(*DeviceId, &Device)))
-		{
-			FCacheEntry Info{ *DeviceId };
-			if (EnumerateDeviceProps(Device, Info))
-			{
-				FWriteScopeLock WriteLock(CacheMutationLock);
-				check(!Cache.Contains(Info.DeviceId));
-				Cache.Add(Info.DeviceId, Info);
-			}
-		}
+		FWriteScopeLock WriteLock(CacheMutationLock);
+		check(!Cache.Contains(NewDeviceEntry->DeviceId));
+		Cache.Emplace(NewDeviceEntry->DeviceId, MoveTemp(*NewDeviceEntry));
+	}
+	else
+	{
+		UE_LOG(LogAudioMixer, Warning, TEXT("FWindowsMMDeviceCache::OnDeviceAdded: Failed to add DeviceID='%s' to cache. "), *DeviceId);
 	}
 }
 
@@ -413,13 +409,16 @@ void FWindowsMMDeviceCache::OnDeviceRemoved(const FString& DeviceId, bool)
 
 TOptional<FWindowsMMDeviceCache::FCacheEntry> FWindowsMMDeviceCache::BuildCacheEntry(const FString& DeviceId)
 {
-	TComPtr<IMMDevice> Device;
-	if (SUCCEEDED(DeviceEnumerator->GetDevice(*DeviceId, &Device)))
+	if (ensure(DeviceEnumerator))
 	{
-		FCacheEntry Info{ *DeviceId };
-		if (EnumerateDeviceProps(Device, Info))
+		TComPtr<IMMDevice> Device;
+		if (SUCCEEDED(DeviceEnumerator->GetDevice(*DeviceId, &Device)))
 		{
-			return Info;
+			FCacheEntry Info{ *DeviceId };
+			if (EnumerateDeviceProps(Device, Info))
+			{
+				return Info;
+			}
 		}
 	}
 	return {};
@@ -434,22 +433,21 @@ FString FWindowsMMDeviceCache::GetFriendlyName(FName InDeviceId) const
 	return TEXT("Unknown");
 }
 
-
-
 void FWindowsMMDeviceCache::OnDeviceStateChanged(const FString& DeviceId, const EAudioDeviceState InState, bool)
 {
-	TOptional<FCacheEntry> Info = BuildCacheEntry(DeviceId);
-
-	FReadScopeLock ReadLock(CacheMutationLock);
 	FName DeviceIdName = *DeviceId;
-	ensureMsgf(Cache.Contains(DeviceIdName), TEXT("Expecting to find '%s' in cache '%s'"), *DeviceId, Info ? *Info->FriendlyName : TEXT("Unknown"));
+	
+	// NOTE: If entry does not exist that's likely because a state change has preempted the OnDeviceAdded call.
 
+	// Scope for Read-lock on Cache Map.
+	FReadScopeLock ReadLock(CacheMutationLock);
 	if (FCacheEntry* Entry = Cache.Find(DeviceIdName))
 	{
+		// Inner Write-Lock on Entry.
 		FWriteScopeLock WriteLock(Entry->MutationLock);
 
 		UE_CLOG(Audio::IAudioMixer::ShouldLogDeviceSwaps(), LogAudioMixer, Verbose, TEXT("FWindowsMMDeviceCache: DeviceName='%s' - DeviceID='%s' state changed from '%s' to '%s'."),
-			Info ? *Info->FriendlyName : TEXT("Unknown"), *DeviceId, ToString(Entry->State), ToString(InState));
+			*Entry->FriendlyName, *DeviceId, ToString(Entry->State), ToString(InState));
 
 		Entry->State = InState;
 	}
