PR #9059: Add TexCreate_Shared to FD3D11Texture3D.
(Contributed by [REDACTED_PERSON])

diff --git a/Engine/Source/Runtime/Windows/D3D11RHI/Private/D3D11Texture.cpp b/Engine/Source/Runtime/Windows/D3D11RHI/Private/D3D11Texture.cpp
--- a/Engine/Source/Runtime/Windows/D3D11RHI/Private/D3D11Texture.cpp
+++ b/Engine/Source/Runtime/Windows/D3D11RHI/Private/D3D11Texture.cpp
@@ -814,6 +814,18 @@ FD3D11Texture* FD3D11DynamicRHI::CreateD3D11Texture3D(FRHITextureCreateDesc cons
 	TextureDesc.CPUAccessFlags = 0;
 	TextureDesc.MiscFlags = 0;
 
+	if (EnumHasAnyFlags(Flags, TexCreate_Shared))
+	{
+		if (GCVarUseSharedKeyedMutex->GetInt() != 0)
+		{
+			TextureDesc.MiscFlags |= D3D11_RESOURCE_MISC_SHARED_KEYEDMUTEX;
+		}
+		else
+		{
+			TextureDesc.MiscFlags |= D3D11_RESOURCE_MISC_SHARED;
+		}
+	}
+
 	if (EnumHasAnyFlags(Flags, TexCreate_GenerateMipCapable))
 	{
 		// Set the flag that allows us to call GenerateMips on this texture later
