Have r.HDR.Display.ColorGamut enum values defined in a single place and replaced hardcoded numbers by enum values.

diff --git a/Engine/Shaders/Private/PostProcessDeviceEncodingOnly.usf b/Engine/Shaders/Private/PostProcessDeviceEncodingOnly.usf
--- a/Engine/Shaders/Private/PostProcessDeviceEncodingOnly.usf
+++ b/Engine/Shaders/Private/PostProcessDeviceEncodingOnly.usf
@@ -76,7 +76,7 @@ float4 DeviceEncodingOnlyCommonPS(
 	BRANCH
 
 	// sRGB, user specified gamut
-	if( GetOutputDevice() == 0 )
+	if( GetOutputDevice() == TONEMAPPER_OUTPUT_sRGB)
 	{		
 		// Convert from sRGB to specified output gamut	
 		// float3 OutputGamutColor = mul( AP1_2_Output, mul( sRGB_2_AP1, FilmColor ) );
@@ -90,7 +90,7 @@ float4 DeviceEncodingOnlyCommonPS(
 	}
 
 	// Rec 709, user specified gamut
-	else if( GetOutputDevice() == 1 )
+	else if( GetOutputDevice() == TONEMAPPER_OUTPUT_Rec709)
 	{
 		// Convert from sRGB to specified output gamut
 		float3 OutputGamutColor = mul( AP1_2_Output, mul( sRGB_2_AP1, FilmColor ) );
@@ -100,7 +100,7 @@ float4 DeviceEncodingOnlyCommonPS(
 	}
 
 	// ACES 1000nit transform with PQ/2084 encoding, user specified gamut 
-	else if( GetOutputDevice() == 3 || GetOutputDevice() == 5 )
+	else if( GetOutputDevice() == TONEMAPPER_OUTPUT_ACES1000nitST2084 || GetOutputDevice() == TONEMAPPER_OUTPUT_ACES1000nitScRGB)
 	{		
 		// 1000 nit ODT
 		float3 ODTColor = ACESOutputTransforms1000( FilmColor );
@@ -113,7 +113,7 @@ float4 DeviceEncodingOnlyCommonPS(
 	}
 
 	// ACES 2000nit transform with PQ/2084 encoding, user specified gamut 
-	else if( GetOutputDevice() == 4 || GetOutputDevice() == 6 )
+	else if( GetOutputDevice() == TONEMAPPER_OUTPUT_ACES2000nitST2084 || GetOutputDevice() == TONEMAPPER_OUTPUT_ACES2000nitScRGB)
 	{		
 		// 2000 nit ODT
 		float3 ODTColor = ACESOutputTransforms2000( FilmColor );
@@ -125,25 +125,25 @@ float4 DeviceEncodingOnlyCommonPS(
 		OutDeviceColor = LinearToST2084( ODTColor );
 	}	
 	
-	else if( GetOutputDevice() == 7 )
+	else if( GetOutputDevice() == TONEMAPPER_OUTPUT_LinearEXR)
 	{
 		float3 OutputGamutColor = mul( AP1_2_Output, mul( sRGB_2_AP1, FilmColor ) );
 		OutDeviceColor = LinearToST2084( OutputGamutColor );
 	}
 	// Linear HDR, including all color correction, but no tone curve
-	else if( GetOutputDevice() == 8 ) 
+	else if( GetOutputDevice() == TONEMAPPER_OUTPUT_NoToneCurve)
 	{
 		OutDeviceColor = FilmColor;
 	}
 	// "Linear" including all color correction and the tone curve, but no device gamma
-	else if (GetOutputDevice() == 9)
+	else if (GetOutputDevice() == TONEMAPPER_OUTPUT_WithToneCurve)
 	{
 		float3 OutputGamutColor = mul(AP1_2_Output, mul(sRGB_2_AP1, SceneColor.rgb)); 
 
 		OutDeviceColor = OutputGamutColor;
 	}
 
-	// OutputDevice == 2
+	// OutputDevice == TONEMAPPER_OUTPUT_ExplicitGammaMapping
 	// Gamma 2.2, user specified gamut
 	else
 	{
diff --git a/Engine/Shaders/Private/TonemapCommon.ush b/Engine/Shaders/Private/TonemapCommon.ush
--- a/Engine/Shaders/Private/TonemapCommon.ush
+++ b/Engine/Shaders/Private/TonemapCommon.ush
@@ -19,6 +19,14 @@
 #define TONEMAPPER_OUTPUT_NoToneCurve 8
 #define TONEMAPPER_OUTPUT_WithToneCurve 9
 
+/** Values of DIM_OUTPUT_DEVICE (matching C++'s EDisplayColorGamut) */
+#define TONEMAPPER_GAMUT_sRGB_D65 0
+#define TONEMAPPER_GAMUT_DCIP3_D65 1
+#define TONEMAPPER_GAMUT_Rec2020_D65 2
+#define TONEMAPPER_GAMUT_ACES_D60 3
+#define TONEMAPPER_GAMUT_ACEScg_D60 4
+
+
 // usually 1/2.2, the .y is used for inverse gamma when "gamma only" mode is not used
 half3 InverseGamma;
 
@@ -303,13 +311,13 @@ float3x3 OuputGamutMappingMatrix( uint OutputGamut )
 	// 3 = ACES AP0 - D60
 	// 4 = ACES AP1 - D60
 
-	if( OutputGamut == 1 )
+	if( OutputGamut == TONEMAPPER_GAMUT_DCIP3_D65)
 		return AP1_2_DCI_D65;
-	else if( OutputGamut == 2 )
+	else if( OutputGamut == TONEMAPPER_GAMUT_Rec2020_D65)
 		return AP1_2_Rec2020;
-	else if( OutputGamut == 3 )
+	else if( OutputGamut == TONEMAPPER_GAMUT_ACES_D60)
 		return AP1_2_AP0_MAT;
-	else if( OutputGamut == 4 )
+	else if( OutputGamut == TONEMAPPER_GAMUT_ACEScg_D60)
 		return GamutMappingIdentityMatrix;
 	else
 		return AP1_2_sRGB;
@@ -328,12 +336,14 @@ float3x3 OuputInverseGamutMappingMatrix( uint OutputGamut )
 	// 2 = Rec.2020 - D65
 	// 3 = ACES AP0 - D60
 	float3x3 GamutMappingMatrix = sRGB_2_AP1;
-	if( OutputGamut == 1 )
+	if( OutputGamut == TONEMAPPER_GAMUT_DCIP3_D65)
 		GamutMappingMatrix = DCI_D65_2_AP1;
-	else if( OutputGamut == 2 )
+	else if( OutputGamut == TONEMAPPER_GAMUT_Rec2020_D65)
 		GamutMappingMatrix = Rec2020_2_AP1;
-	else if( OutputGamut == 3 )
+	else if( OutputGamut == TONEMAPPER_GAMUT_ACES_D60)
 		GamutMappingMatrix = AP0_2_AP1_MAT;
+	else if (OutputGamut == TONEMAPPER_GAMUT_ACEScg_D60)
+		GamutMappingMatrix = GamutMappingIdentityMatrix;
 
 	return GamutMappingMatrix;
 }
diff --git a/Engine/Source/Runtime/Core/Public/GenericPlatform/GenericPlatformMisc.h b/Engine/Source/Runtime/Core/Public/GenericPlatform/GenericPlatformMisc.h
--- a/Engine/Source/Runtime/Core/Public/GenericPlatform/GenericPlatformMisc.h
+++ b/Engine/Source/Runtime/Core/Public/GenericPlatform/GenericPlatformMisc.h
@@ -445,9 +445,23 @@ enum class EDisplayOutputFormat
 	HDR_LinearEXR = 7,
 	HDR_LinearNoToneCurve = 8,
 	HDR_LinearWithToneCurve = 9,
+
 	MAX
 };
 
+/**
+ * Display gamut, format, and chromacities. You must update values in TonemapCommon.ush when changing this enum + EHDRCaptureGamut
+ */
+enum class EDisplayColorGamut
+{
+	sRGB_D65 = 0,
+	DCIP3_D65 = 1,
+	Rec2020_D65 = 2,
+	ACES_D60 = 3,
+	ACEScg_D60 = 4,
+
+	MAX
+};
 
 /**
  * Different types of Context Switch stats
@@ -1663,7 +1677,7 @@ public:
 		return false;
 	}
 
-	FORCEINLINE static void ChooseHDRDeviceAndColorGamut(uint32 DeviceId, uint32 DisplayNitLevel, int32& OutputDevice, int32& ColorGamut)
+	FORCEINLINE static void ChooseHDRDeviceAndColorGamut(uint32 DeviceId, uint32 DisplayNitLevel, EDisplayOutputFormat& OutputDevice, EDisplayColorGamut& ColorGamut)
 	{
 	}
 
diff --git a/Engine/Source/Runtime/Core/Public/IOS/IOSPlatformMisc.h b/Engine/Source/Runtime/Core/Public/IOS/IOSPlatformMisc.h
--- a/Engine/Source/Runtime/Core/Public/IOS/IOSPlatformMisc.h
+++ b/Engine/Source/Runtime/Core/Public/IOS/IOSPlatformMisc.h
@@ -194,11 +194,11 @@ struct CORE_API FIOSPlatformMisc : public FApplePlatformMisc
 
 	static bool RequestDeviceCheckToken(TFunction<void(const TArray<uint8>&)> QuerySucceededFunc, TFunction<void(const FString&, const FString&)> QueryFailedFunc);
     
-	FORCEINLINE static void ChooseHDRDeviceAndColorGamut(uint32 DeviceId, uint32 DisplayNitLevel, int32& OutputDevice, int32& ColorGamut)
+	FORCEINLINE static void ChooseHDRDeviceAndColorGamut(uint32 DeviceId, uint32 DisplayNitLevel, EDisplayOutputFormat& OutputDevice, EDisplayColorGamut& ColorGamut)
 	{
 		// Linear output to Apple's specific format.
-		OutputDevice = (int32)EDisplayOutputFormat::HDR_LinearEXR;
-		ColorGamut = 0;
+		OutputDevice = EDisplayOutputFormat::HDR_LinearEXR;
+		ColorGamut = EDisplayColorGamut::sRGB_D65;
 	}
 
 	static int32 GetMaxRefreshRate();
diff --git a/Engine/Source/Runtime/Core/Public/Mac/MacPlatformMisc.h b/Engine/Source/Runtime/Core/Public/Mac/MacPlatformMisc.h
--- a/Engine/Source/Runtime/Core/Public/Mac/MacPlatformMisc.h
+++ b/Engine/Source/Runtime/Core/Public/Mac/MacPlatformMisc.h
@@ -231,11 +231,11 @@ struct CORE_API FMacPlatformMisc : public FApplePlatformMisc
 
 	static CGDisplayModeRef GetSupportedDisplayMode(CGDirectDisplayID DisplayID, uint32 Width, uint32 Height);
 
-	FORCEINLINE static void ChooseHDRDeviceAndColorGamut(uint32 DeviceId, uint32 DisplayNitLevel, int32& OutputDevice, int32& ColorGamut)
+	FORCEINLINE static void ChooseHDRDeviceAndColorGamut(uint32 DeviceId, uint32 DisplayNitLevel, EDisplayOutputFormat& OutputDevice, EDisplayColorGamut& ColorGamut)
 	{
 		// ScRGB, 1000 or 2000 nits, DCI-P3
-		OutputDevice = DisplayNitLevel == 1000 ? (int32)EDisplayOutputFormat::HDR_ACES_1000nit_ScRGB : (int32)EDisplayOutputFormat::HDR_ACES_2000nit_ScRGB;
-		ColorGamut = 1;
+		OutputDevice = DisplayNitLevel == 1000 ? EDisplayOutputFormat::HDR_ACES_1000nit_ScRGB : EDisplayOutputFormat::HDR_ACES_2000nit_ScRGB;
+		ColorGamut = EDisplayColorGamut::DCIP3_D65;
 	}
 };
 
diff --git a/Engine/Source/Runtime/Core/Public/Windows/WindowsPlatformMisc.h b/Engine/Source/Runtime/Core/Public/Windows/WindowsPlatformMisc.h
--- a/Engine/Source/Runtime/Core/Public/Windows/WindowsPlatformMisc.h
+++ b/Engine/Source/Runtime/Core/Public/Windows/WindowsPlatformMisc.h
@@ -255,21 +255,21 @@ struct CORE_API FWindowsPlatformMisc
 	 */
 	static bool IsRunningOnBattery();
 
-	FORCEINLINE static void ChooseHDRDeviceAndColorGamut(uint32 DeviceId, uint32 DisplayNitLevel, int32& OutputDevice, int32& ColorGamut)
+	FORCEINLINE static void ChooseHDRDeviceAndColorGamut(uint32 DeviceId, uint32 DisplayNitLevel, EDisplayOutputFormat& OutputDevice, EDisplayColorGamut& ColorGamut)
 	{
 		if (DeviceId == 0x1002 /*AMD*/ || DeviceId == 0x10DE /*NVIDIA*/)
 		{
 			// needs to match GRHIHDRDisplayOutputFormat chosen in FD3D12DynamicRHI::Init
 #if WITH_EDITOR
 		// ScRGB, 1000 or 2000 nits
-			OutputDevice = (DisplayNitLevel == 1000) ? (int32)EDisplayOutputFormat::HDR_ACES_1000nit_ScRGB : (int32)EDisplayOutputFormat::HDR_ACES_2000nit_ScRGB;
+			OutputDevice = (DisplayNitLevel == 1000) ? EDisplayOutputFormat::HDR_ACES_1000nit_ScRGB : EDisplayOutputFormat::HDR_ACES_2000nit_ScRGB;
 			// Rec709
-			ColorGamut = 0;
+			ColorGamut = EDisplayColorGamut::sRGB_D65;
 #else
 		// ST-2084, 1000 or 2000 nits
-			OutputDevice = (DisplayNitLevel == 1000) ? (int32)EDisplayOutputFormat::HDR_ACES_1000nit_ST2084 : (int32)EDisplayOutputFormat::HDR_ACES_2000nit_ST2084;
+			OutputDevice = (DisplayNitLevel == 1000) ? EDisplayOutputFormat::HDR_ACES_1000nit_ST2084 : EDisplayOutputFormat::HDR_ACES_2000nit_ST2084;
 			// Rec2020
-			ColorGamut = 2;
+			ColorGamut = EDisplayColorGamut::Rec2020_D65;
 #endif
 		}
 	}
diff --git a/Engine/Source/Runtime/D3D12RHI/Private/Windows/WindowsD3D12Viewport.cpp b/Engine/Source/Runtime/D3D12RHI/Private/Windows/WindowsD3D12Viewport.cpp
--- a/Engine/Source/Runtime/D3D12RHI/Private/Windows/WindowsD3D12Viewport.cpp
+++ b/Engine/Source/Runtime/D3D12RHI/Private/Windows/WindowsD3D12Viewport.cpp
@@ -408,7 +408,7 @@ void FD3D12Viewport::EnableHDR()
 		// Ideally we can avoid setting TV meta data and instead the engine can do tone mapping based on the
 		// actual current display properties (display mapping).
 		static const auto CVarHDRColorGamut = IConsoleManager::Get().FindTConsoleVariableDataInt(TEXT("r.HDR.Display.ColorGamut"));
-		const EDisplayGamut DisplayGamut = EDisplayGamut(CVarHDRColorGamut->GetValueOnAnyThread());
+		const EDisplayColorGamut DisplayGamut = EDisplayColorGamut(CVarHDRColorGamut->GetValueOnAnyThread());
 		SetHDRTVMode(true,
 			DisplayGamut,
 			DisplayMaxOutputNits,
@@ -426,7 +426,7 @@ void FD3D12Viewport::ShutdownHDR()
 	if (GRHISupportsHDROutput)
 	{
 		// Default SDR display data
-		const EDisplayGamut DisplayGamut = DG_Rec709;
+		const EDisplayColorGamut DisplayGamut = EDisplayColorGamut::sRGB_D65;
 		const EDisplayOutputFormat OutputDevice = EDisplayOutputFormat::SDR_sRGB;
 
 		// Note: These values aren't actually used.
@@ -508,12 +508,12 @@ static const FString GetDXGIColorSpaceString(DXGI_COLOR_SPACE_TYPE ColorSpace)
 	return FString::FromInt(ColorSpace);
 };
 
-void FD3D12Viewport::EnsureColorSpace(EDisplayGamut DisplayGamut, EDisplayOutputFormat OutputDevice)
+void FD3D12Viewport::EnsureColorSpace(EDisplayColorGamut DisplayGamut, EDisplayOutputFormat OutputDevice)
 {
 	ensure(SwapChain4.GetReference());
 
 	DXGI_COLOR_SPACE_TYPE NewColorSpace = DXGI_COLOR_SPACE_RGB_FULL_G22_NONE_P709;	// sRGB;
-	const bool bPrimaries2020 = (DisplayGamut == DG_Rec2020);
+	const bool bPrimaries2020 = (DisplayGamut == EDisplayColorGamut::Rec2020_D65);
 
 	// See console variable r.HDR.Display.OutputDevice.
 	switch (OutputDevice)
@@ -552,22 +552,22 @@ void FD3D12Viewport::EnsureColorSpace(EDisplayGamut DisplayGamut, EDisplayOutput
 	}
 }
 
-void FD3D12Viewport::SetHDRTVMode(bool bEnableHDR, EDisplayGamut DisplayGamut, float MaxOutputNits, float MinOutputNits, float MaxCLL, float MaxFALL)
+void FD3D12Viewport::SetHDRTVMode(bool bEnableHDR, EDisplayColorGamut DisplayGamut, float MaxOutputNits, float MinOutputNits, float MaxCLL, float MaxFALL)
 {
 	ensure(SwapChain4.GetReference());
 
 	static const DisplayChromacities DisplayChromacityList[] =
 	{
-		{ 0.64000f, 0.33000f, 0.30000f, 0.60000f, 0.15000f, 0.06000f, 0.31270f, 0.32900f }, // DG_Rec709
-		{ 0.68000f, 0.32000f, 0.26500f, 0.69000f, 0.15000f, 0.06000f, 0.31270f, 0.32900f }, // DG_DCI-P3 D65
-		{ 0.70800f, 0.29200f, 0.17000f, 0.79700f, 0.13100f, 0.04600f, 0.31270f, 0.32900f }, // DG_Rec2020
-		{ 0.73470f, 0.26530f, 0.00000f, 1.00000f, 0.00010f,-0.07700f, 0.32168f, 0.33767f }, // DG_ACES
-		{ 0.71300f, 0.29300f, 0.16500f, 0.83000f, 0.12800f, 0.04400f, 0.32168f, 0.33767f }, // DG_ACEScg
+		{ 0.64000f, 0.33000f, 0.30000f, 0.60000f, 0.15000f, 0.06000f, 0.31270f, 0.32900f }, // EDisplayColorGamut::sRGB_D65
+		{ 0.68000f, 0.32000f, 0.26500f, 0.69000f, 0.15000f, 0.06000f, 0.31270f, 0.32900f }, // EDisplayColorGamut::DCIP3_D65
+		{ 0.70800f, 0.29200f, 0.17000f, 0.79700f, 0.13100f, 0.04600f, 0.31270f, 0.32900f }, // EDisplayColorGamut::Rec2020_D65
+		{ 0.73470f, 0.26530f, 0.00000f, 1.00000f, 0.00010f,-0.07700f, 0.32168f, 0.33767f }, // EDisplayColorGamut::ACES_D60
+		{ 0.71300f, 0.29300f, 0.16500f, 0.83000f, 0.12800f, 0.04400f, 0.32168f, 0.33767f }, // EDisplayColorGamut::ACEScg_D60
 	};
 
 	if (bEnableHDR)
 	{
-		const DisplayChromacities& Chroma = DisplayChromacityList[DisplayGamut];
+		const DisplayChromacities& Chroma = DisplayChromacityList[(int32)DisplayGamut];
 
 		// Set HDR meta data
 		DXGI_HDR_METADATA_HDR10 HDR10MetaData = {};
diff --git a/Engine/Source/Runtime/D3D12RHI/Public/D3D12Viewport.h b/Engine/Source/Runtime/D3D12RHI/Public/D3D12Viewport.h
--- a/Engine/Source/Runtime/D3D12RHI/Public/D3D12Viewport.h
+++ b/Engine/Source/Runtime/D3D12RHI/Public/D3D12Viewport.h
@@ -259,17 +259,6 @@ private:
 	FD3D12FramePacing* FramePacerRunnable;
 #endif //WITH_MGPU
 
-	// Display gamut, format, and chromacities
-	// Note: Must be kept in sync with CVars and Tonemapping shaders
-	enum EDisplayGamut
-	{
-		DG_Rec709,
-		DG_DCI_P3,
-		DG_Rec2020,
-		DG_ACES,
-		DG_ACEScg
-	};
-
 	struct DisplayChromacities
 	{
 		float RedX, RedY;
@@ -289,14 +278,14 @@ private:
 
 #if PLATFORM_WINDOWS || PLATFORM_HOLOLENS
 	/** Ensure the correct color space is set on the swap chain */
-	void EnsureColorSpace(EDisplayGamut DisplayGamut, EDisplayOutputFormat OutputDevice);
+	void EnsureColorSpace(EDisplayColorGamut DisplayGamut, EDisplayOutputFormat OutputDevice);
 
 	/** 
 	 * Set HDR meta data. 
 	 * Note: Meta data should only be provided for TVs, not monitors. 
 	 * This is because the TV is doing the work to display the colors correctly.
 	 */
-	void SetHDRTVMode(bool bEnableHDR, EDisplayGamut DisplayGamut, float MaxOutputNits, float MinOutputNits, float MaxCLL, float MaxFALL);
+	void SetHDRTVMode(bool bEnableHDR, EDisplayColorGamut DisplayGamut, float MaxOutputNits, float MinOutputNits, float MaxCLL, float MaxFALL);
 #endif
 };
 
diff --git a/Engine/Source/Runtime/Engine/Private/UnrealEngine.cpp b/Engine/Source/Runtime/Engine/Private/UnrealEngine.cpp
--- a/Engine/Source/Runtime/Engine/Private/UnrealEngine.cpp
+++ b/Engine/Source/Runtime/Engine/Private/UnrealEngine.cpp
@@ -852,8 +852,8 @@ void ENGINE_API HDRSettingChangedSinkCallback()
 			DisplayNitLevel = 1000;
 		}
 		
-		int32 OutputDevice = 0;
-		int32 ColorGamut = 0;
+		EDisplayOutputFormat OutputDevice = EDisplayOutputFormat::SDR_sRGB;
+		EDisplayColorGamut ColorGamut = EDisplayColorGamut::sRGB_D65;
 		
 		// If we are turning HDR on we must set the appropriate OutputDevice and ColorGamut.
 		// If we are turning it off, we'll reset back to 0/0
@@ -867,8 +867,8 @@ void ENGINE_API HDRSettingChangedSinkCallback()
 		check(CVarHDROutputDevice);
 		check(CVarHDRColorGamut);
 		
-		CVarHDROutputDevice->Set(OutputDevice, ECVF_SetByDeviceProfile);
-		CVarHDRColorGamut->Set(ColorGamut, ECVF_SetByDeviceProfile);
+		CVarHDROutputDevice->Set((int32)OutputDevice, ECVF_SetByDeviceProfile);
+		CVarHDRColorGamut->Set((int32)ColorGamut, ECVF_SetByDeviceProfile);
 		
 		// Now set the HDR setting.
 		GRHIIsHDREnabled = CVarHDROutputEnabled->GetValueOnAnyThread() != 0;
diff --git a/Engine/Source/Runtime/MovieSceneCapture/Private/CompositionGraphCaptureProtocol.cpp b/Engine/Source/Runtime/MovieSceneCapture/Private/CompositionGraphCaptureProtocol.cpp
--- a/Engine/Source/Runtime/MovieSceneCapture/Private/CompositionGraphCaptureProtocol.cpp
+++ b/Engine/Source/Runtime/MovieSceneCapture/Private/CompositionGraphCaptureProtocol.cpp
@@ -34,17 +34,17 @@ struct FFrameCaptureViewExtension : public FSceneViewExtensionBase
 
 		RestoreDumpHDR = CVarDumpFramesAsHDR->GetInt();
 		RestoreHDRCompressionQuality = CVarHDRCompressionQuality->GetInt();
-		RestoreDumpGamut = CVarDumpGamut->GetInt();
-		RestoreDumpDevice = CVarDumpDevice->GetInt();
+		RestoreDumpGamut = (EDisplayColorGamut)CVarDumpGamut->GetInt();
+		RestoreDumpDevice = (EDisplayOutputFormat)CVarDumpDevice->GetInt();
 
 		if (CaptureGamut == HCGM_Linear)
 		{
-			CVarDumpGamut->Set(1);
+			CVarDumpGamut->Set((int32)EDisplayColorGamut::DCIP3_D65);
 			CVarDumpDevice->Set((int32)EDisplayOutputFormat::HDR_LinearEXR);
 		}
 		else
 		{
-			CVarDumpGamut->Set(CaptureGamut);
+			CVarDumpGamut->Set((int32)CaptureGamut);
 		}
 
 		Disable();
@@ -54,8 +54,8 @@ struct FFrameCaptureViewExtension : public FSceneViewExtensionBase
 	{
 		Disable();
 
-		CVarDumpGamut->Set(RestoreDumpGamut);
-		CVarDumpDevice->Set(RestoreDumpDevice);
+		CVarDumpGamut->Set((int32)RestoreDumpGamut);
+		CVarDumpDevice->Set((int32)RestoreDumpDevice);
 	}
 
 	bool IsEnabled() const
@@ -166,8 +166,8 @@ private:
 
 	int32 RestoreDumpHDR;
 	int32 RestoreHDRCompressionQuality;
-	int32 RestoreDumpGamut;
-	int32 RestoreDumpDevice;
+	EDisplayColorGamut RestoreDumpGamut;
+	EDisplayOutputFormat RestoreDumpDevice;
 };
 
 bool UCompositionGraphCaptureProtocol::SetupImpl()
diff --git a/Engine/Source/Runtime/MovieSceneCapture/Private/ImageSequenceProtocol.cpp b/Engine/Source/Runtime/MovieSceneCapture/Private/ImageSequenceProtocol.cpp
--- a/Engine/Source/Runtime/MovieSceneCapture/Private/ImageSequenceProtocol.cpp
+++ b/Engine/Source/Runtime/MovieSceneCapture/Private/ImageSequenceProtocol.cpp
@@ -209,12 +209,12 @@ bool UImageSequenceProtocol_EXR::SetupImpl()
 	IConsoleVariable* CVarDumpGamut = IConsoleManager::Get().FindConsoleVariable(TEXT("r.HDR.Display.ColorGamut"));
 	IConsoleVariable* CVarDumpDevice = IConsoleManager::Get().FindConsoleVariable(TEXT("r.HDR.Display.OutputDevice"));
 
-	RestoreColorGamut = CVarDumpGamut->GetInt();
-	RestoreOutputDevice = CVarDumpDevice->GetInt();
+	RestoreColorGamut = (EDisplayColorGamut)CVarDumpGamut->GetInt();
+	RestoreOutputDevice = (EDisplayOutputFormat)CVarDumpDevice->GetInt();
 
 	if (CaptureGamut == HCGM_Linear)
 	{
-		CVarDumpGamut->Set(1);
+		CVarDumpGamut->Set((int32)EDisplayColorGamut::DCIP3_D65);
 		CVarDumpDevice->Set((int32)EDisplayOutputFormat::HDR_LinearEXR);
 	}
 	else
@@ -232,8 +232,8 @@ void UImageSequenceProtocol_EXR::FinalizeImpl()
 	IConsoleVariable* CVarDumpGamut = IConsoleManager::Get().FindConsoleVariable(TEXT("r.HDR.Display.ColorGamut"));
 	IConsoleVariable* CVarDumpDevice = IConsoleManager::Get().FindConsoleVariable(TEXT("r.HDR.Display.OutputDevice"));
 
-	CVarDumpGamut->Set(RestoreColorGamut);
-	CVarDumpDevice->Set(RestoreOutputDevice);
+	CVarDumpGamut->Set((int32)RestoreColorGamut);
+	CVarDumpDevice->Set((int32)RestoreOutputDevice);
 }
 
 void UImageSequenceProtocol_EXR::AddFormatMappingsImpl(TMap<FString, FStringFormatArg>& FormatMappings) const
diff --git a/Engine/Source/Runtime/MovieSceneCapture/Public/Protocols/CompositionGraphCaptureProtocol.h b/Engine/Source/Runtime/MovieSceneCapture/Public/Protocols/CompositionGraphCaptureProtocol.h
--- a/Engine/Source/Runtime/MovieSceneCapture/Public/Protocols/CompositionGraphCaptureProtocol.h
+++ b/Engine/Source/Runtime/MovieSceneCapture/Public/Protocols/CompositionGraphCaptureProtocol.h
@@ -14,7 +14,7 @@ class UMaterialInterface;
 struct FMovieSceneCaptureSettings;
 struct FFrameCaptureViewExtension;
 
-/** Used by UCompositionGraphCaptureSettings. Matches gamut oreder in TonemapCommon.usf OuputGamutMappingMatrix()*/
+/** Used by UCompositionGraphCaptureSettings. Matches gamut order in EDisplayColorGamut */
 UENUM(BlueprintType)
 enum EHDRCaptureGamut
 {
@@ -27,6 +27,15 @@ enum EHDRCaptureGamut
 	HCGM_MAX,
 };
 
+static_assert(HCGM_Rec709 == (int32)EDisplayColorGamut::sRGB_D65, "EHDRCaptureGamut and EDisplayColorGamut not matching");
+static_assert(HCGM_P3DCI == (int32)EDisplayColorGamut::DCIP3_D65, "EHDRCaptureGamut and EDisplayColorGamut not matching");
+static_assert(HCGM_Rec2020 == (int32)EDisplayColorGamut::Rec2020_D65, "EHDRCaptureGamut and EDisplayColorGamut not matching");
+static_assert(HCGM_ACES == (int32)EDisplayColorGamut::ACES_D60, "EHDRCaptureGamut and EDisplayColorGamut not matching");
+static_assert(HCGM_ACEScg == (int32)EDisplayColorGamut::ACEScg_D60, "EHDRCaptureGamut and EDisplayColorGamut not matching");
+// HCGM_Linear gets remapped to DCIP3_D65 internally
+//static_assert(HCGM_Linear == (int32)EDisplayColorGamut::DCIP3_D65, "EHDRCaptureGamut and EDisplayColorGamut not matching")
+static_assert(HCGM_MAX == (int32)EDisplayColorGamut::MAX + 1, "EHDRCaptureGamut and EDisplayColorGamut not matching");
+
 USTRUCT(BlueprintType)
 struct MOVIESCENECAPTURE_API FCompositionGraphCapturePasses
 {
diff --git a/Engine/Source/Runtime/MovieSceneCapture/Public/Protocols/ImageSequenceProtocol.h b/Engine/Source/Runtime/MovieSceneCapture/Public/Protocols/ImageSequenceProtocol.h
--- a/Engine/Source/Runtime/MovieSceneCapture/Public/Protocols/ImageSequenceProtocol.h
+++ b/Engine/Source/Runtime/MovieSceneCapture/Public/Protocols/ImageSequenceProtocol.h
@@ -148,7 +148,7 @@ private:
 	virtual void FinalizeImpl() override;
 	virtual void AddFormatMappingsImpl(TMap<FString, FStringFormatArg>& FormatMappings) const override;
 
-	int32 RestoreColorGamut;
-	int32 RestoreOutputDevice;
+	EDisplayColorGamut RestoreColorGamut;
+	EDisplayOutputFormat RestoreOutputDevice;
 };
 
diff --git a/Engine/Source/Runtime/Windows/D3D11RHI/Private/Windows/WindowsD3D11Device.cpp b/Engine/Source/Runtime/Windows/D3D11RHI/Private/Windows/WindowsD3D11Device.cpp
--- a/Engine/Source/Runtime/Windows/D3D11RHI/Private/Windows/WindowsD3D11Device.cpp
+++ b/Engine/Source/Runtime/Windows/D3D11RHI/Private/Windows/WindowsD3D11Device.cpp
@@ -352,17 +352,6 @@ static bool SafeTestD3D11CreateDevice(IDXGIAdapter* Adapter,D3D_FEATURE_LEVEL Mi
 	return false;
 }
 
-// Display gamut and chromaticities
-// Note: Must be kept in sync with CVars and Tonemapping shaders
-enum EDisplayGamut
-{
-	DG_Rec709,
-	DG_DCI_P3,
-	DG_Rec2020,
-	DG_ACES,
-	DG_ACEScg
-};
-
 struct DisplayChromacities
 {
 	float RedX, RedY;
@@ -373,14 +362,14 @@ struct DisplayChromacities
 
 const DisplayChromacities DisplayChromacityList[] =
 {
-	{ 0.64000f, 0.33000f, 0.30000f, 0.60000f, 0.15000f, 0.06000f, 0.31270f, 0.32900f }, // DG_Rec709
-	{ 0.68000f, 0.32000f, 0.26500f, 0.69000f, 0.15000f, 0.06000f, 0.31270f, 0.32900f }, // DG_DCI-P3 D65
-	{ 0.70800f, 0.29200f, 0.17000f, 0.79700f, 0.13100f, 0.04600f, 0.31270f, 0.32900f }, // DG_Rec2020
-	{ 0.73470f, 0.26530f, 0.00000f, 1.00000f, 0.00010f,-0.07700f, 0.32168f, 0.33767f }, // DG_ACES
-	{ 0.71300f, 0.29300f, 0.16500f, 0.83000f, 0.12800f, 0.04400f, 0.32168f, 0.33767f }, // DG_ACEScg
+	{ 0.64000f, 0.33000f, 0.30000f, 0.60000f, 0.15000f, 0.06000f, 0.31270f, 0.32900f }, // EDisplayColorGamut::sRGB_D65
+	{ 0.68000f, 0.32000f, 0.26500f, 0.69000f, 0.15000f, 0.06000f, 0.31270f, 0.32900f }, // EDisplayColorGamut::DCIP3_D65
+	{ 0.70800f, 0.29200f, 0.17000f, 0.79700f, 0.13100f, 0.04600f, 0.31270f, 0.32900f }, // EDisplayColorGamut::Rec2020_D65
+	{ 0.73470f, 0.26530f, 0.00000f, 1.00000f, 0.00010f,-0.07700f, 0.32168f, 0.33767f }, // EDisplayColorGamut::ACES_D60
+	{ 0.71300f, 0.29300f, 0.16500f, 0.83000f, 0.12800f, 0.04400f, 0.32168f, 0.33767f }, // EDisplayColorGamut::ACEScg_D60
 };
 
-static void SetHDRMonitorModeNVIDIA(uint32 IHVDisplayIndex, bool bEnableHDR, EDisplayGamut DisplayGamut, float MaxOutputNits, float MinOutputNits, float MaxCLL, float MaxFALL)
+static void SetHDRMonitorModeNVIDIA(uint32 IHVDisplayIndex, bool bEnableHDR, EDisplayColorGamut DisplayGamut, float MaxOutputNits, float MinOutputNits, float MaxCLL, float MaxFALL)
 {
 #ifdef NVAPI_INTERFACE
 	NvAPI_Status NvStatus = NVAPI_OK;
@@ -404,7 +393,7 @@ static void SetHDRMonitorModeNVIDIA(uint32 IHVDisplayIndex, bool bEnableHDR, EDi
 			HDRColorData.static_metadata_descriptor_id = NV_STATIC_METADATA_TYPE_1;
 			HDRColorData.hdrMode = bEnableHDR ? NV_HDR_MODE_UHDBD : NV_HDR_MODE_OFF;
 
-			const DisplayChromacities& Chroma = DisplayChromacityList[DisplayGamut];
+			const DisplayChromacities& Chroma = DisplayChromacityList[(int32)DisplayGamut];
 
 			HDRColorData.mastering_display_data.displayPrimary_x0 = NvU16(Chroma.RedX * 50000.0f);
 			HDRColorData.mastering_display_data.displayPrimary_y0 = NvU16(Chroma.RedY * 50000.0f);
@@ -433,7 +422,7 @@ static void SetHDRMonitorModeNVIDIA(uint32 IHVDisplayIndex, bool bEnableHDR, EDi
 #endif //NVAPI_INTERFACE
 }
 
-static void SetHDRMonitorModeAMD(uint32 IHVDisplayIndex, bool bEnableHDR, EDisplayGamut DisplayGamut, float MaxOutputNits, float MinOutputNits, float MaxCLL, float MaxFALL)
+static void SetHDRMonitorModeAMD(uint32 IHVDisplayIndex, bool bEnableHDR, EDisplayColorGamut DisplayGamut, float MaxOutputNits, float MinOutputNits, float MaxCLL, float MaxFALL)
 {
 #ifdef AMD_AGS_API
 	const int32 AmdHDRDeviceIndex = (IHVDisplayIndex & 0xffff0000) >> 16;
@@ -454,7 +443,7 @@ static void SetHDRMonitorModeAMD(uint32 IHVDisplayIndex, bool bEnableHDR, EDispl
 
 		if (bEnableHDR)
 		{
-			const DisplayChromacities& Chroma = DisplayChromacityList[DisplayGamut];
+			const DisplayChromacities& Chroma = DisplayChromacityList[(int32)DisplayGamut];
 			HDRDisplaySettings.chromaticityRedX   = Chroma.RedX;
 			HDRDisplaySettings.chromaticityRedY   = Chroma.RedY;
 			HDRDisplaySettings.chromaticityGreenX = Chroma.GreenX;
@@ -488,9 +477,9 @@ void FD3D11DynamicRHI::EnableHDR()
 
 	if ( GRHISupportsHDROutput && IsHDREnabled() )
 	{
-		const int32 OutputDevice = CVarHDROutputDevice->GetValueOnAnyThread();
+		const EDisplayOutputFormat OutputDevice = (EDisplayOutputFormat)CVarHDROutputDevice->GetValueOnAnyThread();
 
-		const float DisplayMaxOutputNits = (OutputDevice == (int32)EDisplayOutputFormat::HDR_ACES_2000nit_ST2084 || OutputDevice == (int32)EDisplayOutputFormat::HDR_ACES_2000nit_ScRGB) ? 2000.f : 1000.f;
+		const float DisplayMaxOutputNits = (OutputDevice == EDisplayOutputFormat::HDR_ACES_2000nit_ST2084 || OutputDevice == EDisplayOutputFormat::HDR_ACES_2000nit_ScRGB) ? 2000.f : 1000.f;
 		const float DisplayMinOutputNits = 0.0f;	// Min output of the display
 		const float DisplayMaxCLL = 0.0f;			// Max content light level in lumens (0.0 == unknown)
 		const float DisplayFALL = 0.0f;				// Frame average light level (0.0 == unknown)
@@ -500,7 +489,7 @@ void FD3D11DynamicRHI::EnableHDR()
 			SetHDRMonitorModeNVIDIA(
 				HDRDetectedDisplayIHVIndex,
 				true,
-				EDisplayGamut(CVarHDRColorGamut->GetValueOnAnyThread()),
+				EDisplayColorGamut(CVarHDRColorGamut->GetValueOnAnyThread()),
 				DisplayMaxOutputNits,
 				DisplayMinOutputNits,
 				DisplayMaxCLL,
@@ -511,7 +500,7 @@ void FD3D11DynamicRHI::EnableHDR()
 			SetHDRMonitorModeAMD(
 				HDRDetectedDisplayIHVIndex,
 				true,
-				EDisplayGamut(CVarHDRColorGamut->GetValueOnAnyThread()),
+				EDisplayColorGamut(CVarHDRColorGamut->GetValueOnAnyThread()),
 				DisplayMaxOutputNits,
 				DisplayMinOutputNits,
 				DisplayMaxCLL,
@@ -540,7 +529,7 @@ void FD3D11DynamicRHI::ShutdownHDR()
 			SetHDRMonitorModeNVIDIA(
 				HDRDetectedDisplayIHVIndex,
 				false,
-				DG_Rec709,
+				EDisplayColorGamut::sRGB_D65,
 				DisplayMaxOutputNits,
 				DisplayMinOutputNits,
 				DisplayMaxCLL,
@@ -551,7 +540,7 @@ void FD3D11DynamicRHI::ShutdownHDR()
 			SetHDRMonitorModeAMD(
 				HDRDetectedDisplayIHVIndex,
 				false,
-				DG_Rec709,
+				EDisplayColorGamut::sRGB_D65,
 				DisplayMaxOutputNits,
 				DisplayMinOutputNits,
 				DisplayMaxCLL,
diff --git a/Engine/Source/Runtime/Windows/D3D11RHI/Private/Windows/WindowsD3D11Viewport.cpp b/Engine/Source/Runtime/Windows/D3D11RHI/Private/Windows/WindowsD3D11Viewport.cpp
--- a/Engine/Source/Runtime/Windows/D3D11RHI/Private/Windows/WindowsD3D11Viewport.cpp
+++ b/Engine/Source/Runtime/Windows/D3D11RHI/Private/Windows/WindowsD3D11Viewport.cpp
@@ -261,57 +261,123 @@ FD3D11Viewport::FD3D11Viewport(FD3D11DynamicRHI* InD3DRHI,HWND InWindowHandle,ui
 	});	
 }
 
-// When a window has moved or resized we need to check whether it is on a HDR monitor or not. Set the correct color space of the monitor
-void FD3D11Viewport::CheckHDRMonitorStatus()
+static const FString GetDXGIColorSpaceString(DXGI_COLOR_SPACE_TYPE ColorSpace)
 {
-#if WITH_EDITOR
+	switch (ColorSpace)
+	{
+	case DXGI_COLOR_SPACE_RGB_FULL_G22_NONE_P709:
+		return TEXT("RGB_FULL_G22_NONE_P709");
+	case DXGI_COLOR_SPACE_RGB_FULL_G10_NONE_P709:
+		return TEXT("RGB_FULL_G10_NONE_P709");
+	case DXGI_COLOR_SPACE_RGB_FULL_G2084_NONE_P2020:
+		return TEXT("RGB_FULL_G2084_NONE_P2020");
+	case DXGI_COLOR_SPACE_RGB_FULL_G22_NONE_P2020:
+		return TEXT("RGB_FULL_G22_NONE_P2020");
+	default:
+		break;
+	}
 
-	DXGI_COLOR_SPACE_TYPE colorSpace = DXGI_COLOR_SPACE_RGB_FULL_G22_NONE_P709;
+	return FString::FromInt(ColorSpace);
+};
 
-	static auto CVarHDREnable = IConsoleManager::Get().FindConsoleVariable(TEXT("Editor.HDRSupport"));
-	if (CVarHDREnable->GetInt() != 0)
+inline bool EnsureColorSpace(IDXGISwapChain* SwapChain, EDisplayColorGamut DisplayGamut, EDisplayOutputFormat OutputDevice, EPixelFormat PixelFormat)
+{
+	if (!SwapChain)
 	{
-		FlushRenderingCommands();
+		return false;
+	}
+	TRefCountPtr<IDXGIOutput> Output;
+	if (FAILED(SwapChain->GetContainingOutput(Output.GetInitReference())))
+	{
+		return false;
+	}
+	
+	TRefCountPtr<IDXGIOutput6> Output6;
+	if (FAILED(Output->QueryInterface(IID_PPV_ARGS(Output6.GetInitReference()))))
+	{
+		return false;
+	}
+	
+	DXGI_OUTPUT_DESC1 desc;
+	Output6->GetDesc1(&desc);
 
-		if (SwapChain)
-		{
-			TRefCountPtr<IDXGIOutput> Output;
-			if (SUCCEEDED(SwapChain->GetContainingOutput(Output.GetInitReference())))
-			{
-				TRefCountPtr<IDXGIOutput6> Output6;
-				if (SUCCEEDED(Output->QueryInterface(IID_PPV_ARGS(Output6.GetInitReference()))))
-				{
-					DXGI_OUTPUT_DESC1 desc;
-					Output6->GetDesc1(&desc);
 
-					if (desc.ColorSpace == DXGI_COLOR_SPACE_RGB_FULL_G2084_NONE_P2020)
-					{
-						// Display output is HDR10.
-						colorSpace = DXGI_COLOR_SPACE_RGB_FULL_G2084_NONE_P2020;
-					}
-				}
-			}
-		}
+	TRefCountPtr<IDXGISwapChain3> swapChain3;
+	if (FAILED(SwapChain->QueryInterface(IID_PPV_ARGS(swapChain3.GetInitReference()))))
+	{
+		return false;
+	}
 
-		TRefCountPtr<IDXGISwapChain3> swapChain3;
+	bool bHdrSupported = (desc.ColorSpace == DXGI_COLOR_SPACE_RGB_FULL_G2084_NONE_P2020);
 
-		if (SUCCEEDED(SwapChain->QueryInterface(IID_PPV_ARGS(swapChain3.GetInitReference()))))
+	DXGI_COLOR_SPACE_TYPE NewColorSpace = DXGI_COLOR_SPACE_RGB_FULL_G22_NONE_P709;	// sRGB;
+	const bool bPrimaries2020 = (DisplayGamut == EDisplayColorGamut::Rec2020_D65);
+
+	// See console variable r.HDR.Display.OutputDevice.
+	switch (OutputDevice)
+	{
+		// Gamma 2.2
+	case EDisplayOutputFormat::SDR_sRGB:
+	case EDisplayOutputFormat::SDR_Rec709:
+		NewColorSpace = bPrimaries2020 ? DXGI_COLOR_SPACE_RGB_FULL_G22_NONE_P2020 : DXGI_COLOR_SPACE_RGB_FULL_G22_NONE_P709;
+		break;
+
+		// Gamma ST.2084
+	case EDisplayOutputFormat::HDR_ACES_1000nit_ST2084:
+	case EDisplayOutputFormat::HDR_ACES_2000nit_ST2084:
+		NewColorSpace = DXGI_COLOR_SPACE_RGB_FULL_G2084_NONE_P2020;
+		ensure(bHdrSupported);
+		ensure(PixelFormat == PF_A2B10G10R10);
+		break;
+
+		// Gamma 1.0 (Linear)
+	case EDisplayOutputFormat::HDR_ACES_1000nit_ScRGB:
+	case EDisplayOutputFormat::HDR_ACES_2000nit_ScRGB:
+		// Linear. Still supports expanded color space with values >1.0f and <0.0f.
+		// The actual range is determined by the pixel format (e.g. a UNORM format can only ever have 0-1).
+		NewColorSpace = DXGI_COLOR_SPACE_RGB_FULL_G10_NONE_P709;
+		ensure(bHdrSupported);
+		ensure(PixelFormat == PF_FloatRGBA);
+		break;
+	}
+
+	{
+		UINT ColorSpaceSupport = 0;
+		HRESULT hr = swapChain3->CheckColorSpaceSupport(NewColorSpace, &ColorSpaceSupport);
+		FString ColorSpaceName = GetDXGIColorSpaceString(NewColorSpace);
+		if (SUCCEEDED(hr) && (ColorSpaceSupport & DXGI_SWAP_CHAIN_COLOR_SPACE_SUPPORT_FLAG_PRESENT))
 		{
-			UINT colorSpaceSupport = 0;
-		
-			if (SUCCEEDED(swapChain3->CheckColorSpaceSupport(colorSpace, &colorSpaceSupport)) && 
-				(colorSpaceSupport & DXGI_SWAP_CHAIN_COLOR_SPACE_SUPPORT_FLAG_PRESENT))
-			{
-				swapChain3->SetColorSpace1(colorSpace);
-			}
+			swapChain3->SetColorSpace1(NewColorSpace);
+			UE_LOG(LogD3D11RHI, Log, TEXT("Setting color space on swap chain: %s"), *ColorSpaceName);
+		}
+		else
+		{
+			UE_LOG(LogD3D11RHI, Error, TEXT("Warning: unabled to set color space %s to the swapchain: verify EDisplayOutputFormat / swapchain format"), *ColorSpaceName);
 		}
 	}
-	
-	if (colorSpace == DXGI_COLOR_SPACE_RGB_FULL_G2084_NONE_P2020)
+
+	return bHdrSupported;
+}
+
+// When a window has moved or resized we need to check whether it is on a HDR monitor or not. Set the correct color space of the monitor
+void FD3D11Viewport::CheckHDRMonitorStatus()
+{
+#if WITH_EDITOR
+
+	static auto CVarHDREnable = IConsoleManager::Get().FindConsoleVariable(TEXT("Editor.HDRSupport"));
+	bool bHDRSupported = false;
+	//if (CVarHDREnable->GetInt() != 0 || IsHDREnabled())
 	{
-		PixelColorSpace = EColorSpaceAndEOTF::ERec2020_PQ;
+		FlushRenderingCommands();
+
+		static const auto CVarHDRColorGamut = IConsoleManager::Get().FindTConsoleVariableDataInt(TEXT("r.HDR.Display.ColorGamut"));
+		const EDisplayColorGamut DisplayGamut = EDisplayColorGamut(CVarHDRColorGamut->GetValueOnAnyThread());
+		static const auto CVarHDROutputDevice = IConsoleManager::Get().FindTConsoleVariableDataInt(TEXT("r.HDR.Display.OutputDevice"));
+		const EDisplayOutputFormat OutputDevice = EDisplayOutputFormat(CVarHDROutputDevice->GetValueOnAnyThread());
+
+		bHDRSupported = EnsureColorSpace(SwapChain, DisplayGamut, OutputDevice, PixelFormat);
 	}
-	else
+	
 	{
 		PixelColorSpace =  EColorSpaceAndEOTF::ERec709_sRGB;
 	}
