undo changes to d3d11 viewport that were not supposed to be submitted

diff --git a/Engine/Source/Runtime/Windows/D3D11RHI/Private/Windows/WindowsD3D11Viewport.cpp b/Engine/Source/Runtime/Windows/D3D11RHI/Private/Windows/WindowsD3D11Viewport.cpp
--- a/Engine/Source/Runtime/Windows/D3D11RHI/Private/Windows/WindowsD3D11Viewport.cpp
+++ b/Engine/Source/Runtime/Windows/D3D11RHI/Private/Windows/WindowsD3D11Viewport.cpp
@@ -261,123 +261,57 @@ FD3D11Viewport::FD3D11Viewport(FD3D11DynamicRHI* InD3DRHI,HWND InWindowHandle,ui
 	});	
 }
 
-static const FString GetDXGIColorSpaceString(DXGI_COLOR_SPACE_TYPE ColorSpace)
-{
-	switch (ColorSpace)
-	{
-	case DXGI_COLOR_SPACE_RGB_FULL_G22_NONE_P709:
-		return TEXT("RGB_FULL_G22_NONE_P709");
-	case DXGI_COLOR_SPACE_RGB_FULL_G10_NONE_P709:
-		return TEXT("RGB_FULL_G10_NONE_P709");
-	case DXGI_COLOR_SPACE_RGB_FULL_G2084_NONE_P2020:
-		return TEXT("RGB_FULL_G2084_NONE_P2020");
-	case DXGI_COLOR_SPACE_RGB_FULL_G22_NONE_P2020:
-		return TEXT("RGB_FULL_G22_NONE_P2020");
-	default:
-		break;
-	}
-
-	return FString::FromInt(ColorSpace);
-};
-
-inline bool EnsureColorSpace(IDXGISwapChain* SwapChain, EDisplayColorGamut DisplayGamut, EDisplayOutputFormat OutputDevice, EPixelFormat PixelFormat)
+// When a window has moved or resized we need to check whether it is on a HDR monitor or not. Set the correct color space of the monitor
+void FD3D11Viewport::CheckHDRMonitorStatus()
 {
-	if (!SwapChain)
-	{
-		return false;
-	}
-	TRefCountPtr<IDXGIOutput> Output;
-	if (FAILED(SwapChain->GetContainingOutput(Output.GetInitReference())))
-	{
-		return false;
-	}
-	
-	TRefCountPtr<IDXGIOutput6> Output6;
-	if (FAILED(Output->QueryInterface(IID_PPV_ARGS(Output6.GetInitReference()))))
-	{
-		return false;
-	}
-	
-	DXGI_OUTPUT_DESC1 desc;
-	Output6->GetDesc1(&desc);
+#if WITH_EDITOR
 
+	DXGI_COLOR_SPACE_TYPE colorSpace = DXGI_COLOR_SPACE_RGB_FULL_G22_NONE_P709;
 
-	TRefCountPtr<IDXGISwapChain3> swapChain3;
-	if (FAILED(SwapChain->QueryInterface(IID_PPV_ARGS(swapChain3.GetInitReference()))))
+	static auto CVarHDREnable = IConsoleManager::Get().FindConsoleVariable(TEXT("Editor.HDRSupport"));
+	if (CVarHDREnable->GetInt() != 0)
 	{
-		return false;
-	}
+		FlushRenderingCommands();
 
-	bool bHdrSupported = (desc.ColorSpace == DXGI_COLOR_SPACE_RGB_FULL_G2084_NONE_P2020);
+		if (SwapChain)
+		{
+			TRefCountPtr<IDXGIOutput> Output;
+			if (SUCCEEDED(SwapChain->GetContainingOutput(Output.GetInitReference())))
+			{
+				TRefCountPtr<IDXGIOutput6> Output6;
+				if (SUCCEEDED(Output->QueryInterface(IID_PPV_ARGS(Output6.GetInitReference()))))
+				{
+					DXGI_OUTPUT_DESC1 desc;
+					Output6->GetDesc1(&desc);
 
-	DXGI_COLOR_SPACE_TYPE NewColorSpace = DXGI_COLOR_SPACE_RGB_FULL_G22_NONE_P709;	// sRGB;
-	const bool bPrimaries2020 = (DisplayGamut == EDisplayColorGamut::Rec2020_D65);
+					if (desc.ColorSpace == DXGI_COLOR_SPACE_RGB_FULL_G2084_NONE_P2020)
+					{
+						// Display output is HDR10.
+						colorSpace = DXGI_COLOR_SPACE_RGB_FULL_G2084_NONE_P2020;
+					}
+				}
+			}
+		}
 
-	// See console variable r.HDR.Display.OutputDevice.
-	switch (OutputDevice)
-	{
-		// Gamma 2.2
-	case EDisplayOutputFormat::SDR_sRGB:
-	case EDisplayOutputFormat::SDR_Rec709:
-		NewColorSpace = bPrimaries2020 ? DXGI_COLOR_SPACE_RGB_FULL_G22_NONE_P2020 : DXGI_COLOR_SPACE_RGB_FULL_G22_NONE_P709;
-		break;
-
-		// Gamma ST.2084
-	case EDisplayOutputFormat::HDR_ACES_1000nit_ST2084:
-	case EDisplayOutputFormat::HDR_ACES_2000nit_ST2084:
-		NewColorSpace = DXGI_COLOR_SPACE_RGB_FULL_G2084_NONE_P2020;
-		ensure(bHdrSupported);
-		ensure(PixelFormat == PF_A2B10G10R10);
-		break;
-
-		// Gamma 1.0 (Linear)
-	case EDisplayOutputFormat::HDR_ACES_1000nit_ScRGB:
-	case EDisplayOutputFormat::HDR_ACES_2000nit_ScRGB:
-		// Linear. Still supports expanded color space with values >1.0f and <0.0f.
-		// The actual range is determined by the pixel format (e.g. a UNORM format can only ever have 0-1).
-		NewColorSpace = DXGI_COLOR_SPACE_RGB_FULL_G10_NONE_P709;
-		ensure(bHdrSupported);
-		ensure(PixelFormat == PF_FloatRGBA);
-		break;
-	}
+		TRefCountPtr<IDXGISwapChain3> swapChain3;
 
-	{
-		UINT ColorSpaceSupport = 0;
-		HRESULT hr = swapChain3->CheckColorSpaceSupport(NewColorSpace, &ColorSpaceSupport);
-		FString ColorSpaceName = GetDXGIColorSpaceString(NewColorSpace);
-		if (SUCCEEDED(hr) && (ColorSpaceSupport & DXGI_SWAP_CHAIN_COLOR_SPACE_SUPPORT_FLAG_PRESENT))
+		if (SUCCEEDED(SwapChain->QueryInterface(IID_PPV_ARGS(swapChain3.GetInitReference()))))
 		{
-			swapChain3->SetColorSpace1(NewColorSpace);
-			UE_LOG(LogD3D11RHI, Log, TEXT("Setting color space on swap chain: %s"), *ColorSpaceName);
-		}
-		else
-		{
-			UE_LOG(LogD3D11RHI, Error, TEXT("Warning: unabled to set color space %s to the swapchain: verify EDisplayOutputFormat / swapchain format"), *ColorSpaceName);
+			UINT colorSpaceSupport = 0;
+		
+			if (SUCCEEDED(swapChain3->CheckColorSpaceSupport(colorSpace, &colorSpaceSupport)) && 
+				(colorSpaceSupport & DXGI_SWAP_CHAIN_COLOR_SPACE_SUPPORT_FLAG_PRESENT))
+			{
+				swapChain3->SetColorSpace1(colorSpace);
+			}
 		}
 	}
-
-	return bHdrSupported;
-}
-
-// When a window has moved or resized we need to check whether it is on a HDR monitor or not. Set the correct color space of the monitor
-void FD3D11Viewport::CheckHDRMonitorStatus()
-{
-#if WITH_EDITOR
-
-	static auto CVarHDREnable = IConsoleManager::Get().FindConsoleVariable(TEXT("Editor.HDRSupport"));
-	bool bHDRSupported = false;
-	//if (CVarHDREnable->GetInt() != 0 || IsHDREnabled())
+	
+	if (colorSpace == DXGI_COLOR_SPACE_RGB_FULL_G2084_NONE_P2020)
 	{
-		FlushRenderingCommands();
-
-		static const auto CVarHDRColorGamut = IConsoleManager::Get().FindTConsoleVariableDataInt(TEXT("r.HDR.Display.ColorGamut"));
-		const EDisplayColorGamut DisplayGamut = EDisplayColorGamut(CVarHDRColorGamut->GetValueOnAnyThread());
-		static const auto CVarHDROutputDevice = IConsoleManager::Get().FindTConsoleVariableDataInt(TEXT("r.HDR.Display.OutputDevice"));
-		const EDisplayOutputFormat OutputDevice = EDisplayOutputFormat(CVarHDROutputDevice->GetValueOnAnyThread());
-
-		bHDRSupported = EnsureColorSpace(SwapChain, DisplayGamut, OutputDevice, PixelFormat);
+		PixelColorSpace = EColorSpaceAndEOTF::ERec2020_PQ;
 	}
-	
+	else
 	{
 		PixelColorSpace =  EColorSpaceAndEOTF::ERec709_sRGB;
 	}
@@ -432,4 +366,4 @@ void FD3D11Viewport::ResetSwapChainInternal(bool bIgnoreFocus)
 	}
 }
 
-#include "Windows/HideWindowsPlatformTypes.h"
+#include "Windows/HideWindowsPlatformTypes.h"
\ No newline at end of file
