Fixed inconsistencies between RHI's with CopyTexture. Replaced implicit check for equality with an assert at the command list.

diff --git a/Engine/Source/Runtime/Apple/AGXRHI/Private/AGXTexture.cpp b/Engine/Source/Runtime/Apple/AGXRHI/Private/AGXTexture.cpp
--- a/Engine/Source/Runtime/Apple/AGXRHI/Private/AGXTexture.cpp
+++ b/Engine/Source/Runtime/Apple/AGXRHI/Private/AGXTexture.cpp
@@ -1998,12 +1998,6 @@ inline bool AGXRHICopyTexutre_IsTextureFormatCompatible(EPixelFormat SrcFmt, EPi
 
 void FAGXRHICommandContext::RHICopyTexture(FRHITexture* SourceTextureRHI, FRHITexture* DestTextureRHI, const FRHICopyTextureInfo& CopyInfo)
 {
-	if (!SourceTextureRHI || !DestTextureRHI || SourceTextureRHI == DestTextureRHI)
-	{
-		// no need to do anything (silently ignored)
-		return;
-	}
-
 	@autoreleasepool {
 		check(SourceTextureRHI);
 		check(DestTextureRHI);
diff --git a/Engine/Source/Runtime/Apple/MetalRHI/Private/MetalTexture.cpp b/Engine/Source/Runtime/Apple/MetalRHI/Private/MetalTexture.cpp
--- a/Engine/Source/Runtime/Apple/MetalRHI/Private/MetalTexture.cpp
+++ b/Engine/Source/Runtime/Apple/MetalRHI/Private/MetalTexture.cpp
@@ -2056,12 +2056,6 @@ inline bool MetalRHICopyTexutre_IsTextureFormatCompatible(EPixelFormat SrcFmt, E
 
 void FMetalRHICommandContext::RHICopyTexture(FRHITexture* SourceTextureRHI, FRHITexture* DestTextureRHI, const FRHICopyTextureInfo& CopyInfo)
 {
-	if (!SourceTextureRHI || !DestTextureRHI || SourceTextureRHI == DestTextureRHI)
-	{
-		// no need to do anything (silently ignored)
-		return;
-	}
-
 	@autoreleasepool {
 		check(SourceTextureRHI);
 		check(DestTextureRHI);
diff --git a/Engine/Source/Runtime/D3D12RHI/Private/D3D12Texture.cpp b/Engine/Source/Runtime/D3D12RHI/Private/D3D12Texture.cpp
--- a/Engine/Source/Runtime/D3D12RHI/Private/D3D12Texture.cpp
+++ b/Engine/Source/Runtime/D3D12RHI/Private/D3D12Texture.cpp
@@ -2907,7 +2907,17 @@ void FD3D12CommandContext::RHICopyTexture(FRHITexture* SourceTextureRHI, FRHITex
 
 	const bool bReadback = EnumHasAnyFlags(DestTextureRHI->GetFlags(), TexCreate_CPUReadback);
 
-	if (CopyInfo.Size != FIntVector::ZeroValue || bReadback)
+	const FRHITextureDesc& SourceDesc = SourceTextureRHI->GetDesc();
+	const FRHITextureDesc& DestDesc = DestTextureRHI->GetDesc();
+
+	const bool bAllPixels =
+		SourceDesc.GetSize() == DestDesc.GetSize() && (CopyInfo.Size == FIntVector::ZeroValue || CopyInfo.Size == SourceDesc.GetSize());
+
+	const bool bAllSubresources =
+		SourceDesc.NumMips   == DestDesc.NumMips   && SourceDesc.NumMips   == CopyInfo.NumMips    &&
+		SourceDesc.ArraySize == DestDesc.ArraySize && SourceDesc.ArraySize == CopyInfo.NumSlices;
+
+	if (!bAllPixels || !bAllSubresources || bReadback)
 	{
 		// Interpret zero size as source size
 		const FIntVector CopySize = CopyInfo.Size == FIntVector::ZeroValue ? SourceTextureRHI->GetSizeXYZ() : CopyInfo.Size;
diff --git a/Engine/Source/Runtime/OpenGLDrv/Private/OpenGLRenderTarget.cpp b/Engine/Source/Runtime/OpenGLDrv/Private/OpenGLRenderTarget.cpp
--- a/Engine/Source/Runtime/OpenGLDrv/Private/OpenGLRenderTarget.cpp
+++ b/Engine/Source/Runtime/OpenGLDrv/Private/OpenGLRenderTarget.cpp
@@ -1165,19 +1165,19 @@ void FOpenGLDynamicRHI::RHIEndRenderPass()
 
 	for (int32 Index = 0; Index < MaxSimultaneousRenderTargets; ++Index)
 	{
-		if (!RenderPassInfo.ColorRenderTargets[Index].RenderTarget)
-		{
-			break;
-		}
-		if (RenderPassInfo.ColorRenderTargets[Index].ResolveTarget)
+		const auto& RTV = RenderPassInfo.ColorRenderTargets[Index];
+
+		if (RTV.RenderTarget && RTV.ResolveTarget && RTV.RenderTarget != RTV.ResolveTarget)
 		{
-			RHICopyToResolveTarget(RenderPassInfo.ColorRenderTargets[Index].RenderTarget, RenderPassInfo.ColorRenderTargets[Index].ResolveTarget, RenderPassInfo.ResolveParameters);
+			RHICopyToResolveTarget(RTV.RenderTarget, RTV.ResolveTarget, RenderPassInfo.ResolveParameters);
 		}
 	}
 
-	if (RenderPassInfo.DepthStencilRenderTarget.DepthStencilTarget && RenderPassInfo.DepthStencilRenderTarget.ResolveTarget)
+	const auto& DSV = RenderPassInfo.DepthStencilRenderTarget;
+
+	if (DSV.DepthStencilTarget && DSV.ResolveTarget && DSV.DepthStencilTarget != DSV.ResolveTarget)
 	{
-		RHICopyToResolveTarget(RenderPassInfo.DepthStencilRenderTarget.DepthStencilTarget, RenderPassInfo.DepthStencilRenderTarget.ResolveTarget, RenderPassInfo.ResolveParameters);
+		RHICopyToResolveTarget(DSV.DepthStencilTarget, DSV.ResolveTarget, RenderPassInfo.ResolveParameters);
 	}
 
 	// Drop depth and stencil to avoid export
diff --git a/Engine/Source/Runtime/RHI/Public/RHICommandList.h b/Engine/Source/Runtime/RHI/Public/RHICommandList.h
--- a/Engine/Source/Runtime/RHI/Public/RHICommandList.h
+++ b/Engine/Source/Runtime/RHI/Public/RHICommandList.h
@@ -3552,6 +3552,8 @@ public:
 
 	FORCEINLINE_DEBUGGABLE void CopyTexture(FRHITexture* SourceTextureRHI, FRHITexture* DestTextureRHI, const FRHICopyTextureInfo& CopyInfo)
 	{
+		check(SourceTextureRHI && DestTextureRHI);
+		check(SourceTextureRHI != DestTextureRHI);
 		check(IsOutsideRenderPass());
 		if (GRHISupportsCopyToTextureMultipleMips)
 		{
diff --git a/Engine/Source/Runtime/Windows/D3D11RHI/Private/D3D11Texture.cpp b/Engine/Source/Runtime/Windows/D3D11RHI/Private/D3D11Texture.cpp
--- a/Engine/Source/Runtime/Windows/D3D11RHI/Private/D3D11Texture.cpp
+++ b/Engine/Source/Runtime/Windows/D3D11RHI/Private/D3D11Texture.cpp
@@ -2071,12 +2071,6 @@ FTextureRHIRef FD3D11DynamicRHI::RHICreateAliasedTexture(FTextureRHIRef& SrcText
 
 void FD3D11DynamicRHI::RHICopyTexture(FRHITexture* SourceTextureRHI, FRHITexture* DestTextureRHI, const FRHICopyTextureInfo& CopyInfo)
 {
-	if (!SourceTextureRHI || !DestTextureRHI || SourceTextureRHI == DestTextureRHI)
-	{
-		// no need to do anything (silently ignored)
-		return;
-	}
-
 	FRHICommandList_RecursiveHazardous RHICmdList(this);	
 
 	FD3D11Texture* SourceTexture = GetD3D11TextureFromRHITexture(SourceTextureRHI);
@@ -2086,7 +2080,17 @@ void FD3D11DynamicRHI::RHICopyTexture(FRHITexture* SourceTextureRHI, FRHITexture
 
 	GPUProfilingData.RegisterGPUWork();
 
-	if (CopyInfo.Size != FIntVector::ZeroValue)
+	const FRHITextureDesc& SourceDesc = SourceTextureRHI->GetDesc();
+	const FRHITextureDesc& DestDesc = DestTextureRHI->GetDesc();
+
+	const bool bAllPixels =
+		SourceDesc.GetSize() == DestDesc.GetSize() && (CopyInfo.Size == FIntVector::ZeroValue || CopyInfo.Size == SourceDesc.GetSize());
+
+	const bool bAllSubresources =
+		SourceDesc.NumMips == DestDesc.NumMips && SourceDesc.NumMips == CopyInfo.NumMips &&
+		SourceDesc.ArraySize == DestDesc.ArraySize && SourceDesc.ArraySize == CopyInfo.NumSlices;
+
+	if (!bAllPixels || !bAllSubresources)
 	{
 		const FPixelFormatInfo& PixelFormatInfo = GPixelFormats[SourceTextureRHI->GetFormat()];
 
