d3d11 editor HDR support
Converge the way we do HDR in editor for both d3d11 and d3d12. This assumes that IDXGIOutput6 is available, but AFAICT, Editor.HDRSupport / FD3D11Viewport::CheckHDRMonitorStatus was requiring it anyway (https://p4-swarm.epicgames.net/changes/8174116 - WindowsD3D11Viewport.cpp#20 FD3D11Viewport::CheckHDRMonitorStatus).
Notes for future self: Editor was initially done with HDR10 (see FD3D11Viewport constructor where it forces swapchain format to DXGI_FORMAT_R10G10B10A2_UNORM). However, it got de-hardcoded with 17369147 which was a good idea, but unfortunately GRHIHDRDisplayOutputFormat is left to the default value on d3d11, which is PF_FloatRGBA.
CheckColorSpaceSupport would now fail, not returning DXGI_SWAP_CHAIN_COLOR_SPACE_SUPPORT_FLAG_PRESENT (PF_FloatRGBA only supports scRGB, not DXGI_COLOR_SPACE_RGB_FULL_G2084_NONE_P2020) and we would not setup the swapchain colorspace properly.
Since it looks like DXGI 1.6 was always needed for editor, I'm changing the code to be the same between d3d11 and d3d12. Both editors now use scRGB for HDR, which is consistent with ChooseHDRDeviceAndColorGamut

diff --git a/Engine/Source/Runtime/Windows/D3D11RHI/Private/Windows/WindowsD3D11Device.cpp b/Engine/Source/Runtime/Windows/D3D11RHI/Private/Windows/WindowsD3D11Device.cpp
--- a/Engine/Source/Runtime/Windows/D3D11RHI/Private/Windows/WindowsD3D11Device.cpp
+++ b/Engine/Source/Runtime/Windows/D3D11RHI/Private/Windows/WindowsD3D11Device.cpp
@@ -561,6 +561,49 @@ static bool SupportsHDROutput(FD3D11DynamicRHI* D3DRHI)
 	// Default to primary display
 	D3DRHI->SetHDRDetectedDisplayIndices(0, 0);
 	
+#if WITH_EDITOR
+	// Determines if any displays support HDR
+	bool bSupportsHDROutput = false;
+	{
+		const FD3D11Adapter& Adapter = D3DRHI->GetAdapter();
+		IDXGIAdapter* DXGIAdapter = Adapter.DXGIAdapter;
+		if (!DXGIAdapter)
+		{
+			return false;
+		}
+
+		for (uint32 DisplayIndex = 0; true; ++DisplayIndex)
+		{
+			TRefCountPtr<IDXGIOutput> DXGIOutput;
+			if (S_OK != DXGIAdapter->EnumOutputs(DisplayIndex, DXGIOutput.GetInitReference()))
+			{
+				break;
+			}
+
+			TRefCountPtr<IDXGIOutput6> Output6;
+			if (SUCCEEDED(DXGIOutput->QueryInterface(IID_PPV_ARGS(Output6.GetInitReference()))))
+			{
+				DXGI_OUTPUT_DESC1 OutputDesc;
+				VERIFYD3D11RESULT(Output6->GetDesc1(&OutputDesc));
+
+				// Check for HDR support on the display.
+				const bool bDisplaySupportsHDROutput = (OutputDesc.ColorSpace == DXGI_COLOR_SPACE_RGB_FULL_G2084_NONE_P2020);
+				if (bDisplaySupportsHDROutput)
+				{
+					UE_LOG(LogD3D11RHI, Log, TEXT("HDR output is supported on adapter %i, display %u:"), 0, DisplayIndex);
+					UE_LOG(LogD3D11RHI, Log, TEXT("\t\tMinLuminance = %f"), OutputDesc.MinLuminance);
+					UE_LOG(LogD3D11RHI, Log, TEXT("\t\tMaxLuminance = %f"), OutputDesc.MaxLuminance);
+					UE_LOG(LogD3D11RHI, Log, TEXT("\t\tMaxFullFrameLuminance = %f"), OutputDesc.MaxFullFrameLuminance);
+
+					bSupportsHDROutput = true;
+				}
+			}
+		}
+	}
+
+	return bSupportsHDROutput;
+#else
+
 	// Grab the adapter
 	TRefCountPtr<IDXGIDevice> DXGIDevice;
 	VERIFYD3D11RESULT(Direct3DDevice->QueryInterface(IID_IDXGIDevice, (void**)DXGIDevice.GetInitReference()));
@@ -653,6 +696,7 @@ static bool SupportsHDROutput(FD3D11DynamicRHI* D3DRHI)
 	}
 
 	return false;
+#endif
 }
 
 static bool IsDeviceOverclocked()
diff --git a/Engine/Source/Runtime/Windows/D3D11RHI/Private/Windows/WindowsD3D11Viewport.cpp b/Engine/Source/Runtime/Windows/D3D11RHI/Private/Windows/WindowsD3D11Viewport.cpp
--- a/Engine/Source/Runtime/Windows/D3D11RHI/Private/Windows/WindowsD3D11Viewport.cpp
+++ b/Engine/Source/Runtime/Windows/D3D11RHI/Private/Windows/WindowsD3D11Viewport.cpp
@@ -261,57 +261,96 @@ FD3D11Viewport::FD3D11Viewport(FD3D11DynamicRHI* InD3DRHI,HWND InWindowHandle,ui
 	});	
 }
 
+static const FString GetDXGIColorSpaceString(DXGI_COLOR_SPACE_TYPE ColorSpace)
+{
+	switch (ColorSpace)
+	{
+	case DXGI_COLOR_SPACE_RGB_FULL_G22_NONE_P709:
+		return TEXT("RGB_FULL_G22_NONE_P709");
+	case DXGI_COLOR_SPACE_RGB_FULL_G10_NONE_P709:
+		return TEXT("RGB_FULL_G10_NONE_P709");
+	case DXGI_COLOR_SPACE_RGB_FULL_G2084_NONE_P2020:
+		return TEXT("RGB_FULL_G2084_NONE_P2020");
+	case DXGI_COLOR_SPACE_RGB_FULL_G22_NONE_P2020:
+		return TEXT("RGB_FULL_G22_NONE_P2020");
+	default:
+		break;
+	}
+
+	return FString::FromInt(ColorSpace);
+};
+
+inline void EnsureColorSpace(IDXGISwapChain* SwapChain, EDisplayColorGamut DisplayGamut, EDisplayOutputFormat OutputDevice, EPixelFormat PixelFormat)
+{
+	TRefCountPtr<IDXGISwapChain3> swapChain3;
+	if (FAILED(SwapChain->QueryInterface(IID_PPV_ARGS(swapChain3.GetInitReference()))))
+	{
+		return;
+	}
+
+	DXGI_COLOR_SPACE_TYPE NewColorSpace = DXGI_COLOR_SPACE_RGB_FULL_G22_NONE_P709;	// sRGB;
+	const bool bPrimaries2020 = (DisplayGamut == EDisplayColorGamut::Rec2020_D65);
+
+	// See console variable r.HDR.Display.OutputDevice.
+	switch (OutputDevice)
+	{
+		// Gamma 2.2
+	case EDisplayOutputFormat::SDR_sRGB:
+	case EDisplayOutputFormat::SDR_Rec709:
+		NewColorSpace = bPrimaries2020 ? DXGI_COLOR_SPACE_RGB_FULL_G22_NONE_P2020 : DXGI_COLOR_SPACE_RGB_FULL_G22_NONE_P709;
+		break;
+
+		// Gamma ST.2084
+	case EDisplayOutputFormat::HDR_ACES_1000nit_ST2084:
+	case EDisplayOutputFormat::HDR_ACES_2000nit_ST2084:
+		NewColorSpace = DXGI_COLOR_SPACE_RGB_FULL_G2084_NONE_P2020;
+		ensure(PixelFormat == PF_A2B10G10R10);
+		break;
+
+		// Gamma 1.0 (Linear)
+	case EDisplayOutputFormat::HDR_ACES_1000nit_ScRGB:
+	case EDisplayOutputFormat::HDR_ACES_2000nit_ScRGB:
+		// Linear. Still supports expanded color space with values >1.0f and <0.0f.
+		// The actual range is determined by the pixel format (e.g. a UNORM format can only ever have 0-1).
+		NewColorSpace = DXGI_COLOR_SPACE_RGB_FULL_G10_NONE_P709;
+		ensure(PixelFormat == PF_FloatRGBA);
+		break;
+	}
+
+	{
+		UINT ColorSpaceSupport = 0;
+		HRESULT hr = swapChain3->CheckColorSpaceSupport(NewColorSpace, &ColorSpaceSupport);
+		FString ColorSpaceName = GetDXGIColorSpaceString(NewColorSpace);
+		if (SUCCEEDED(hr) && (ColorSpaceSupport & DXGI_SWAP_CHAIN_COLOR_SPACE_SUPPORT_FLAG_PRESENT))
+		{
+			swapChain3->SetColorSpace1(NewColorSpace);
+			UE_LOG(LogD3D11RHI, Log, TEXT("Setting color space on swap chain: %s"), *ColorSpaceName);
+		}
+		else
+		{
+			UE_LOG(LogD3D11RHI, Error, TEXT("Warning: unabled to set color space %s to the swapchain: verify EDisplayOutputFormat / swapchain format"), *ColorSpaceName);
+		}
+	}
+}
+
 // When a window has moved or resized we need to check whether it is on a HDR monitor or not. Set the correct color space of the monitor
 void FD3D11Viewport::CheckHDRMonitorStatus()
 {
 #if WITH_EDITOR
 
-	DXGI_COLOR_SPACE_TYPE colorSpace = DXGI_COLOR_SPACE_RGB_FULL_G22_NONE_P709;
-
 	static auto CVarHDREnable = IConsoleManager::Get().FindConsoleVariable(TEXT("Editor.HDRSupport"));
 	if (CVarHDREnable->GetInt() != 0)
 	{
 		FlushRenderingCommands();
 
-		if (SwapChain)
-		{
-			TRefCountPtr<IDXGIOutput> Output;
-			if (SUCCEEDED(SwapChain->GetContainingOutput(Output.GetInitReference())))
-			{
-				TRefCountPtr<IDXGIOutput6> Output6;
-				if (SUCCEEDED(Output->QueryInterface(IID_PPV_ARGS(Output6.GetInitReference()))))
-				{
-					DXGI_OUTPUT_DESC1 desc;
-					Output6->GetDesc1(&desc);
-
-					if (desc.ColorSpace == DXGI_COLOR_SPACE_RGB_FULL_G2084_NONE_P2020)
-					{
-						// Display output is HDR10.
-						colorSpace = DXGI_COLOR_SPACE_RGB_FULL_G2084_NONE_P2020;
-					}
-				}
-			}
-		}
-
-		TRefCountPtr<IDXGISwapChain3> swapChain3;
+		static const auto CVarHDRColorGamut = IConsoleManager::Get().FindTConsoleVariableDataInt(TEXT("r.HDR.Display.ColorGamut"));
+		const EDisplayColorGamut DisplayGamut = EDisplayColorGamut(CVarHDRColorGamut->GetValueOnAnyThread());
+		static const auto CVarHDROutputDevice = IConsoleManager::Get().FindTConsoleVariableDataInt(TEXT("r.HDR.Display.OutputDevice"));
+		const EDisplayOutputFormat OutputDevice = EDisplayOutputFormat(CVarHDROutputDevice->GetValueOnAnyThread());
 
-		if (SUCCEEDED(SwapChain->QueryInterface(IID_PPV_ARGS(swapChain3.GetInitReference()))))
-		{
-			UINT colorSpaceSupport = 0;
-		
-			if (SUCCEEDED(swapChain3->CheckColorSpaceSupport(colorSpace, &colorSpaceSupport)) && 
-				(colorSpaceSupport & DXGI_SWAP_CHAIN_COLOR_SPACE_SUPPORT_FLAG_PRESENT))
-			{
-				swapChain3->SetColorSpace1(colorSpace);
-			}
-		}
+		EnsureColorSpace(SwapChain, DisplayGamut, OutputDevice, PixelFormat);
 	}
 	
-	if (colorSpace == DXGI_COLOR_SPACE_RGB_FULL_G2084_NONE_P2020)
-	{
-		PixelColorSpace = EColorSpaceAndEOTF::ERec2020_PQ;
-	}
-	else
 	{
 		PixelColorSpace =  EColorSpaceAndEOTF::ERec709_sRGB;
 	}
