Fix a texture array size comparison error that caused GPU crashes and/or corrupt textures.

diff --git a/Engine/Source/Runtime/D3D12RHI/Private/D3D12Texture.cpp b/Engine/Source/Runtime/D3D12RHI/Private/D3D12Texture.cpp
--- a/Engine/Source/Runtime/D3D12RHI/Private/D3D12Texture.cpp
+++ b/Engine/Source/Runtime/D3D12RHI/Private/D3D12Texture.cpp
@@ -2911,7 +2911,7 @@ void FD3D12CommandContext::RHICopyTexture(FRHITexture* SourceTextureRHI, FRHITex
 	const FRHITextureDesc& DestDesc = DestTextureRHI->GetDesc();
 	
 	const uint16 SourceArraySize = SourceDesc.ArraySize * (SourceDesc.IsTextureCube() ? 6 : 1);
-	const uint16 DestArraySize   = SourceDesc.ArraySize * (DestDesc.IsTextureCube()   ? 6 : 1);
+	const uint16 DestArraySize   = DestDesc.ArraySize   * (DestDesc.IsTextureCube()   ? 6 : 1);
 
 	const bool bAllPixels =
 		SourceDesc.GetSize() == DestDesc.GetSize() && (CopyInfo.Size == FIntVector::ZeroValue || CopyInfo.Size == SourceDesc.GetSize());
diff --git a/Engine/Source/Runtime/Windows/D3D11RHI/Private/D3D11Texture.cpp b/Engine/Source/Runtime/Windows/D3D11RHI/Private/D3D11Texture.cpp
--- a/Engine/Source/Runtime/Windows/D3D11RHI/Private/D3D11Texture.cpp
+++ b/Engine/Source/Runtime/Windows/D3D11RHI/Private/D3D11Texture.cpp
@@ -2077,7 +2077,7 @@ void FD3D11DynamicRHI::RHICopyTexture(FRHITexture* SourceTextureRHI, FRHITexture
 	const FRHITextureDesc& DestDesc = DestTextureRHI->GetDesc();
 
 	const uint16 SourceArraySize = SourceDesc.ArraySize * (SourceDesc.IsTextureCube() ? 6 : 1);
-	const uint16 DestArraySize   = SourceDesc.ArraySize * (DestDesc.IsTextureCube()   ? 6 : 1);
+	const uint16 DestArraySize   = DestDesc.ArraySize   * (DestDesc.IsTextureCube()   ? 6 : 1);
 
 	const bool bAllPixels =
 		SourceDesc.GetSize() == DestDesc.GetSize() && (CopyInfo.Size == FIntVector::ZeroValue || CopyInfo.Size == SourceDesc.GetSize());
