Minor clean up of static uniform buffer handling.

diff --git a/Engine/Source/Runtime/Apple/AGXRHI/Private/Shaders/Types/Templates/AGXBaseShader.h b/Engine/Source/Runtime/Apple/AGXRHI/Private/Shaders/Types/Templates/AGXBaseShader.h
--- a/Engine/Source/Runtime/Apple/AGXRHI/Private/Shaders/Types/Templates/AGXBaseShader.h
+++ b/Engine/Source/Runtime/Apple/AGXRHI/Private/Shaders/Types/Templates/AGXBaseShader.h
@@ -13,7 +13,7 @@
 #include "Shaders/Debugging/AGXShaderDebugCache.h"
 #include "Shaders/AGXCompiledShaderKey.h"
 #include "Shaders/AGXCompiledShaderCache.h"
-
+#include "RHICoreShader.h"
 
 //------------------------------------------------------------------------------
 
@@ -335,19 +335,7 @@ void TAGXBaseShader<BaseResourceType, ShaderType>::Init(TArrayView<const uint8>
 	UniformBuffersCopyInfo = Header.UniformBuffersCopyInfo;
 	SideTableBinding = Header.SideTable;
 
-	StaticSlots.Reserve(Bindings.ShaderResourceTable.ResourceTableLayoutHashes.Num());
-
-	for (uint32 LayoutHash : Bindings.ShaderResourceTable.ResourceTableLayoutHashes)
-	{
-		if (const FShaderParametersMetadata* Metadata = FindUniformBufferStructByLayoutHash(LayoutHash))
-		{
-			StaticSlots.Add(Metadata->GetLayout().StaticSlot);
-		}
-		else
-		{
-			StaticSlots.Add(MAX_UNIFORM_BUFFER_STATIC_SLOTS);
-		}
-	}
+	UE::RHICore::InitStaticUniformBufferSlots(StaticSlots, Bindings.ShaderResourceTable);
 }
 
 template<typename BaseResourceType, int32 ShaderType>
diff --git a/Engine/Source/Runtime/Apple/MetalRHI/Private/Shaders/Types/Templates/MetalBaseShader.h b/Engine/Source/Runtime/Apple/MetalRHI/Private/Shaders/Types/Templates/MetalBaseShader.h
--- a/Engine/Source/Runtime/Apple/MetalRHI/Private/Shaders/Types/Templates/MetalBaseShader.h
+++ b/Engine/Source/Runtime/Apple/MetalRHI/Private/Shaders/Types/Templates/MetalBaseShader.h
@@ -13,7 +13,7 @@
 #include "Shaders/Debugging/MetalShaderDebugCache.h"
 #include "Shaders/MetalCompiledShaderKey.h"
 #include "Shaders/MetalCompiledShaderCache.h"
-
+#include "RHICoreShader.h"
 
 //------------------------------------------------------------------------------
 
@@ -341,19 +341,7 @@ void TMetalBaseShader<BaseResourceType, ShaderType>::Init(TArrayView<const uint8
 	UniformBuffersCopyInfo = Header.UniformBuffersCopyInfo;
 	SideTableBinding = Header.SideTable;
 
-	StaticSlots.Reserve(Bindings.ShaderResourceTable.ResourceTableLayoutHashes.Num());
-
-	for (uint32 LayoutHash : Bindings.ShaderResourceTable.ResourceTableLayoutHashes)
-	{
-		if (const FShaderParametersMetadata* Metadata = FindUniformBufferStructByLayoutHash(LayoutHash))
-		{
-			StaticSlots.Add(Metadata->GetLayout().StaticSlot);
-		}
-		else
-		{
-			StaticSlots.Add(MAX_UNIFORM_BUFFER_STATIC_SLOTS);
-		}
-	}
+	UE::RHICore::InitStaticUniformBufferSlots(StaticSlots, Bindings.ShaderResourceTable);
 }
 
 template<typename BaseResourceType, int32 ShaderType>
diff --git a/Engine/Source/Runtime/D3D12RHI/Private/D3D12Shaders.cpp b/Engine/Source/Runtime/D3D12RHI/Private/D3D12Shaders.cpp
--- a/Engine/Source/Runtime/D3D12RHI/Private/D3D12Shaders.cpp
+++ b/Engine/Source/Runtime/D3D12RHI/Private/D3D12Shaders.cpp
@@ -5,6 +5,7 @@
 =============================================================================*/
 
 #include "D3D12RHIPrivate.h"
+#include "RHICoreShader.h"
 
 template <typename TShaderType>
 static inline bool ReadShaderOptionalData(FShaderCodeReader& InShaderCode, TShaderType& OutShader)
@@ -107,23 +108,7 @@ TShaderType* CreateStandardShader(TArrayView<const uint8> InCode)
 		return nullptr;
 	}
 
-	// InitUniformBufferStaticSlots
-
-	const FBaseShaderResourceTable& SRT = Shader->ShaderResourceTable;
-	Shader->StaticSlots.Reserve(SRT.ResourceTableLayoutHashes.Num());
-
-	for (uint32 LayoutHash : SRT.ResourceTableLayoutHashes)
-	{
-		if (const FShaderParametersMetadata* Metadata = FindUniformBufferStructByLayoutHash(LayoutHash))
-		{
-			Shader->StaticSlots.Add(Metadata->GetLayout().StaticSlot);
-		}
-		else
-		{
-			Shader->StaticSlots.Add(MAX_UNIFORM_BUFFER_STATIC_SLOTS);
-		}
-	}
-
+	UE::RHICore::InitStaticUniformBufferSlots(Shader->StaticSlots, Shader->ShaderResourceTable);
 	return Shader;
 }
 
@@ -203,6 +188,8 @@ FRayTracingShaderRHIRef FD3D12DynamicRHI::RHICreateRayTracingShader(TArrayView<c
 		return nullptr;
 	}
 
+	UE::RHICore::InitStaticUniformBufferSlots(Shader->StaticSlots, Shader->ShaderResourceTable);
+
 	FD3D12Adapter& Adapter = GetAdapter();
 
 #if USE_STATIC_ROOT_SIGNATURE
diff --git a/Engine/Source/Runtime/OpenGLDrv/Private/OpenGLShaders.cpp b/Engine/Source/Runtime/OpenGLDrv/Private/OpenGLShaders.cpp
--- a/Engine/Source/Runtime/OpenGLDrv/Private/OpenGLShaders.cpp
+++ b/Engine/Source/Runtime/OpenGLDrv/Private/OpenGLShaders.cpp
@@ -15,6 +15,7 @@
 #include "GlobalShader.h"
 #include "SceneUtils.h"
 #include "PsoLruCache.h"
+#include "RHICoreShader.h"
 
 #define CHECK_FOR_GL_SHADERS_TO_REPLACE 0
 
@@ -685,22 +686,8 @@ ShaderType* CompileOpenGLShader(const FOpenGLCodeHeader& Header, const FOpenGLCo
 	Shader->Resource = Resource;
 	Shader->Bindings = Header.Bindings;
 	Shader->UniformBuffersCopyInfo = Header.UniformBuffersCopyInfo;
-	Shader->StaticSlots.Reserve(Header.Bindings.ShaderResourceTable.ResourceTableLayoutHashes.Num());
 
-	for (uint32 LayoutHash : Header.Bindings.ShaderResourceTable.ResourceTableLayoutHashes)
-	{
-		if (const FShaderParametersMetadata* Metadata = FindUniformBufferStructByLayoutHash(LayoutHash))
-		{
-			Shader->StaticSlots.Add(Metadata->GetLayout().StaticSlot);
-		}
-		else
-		{
-			Shader->StaticSlots.Add(MAX_UNIFORM_BUFFER_STATIC_SLOTS);
-		}
-	}
-
-	checkf(Shader->StaticSlots.Num() == Shader->Bindings.ShaderResourceTable.ResourceTableLayoutHashes.Num(), TEXT("StaticSlots %d, Bindings %d"),
-		Shader->StaticSlots.Num(), Shader->Bindings.ShaderResourceTable.ResourceTableLayoutHashes.Num());
+	UE::RHICore::InitStaticUniformBufferSlots(Shader->StaticSlots, Shader->Bindings.ShaderResourceTable);
 
 #if DEBUG_GL_SHADERS
 	{
diff --git a/Engine/Source/Runtime/RHICore/Public/RHICoreShader.h b/Engine/Source/Runtime/RHICore/Public/RHICoreShader.h
--- a/Engine/Source/Runtime/RHICore/Public/RHICoreShader.h
+++ b/Engine/Source/Runtime/RHICore/Public/RHICoreShader.h
@@ -13,6 +13,23 @@ namespace RHICore
 /** Validates that the uniform buffer at the requested static slot. */
 extern RHICORE_API void ValidateStaticUniformBuffer(FRHIUniformBuffer* UniformBuffer, FUniformBufferStaticSlot Slot, uint32 ExpectedHash);
 
+inline void InitStaticUniformBufferSlots(TArray<FUniformBufferStaticSlot>& StaticSlots, const FBaseShaderResourceTable& ShaderResourceTable)
+{
+	StaticSlots.Reserve(ShaderResourceTable.ResourceTableLayoutHashes.Num());
+
+	for (uint32 LayoutHash : ShaderResourceTable.ResourceTableLayoutHashes)
+	{
+		if (const FShaderParametersMetadata* Metadata = FindUniformBufferStructByLayoutHash(LayoutHash))
+		{
+			StaticSlots.Add(Metadata->GetLayout().StaticSlot);
+		}
+		else
+		{
+			StaticSlots.Add(MAX_UNIFORM_BUFFER_STATIC_SLOTS);
+		}
+	}
+}
+
 template <typename TRHIContext, typename TRHIShader>
 void ApplyStaticUniformBuffers(
 	TRHIContext* CommandContext,
diff --git a/Engine/Source/Runtime/Windows/D3D11RHI/Private/D3D11Shaders.cpp b/Engine/Source/Runtime/Windows/D3D11RHI/Private/D3D11Shaders.cpp
--- a/Engine/Source/Runtime/Windows/D3D11RHI/Private/D3D11Shaders.cpp
+++ b/Engine/Source/Runtime/Windows/D3D11RHI/Private/D3D11Shaders.cpp
@@ -6,6 +6,7 @@
 
 #include "D3D11RHIPrivate.h"
 #include "Serialization/MemoryReader.h"
+#include "RHICoreShader.h"
 
 #if !PLATFORM_HOLOLENS
 #include "nvapi.h"
@@ -114,25 +115,6 @@ static void ResetVendorExtensions(ID3D11Device* Direct3DDevice)
 #endif
 }
 
-static void InitUniformBufferStaticSlots(FD3D11ShaderData* Shader)
-{
-	const FBaseShaderResourceTable& SRT = Shader->ShaderResourceTable;
-
-	Shader->StaticSlots.Reserve(SRT.ResourceTableLayoutHashes.Num());
-
-	for (uint32 LayoutHash : SRT.ResourceTableLayoutHashes)
-	{
-		if (const FShaderParametersMetadata* Metadata = FindUniformBufferStructByLayoutHash(LayoutHash))
-		{
-			Shader->StaticSlots.Add(Metadata->GetLayout().StaticSlot);
-		}
-		else
-		{
-			Shader->StaticSlots.Add(MAX_UNIFORM_BUFFER_STATIC_SLOTS);
-		}
-	}
-}
-
 FVertexShaderRHIRef FD3D11DynamicRHI::RHICreateVertexShader(TArrayView<const uint8> Code, const FSHAHash& Hash)
 {
 	FShaderCodeReader ShaderCode(Code);
@@ -154,7 +136,7 @@ FVertexShaderRHIRef FD3D11DynamicRHI::RHICreateVertexShader(TArrayView<const uin
 		{
 			ResetVendorExtensions(Direct3DDevice);
 		}
-		InitUniformBufferStaticSlots(Shader);
+		UE::RHICore::InitStaticUniformBufferSlots(Shader->StaticSlots, Shader->ShaderResourceTable);
 	}
 	
 	// TEMP
@@ -192,7 +174,7 @@ FGeometryShaderRHIRef FD3D11DynamicRHI::RHICreateGeometryShader(TArrayView<const
 		{
 			ResetVendorExtensions(Direct3DDevice);
 		}
-		InitUniformBufferStaticSlots(Shader);
+		UE::RHICore::InitStaticUniformBufferSlots(Shader->StaticSlots, Shader->ShaderResourceTable);
 	}
 
 	return Shader;
@@ -226,7 +208,7 @@ FPixelShaderRHIRef FD3D11DynamicRHI::RHICreatePixelShader(TArrayView<const uint8
 		{
 			ResetVendorExtensions(Direct3DDevice);
 		}
-		InitUniformBufferStaticSlots(Shader);
+		UE::RHICore::InitStaticUniformBufferSlots(Shader->StaticSlots, Shader->ShaderResourceTable);
 	}
 
 	return Shader;
@@ -260,7 +242,7 @@ FComputeShaderRHIRef FD3D11DynamicRHI::RHICreateComputeShader(TArrayView<const u
 		{
 			ResetVendorExtensions(Direct3DDevice);
 		}
-		InitUniformBufferStaticSlots(Shader);
+		UE::RHICore::InitStaticUniformBufferSlots(Shader->StaticSlots, Shader->ShaderResourceTable);
 	}
 
 	return Shader;
