Fixing incorrect DirectX adapter code that assumes only an Intel GPU can be integrated and cleaning up related dead code.

diff --git a/Engine/Source/Runtime/D3D12RHI/Private/Windows/WindowsD3D12Device.cpp b/Engine/Source/Runtime/D3D12RHI/Private/Windows/WindowsD3D12Device.cpp
--- a/Engine/Source/Runtime/D3D12RHI/Private/Windows/WindowsD3D12Device.cpp
+++ b/Engine/Source/Runtime/D3D12RHI/Private/Windows/WindowsD3D12Device.cpp
@@ -633,17 +633,15 @@ void FD3D12DynamicRHIModule::FindAdapter()
 
 				// Simple heuristic but without profiling it's hard to do better
 				bool bIsNonLocalMemoryPresent = false;
-				if (bIsIntel)
+				TRefCountPtr<IDXGIAdapter3> TempDxgiAdapter3;
+				DXGI_QUERY_VIDEO_MEMORY_INFO NonLocalVideoMemoryInfo;
+				if (SUCCEEDED(TempAdapter->QueryInterface(_uuidof(IDXGIAdapter3), (void**)TempDxgiAdapter3.GetInitReference())) &&
+					TempDxgiAdapter3.IsValid() && SUCCEEDED(TempDxgiAdapter3->QueryVideoMemoryInfo(0, DXGI_MEMORY_SEGMENT_GROUP_NON_LOCAL, &NonLocalVideoMemoryInfo)))
 				{
-					TRefCountPtr<IDXGIAdapter3> TempDxgiAdapter3;
-					DXGI_QUERY_VIDEO_MEMORY_INFO NonLocalVideoMemoryInfo;
-					if (SUCCEEDED(TempAdapter->QueryInterface(_uuidof(IDXGIAdapter3), (void**)TempDxgiAdapter3.GetInitReference())) &&
-						TempDxgiAdapter3.IsValid() && SUCCEEDED(TempDxgiAdapter3->QueryVideoMemoryInfo(0, DXGI_MEMORY_SEGMENT_GROUP_NON_LOCAL, &NonLocalVideoMemoryInfo)))
-					{
-						bIsNonLocalMemoryPresent = NonLocalVideoMemoryInfo.Budget != 0;
-					}
+					bIsNonLocalMemoryPresent = NonLocalVideoMemoryInfo.Budget != 0;
 				}
-				const bool bIsIntegrated = bIsIntel && !bIsNonLocalMemoryPresent;
+				// TODO: Using GPUDetect for Intel GPUs to check for integrated vs discrete status, pending GPUDetect update
+				const bool bIsIntegrated = !bIsNonLocalMemoryPresent;
 
 				// PerfHUD is for performance profiling
 				const bool bIsPerfHUD = !FCString::Stricmp(AdapterDesc.Description, TEXT("NVIDIA PerfHUD"));
diff --git a/Engine/Source/Runtime/Windows/D3D11RHI/Private/Windows/WindowsD3D11Device.cpp b/Engine/Source/Runtime/Windows/D3D11RHI/Private/Windows/WindowsD3D11Device.cpp
--- a/Engine/Source/Runtime/Windows/D3D11RHI/Private/Windows/WindowsD3D11Device.cpp
+++ b/Engine/Source/Runtime/Windows/D3D11RHI/Private/Windows/WindowsD3D11Device.cpp
@@ -922,10 +922,6 @@ void FD3D11DynamicRHIModule::FindAdapter()
 	FD3D11Adapter FirstWithoutIntegratedAdapter;
 	FD3D11Adapter FirstAdapter;
 
-	bool bIsAnyAMD = false;
-	bool bIsAnyIntel = false;
-	bool bIsAnyNVIDIA = false;
-
 	UE_LOG(LogD3D11RHI, Log, TEXT("D3D11 adapters:"));
 
 	int PreferredVendor = D3D11RHI_PreferAdapterVendor();
@@ -998,25 +994,20 @@ void FD3D11DynamicRHIModule::FindAdapter()
 				bool bIsNVIDIA = AdapterDesc.VendorId == 0x10DE;
 				bool bIsMicrosoft = AdapterDesc.VendorId == 0x1414;
 
-				if(bIsAMD) bIsAnyAMD = true;
-				if(bIsIntel) bIsAnyIntel = true;
-				if(bIsNVIDIA) bIsAnyNVIDIA = true;
-
 				// Simple heuristic but without profiling it's hard to do better
 				bool bIsNonLocalMemoryPresent = false;
-				if (bIsIntel)
+				TRefCountPtr<IDXGIAdapter3> TempDxgiAdapter3;
+				DXGI_QUERY_VIDEO_MEMORY_INFO NonLocalVideoMemoryInfo;
+				if (SUCCEEDED(TempAdapter->QueryInterface(_uuidof(IDXGIAdapter3), (void**)TempDxgiAdapter3.GetInitReference())) &&
+					TempDxgiAdapter3.IsValid() && SUCCEEDED(TempDxgiAdapter3->QueryVideoMemoryInfo(0, DXGI_MEMORY_SEGMENT_GROUP_NON_LOCAL, &NonLocalVideoMemoryInfo)))
 				{
-					TRefCountPtr<IDXGIAdapter3> TempDxgiAdapter3;
-					DXGI_QUERY_VIDEO_MEMORY_INFO NonLocalVideoMemoryInfo;
-					if (SUCCEEDED(TempAdapter->QueryInterface(_uuidof(IDXGIAdapter3), (void**)TempDxgiAdapter3.GetInitReference())) &&
-						TempDxgiAdapter3.IsValid() && SUCCEEDED(TempDxgiAdapter3->QueryVideoMemoryInfo(0, DXGI_MEMORY_SEGMENT_GROUP_NON_LOCAL, &NonLocalVideoMemoryInfo)))
-					{
-						bIsNonLocalMemoryPresent = NonLocalVideoMemoryInfo.Budget != 0;
-					}
+					bIsNonLocalMemoryPresent = NonLocalVideoMemoryInfo.Budget != 0;
 				}
 
+				// TODO: Using GPUDetect for Intel GPUs to check for integrated vs discrete status, pending GPUDetect update
+
 				const bool bIsSoftware = bIsMicrosoft;
-				const bool bIsIntegrated = bIsIntel && !bIsNonLocalMemoryPresent;
+				const bool bIsIntegrated = !bIsNonLocalMemoryPresent;
 				// PerfHUD is for performance profiling
 				const bool bIsPerfHUD = !FCString::Stricmp(AdapterDesc.Description,TEXT("NVIDIA PerfHUD"));
 
