Fix for rare assert seen only in the wild

diff --git a/Engine/Source/Runtime/Windows/AudioMixerXAudio2/Private/WindowsMMDeviceInfoCache.cpp b/Engine/Source/Runtime/Windows/AudioMixerXAudio2/Private/WindowsMMDeviceInfoCache.cpp
--- a/Engine/Source/Runtime/Windows/AudioMixerXAudio2/Private/WindowsMMDeviceInfoCache.cpp
+++ b/Engine/Source/Runtime/Windows/AudioMixerXAudio2/Private/WindowsMMDeviceInfoCache.cpp
@@ -403,7 +403,7 @@ void FWindowsMMDeviceCache::OnDeviceRemoved(const FString& DeviceId, bool)
 {
 	FWriteScopeLock WriteLock(CacheMutationLock);
 	FName DeviceIdName = *DeviceId;
-	check(Cache.Contains(DeviceIdName));
+	UE_CLOG(!Cache.Contains(DeviceIdName), LogAudioMixer, Warning, TEXT("FWindowsMMDeviceCache::OnDeviceRemoved: DeviceId='%s' was not in the cache. "));
 	Cache.Remove(DeviceIdName);
 }
 
