Fix for non-pooled uniform buffers being added to the pool on destruction.

diff --git a/Engine/Source/Runtime/Windows/D3D11RHI/Private/D3D11UniformBuffer.cpp b/Engine/Source/Runtime/Windows/D3D11RHI/Private/D3D11UniformBuffer.cpp
--- a/Engine/Source/Runtime/Windows/D3D11RHI/Private/D3D11UniformBuffer.cpp
+++ b/Engine/Source/Runtime/Windows/D3D11RHI/Private/D3D11UniformBuffer.cpp
@@ -185,6 +185,8 @@ FUniformBufferRHIRef FD3D11DynamicRHI::RHICreateUniformBuffer(const void* Conten
 
 		if (IsPoolingEnabled() && (IsInRenderingThread() || IsInRHIThread()))
 		{
+			const bool bAllocatedFromPool = true;
+
 			if (ShouldNotEnqueueRHICommand())
 			{
 				TRefCountPtr<ID3D11Buffer> UniformBufferResource = CreateAndUpdatePooledUniformBuffer(
@@ -192,11 +194,12 @@ FUniformBufferRHIRef FD3D11DynamicRHI::RHICreateUniformBuffer(const void* Conten
 					Direct3DDeviceIMContext.GetReference(),
 					Contents,
 					NumBytes);
-				NewUniformBuffer = new FD3D11UniformBuffer(this, Layout, UniformBufferResource, FRingAllocation());
+
+				NewUniformBuffer = new FD3D11UniformBuffer(this, Layout, UniformBufferResource, FRingAllocation(), bAllocatedFromPool);
 			}
 			else
 			{
-				NewUniformBuffer = new FD3D11UniformBuffer(this, Layout, nullptr, FRingAllocation());
+				NewUniformBuffer = new FD3D11UniformBuffer(this, Layout, nullptr, FRingAllocation(), bAllocatedFromPool);
 				NewUniformBuffer->AddRef();
 				void* CPUContent = FMemory::Malloc(NumBytes);
 				FMemory::Memcpy(CPUContent, Contents, NumBytes);
@@ -232,7 +235,8 @@ FUniformBufferRHIRef FD3D11DynamicRHI::RHICreateUniformBuffer(const void* Conten
 			TRefCountPtr<ID3D11Buffer> UniformBufferResource;
 			VERIFYD3D11RESULT_EX(Direct3DDevice->CreateBuffer(&Desc,&ImmutableData,UniformBufferResource.GetInitReference()), Direct3DDevice);
 
-			NewUniformBuffer = new FD3D11UniformBuffer(this, Layout, UniformBufferResource, FRingAllocation());
+			const bool bAllocatedFromPool = false;
+			NewUniformBuffer = new FD3D11UniformBuffer(this, Layout, UniformBufferResource, FRingAllocation(), bAllocatedFromPool);
 
 			INC_DWORD_STAT(STAT_D3D11NumImmutableUniformBuffers);
 		}
@@ -240,7 +244,8 @@ FUniformBufferRHIRef FD3D11DynamicRHI::RHICreateUniformBuffer(const void* Conten
 	else
 	{
 		// This uniform buffer contains no constants, only a resource table.
-		NewUniformBuffer = new FD3D11UniformBuffer(this, Layout, nullptr, FRingAllocation());
+		const bool bAllocatedFromPool = false;
+		NewUniformBuffer = new FD3D11UniformBuffer(this, Layout, nullptr, FRingAllocation(), bAllocatedFromPool);
 	}
 
 	if (Layout->Resources.Num())
@@ -347,7 +352,7 @@ void FD3D11DynamicRHI::RHIUpdateUniformBuffer(FRHIUniformBuffer* UniformBufferRH
 FD3D11UniformBuffer::~FD3D11UniformBuffer()
 {
 	// Do not return the allocation to the pool if it is in the dynamic constant buffer!
-	if (!RingAllocation.IsValid() && Resource != nullptr)
+	if (bAllocatedFromPool && !RingAllocation.IsValid() && Resource != nullptr)
 	{
 		check(IsInRHIThread() || IsInRenderingThread());
 		D3D11_BUFFER_DESC Desc;
diff --git a/Engine/Source/Runtime/Windows/D3D11RHI/Public/D3D11Resources.h b/Engine/Source/Runtime/Windows/D3D11RHI/Public/D3D11Resources.h
--- a/Engine/Source/Runtime/Windows/D3D11RHI/Public/D3D11Resources.h
+++ b/Engine/Source/Runtime/Windows/D3D11RHI/Public/D3D11Resources.h
@@ -373,17 +373,19 @@ public:
 	TArray<TRefCountPtr<FRHIResource> > ResourceTable;
 
 	/** Initialization constructor. */
-	FD3D11UniformBuffer(class FD3D11DynamicRHI* InD3D11RHI, const FRHIUniformBufferLayout* InLayout, ID3D11Buffer* InResource,const FRingAllocation& InRingAllocation)
+	FD3D11UniformBuffer(class FD3D11DynamicRHI* InD3D11RHI, const FRHIUniformBufferLayout* InLayout, ID3D11Buffer* InResource,const FRingAllocation& InRingAllocation, bool bInAllocatedFromPool)
 	: FRHIUniformBuffer(InLayout)
 	, Resource(InResource)
 	, RingAllocation(InRingAllocation)
 	, D3D11RHI(InD3D11RHI)
+	, bAllocatedFromPool(bInAllocatedFromPool)
 	{}
 
 	virtual ~FD3D11UniformBuffer();
 
 private:
 	class FD3D11DynamicRHI* D3D11RHI;
+	bool bAllocatedFromPool;
 };
 
 /** Buffer resource class. */
