Remove #pragma once from cpp files
#rnx

diff --git a/Engine/Source/Runtime/Windows/AudioMixerXAudio2/Private/ToStringHelpers.cpp b/Engine/Source/Runtime/Windows/AudioMixerXAudio2/Private/ToStringHelpers.cpp
--- a/Engine/Source/Runtime/Windows/AudioMixerXAudio2/Private/ToStringHelpers.cpp
+++ b/Engine/Source/Runtime/Windows/AudioMixerXAudio2/Private/ToStringHelpers.cpp
@@ -1,7 +1,5 @@
 // Copyright Epic Games, Inc. All Rights Reserved.
 
-#pragma once
-
 #include "ToStringHelpers.h"
 
 #if PLATFORM_WINDOWS
