Fix crash in D3D11 when calling RHIBindDebugLabelName on an object without an underlying resource.

#rnx

diff --git a/Engine/Source/Runtime/Windows/D3D11RHI/Private/D3D11Buffer.cpp b/Engine/Source/Runtime/Windows/D3D11RHI/Private/D3D11Buffer.cpp
--- a/Engine/Source/Runtime/Windows/D3D11RHI/Private/D3D11Buffer.cpp
+++ b/Engine/Source/Runtime/Windows/D3D11RHI/Private/D3D11Buffer.cpp
@@ -271,6 +271,9 @@ void FD3D11DynamicRHI::RHIBindDebugLabelName(FRHIBuffer* BufferRHI, const TCHAR*
 	BufferRHI->SetName(DebugName);
 #if UE_BUILD_DEBUG || UE_BUILD_DEVELOPMENT
 	FD3D11Buffer* BufferD3D = ResourceCast(BufferRHI);
-	BufferD3D->Resource->SetPrivateData(WKPDID_D3DDebugObjectName, FCString::Strlen(Name) + 1, TCHAR_TO_ANSI(Name));
+	if (BufferD3D->Resource != nullptr)
+	{
+		BufferD3D->Resource->SetPrivateData(WKPDID_D3DDebugObjectName, FCString::Strlen(Name) + 1, TCHAR_TO_ANSI(Name));
+	}
 #endif
 }
\ No newline at end of file
diff --git a/Engine/Source/Runtime/Windows/D3D11RHI/Private/D3D11Texture.cpp b/Engine/Source/Runtime/Windows/D3D11RHI/Private/D3D11Texture.cpp
--- a/Engine/Source/Runtime/Windows/D3D11RHI/Private/D3D11Texture.cpp
+++ b/Engine/Source/Runtime/Windows/D3D11RHI/Private/D3D11Texture.cpp
@@ -1707,7 +1707,11 @@ void FD3D11DynamicRHI::RHIBindDebugLabelName(FRHITexture* TextureRHI, const TCHA
 	FName DebugName(Name);
 	Texture->SetName(DebugName);
 #if UE_BUILD_DEBUG || UE_BUILD_DEVELOPMENT
-	Texture->GetResource()->SetPrivateData(WKPDID_D3DDebugObjectName, FCString::Strlen(Name) + 1, TCHAR_TO_ANSI(Name));
+	ID3D11Resource* ResourceD3D = Texture->GetResource();
+	if (ResourceD3D != nullptr)
+	{
+		ResourceD3D->SetPrivateData(WKPDID_D3DDebugObjectName, FCString::Strlen(Name) + 1, TCHAR_TO_ANSI(Name));
+	}
 #endif
 }
 
diff --git a/Engine/Source/Runtime/Windows/D3D11RHI/Private/D3D11UAV.cpp b/Engine/Source/Runtime/Windows/D3D11RHI/Private/D3D11UAV.cpp
--- a/Engine/Source/Runtime/Windows/D3D11RHI/Private/D3D11UAV.cpp
+++ b/Engine/Source/Runtime/Windows/D3D11RHI/Private/D3D11UAV.cpp
@@ -479,7 +479,10 @@ void FD3D11DynamicRHI::RHIBindDebugLabelName(FRHIUnorderedAccessView* UnorderedA
 {
 #if UE_BUILD_DEBUG || UE_BUILD_DEVELOPMENT
 	FD3D11UnorderedAccessView* UAV = ResourceCast(UnorderedAccessViewRHI);
-	UAV->View->SetPrivateData(WKPDID_D3DDebugObjectName, FCString::Strlen(Name) + 1, TCHAR_TO_ANSI(Name));
+	if (UAV->View != nullptr)
+	{
+		UAV->View->SetPrivateData(WKPDID_D3DDebugObjectName, FCString::Strlen(Name) + 1, TCHAR_TO_ANSI(Name));
+	}
 #endif
 }
 
