Added telemetry for GPU OOM crashes.

diff --git a/Engine/Source/Runtime/Core/Private/Misc/CoreDelegates.cpp b/Engine/Source/Runtime/Core/Private/Misc/CoreDelegates.cpp
--- a/Engine/Source/Runtime/Core/Private/Misc/CoreDelegates.cpp
+++ b/Engine/Source/Runtime/Core/Private/Misc/CoreDelegates.cpp
@@ -231,6 +231,13 @@ FSimpleMulticastDelegate& FCoreDelegates::GetOutOfMemoryDelegate()
 	return OnOOM;
 }
 
+/**	 Implemented as a function to address global ctor issues */
+FCoreDelegates::FGPUOutOfMemoryDelegate& FCoreDelegates::GetGPUOutOfMemoryDelegate()
+{
+	static FGPUOutOfMemoryDelegate OnGPUOOM;
+	return OnGPUOOM;
+}
+
 FCoreDelegates::FGetOnScreenMessagesDelegate FCoreDelegates::OnGetOnScreenMessages;
 
 typedef void(*TSigningKeyFunc)(TArray<uint8>&, TArray<uint8>&);
diff --git a/Engine/Source/Runtime/Core/Public/Misc/CoreDelegates.h b/Engine/Source/Runtime/Core/Public/Misc/CoreDelegates.h
--- a/Engine/Source/Runtime/Core/Public/Misc/CoreDelegates.h
+++ b/Engine/Source/Runtime/Core/Public/Misc/CoreDelegates.h
@@ -581,6 +581,10 @@ public:
 	// Called when OOM event occurs, after backup memory has been freed, so there's some hope of being effective
 	static FSimpleMulticastDelegate& GetOutOfMemoryDelegate();
 
+	// Called from TerminateOnOutOfMemory in D3D11Util.cpp/D3D12Util.cpp
+	DECLARE_MULTICAST_DELEGATE_TwoParams(FGPUOutOfMemoryDelegate, const uint64, const uint64);
+	static FGPUOutOfMemoryDelegate& GetGPUOutOfMemoryDelegate();
+
 	enum class EOnScreenMessageSeverity : uint8
 	{
 		Info,
diff --git a/Engine/Source/Runtime/D3D12RHI/Private/D3D12Util.cpp b/Engine/Source/Runtime/D3D12RHI/Private/D3D12Util.cpp
--- a/Engine/Source/Runtime/D3D12RHI/Private/D3D12Util.cpp
+++ b/Engine/Source/Runtime/D3D12RHI/Private/D3D12Util.cpp
@@ -682,6 +682,29 @@ extern CORE_API bool GIsGPUCrashed;
 static void TerminateOnOutOfMemory(ID3D12Device* InDevice, HRESULT D3DResult, bool bCreatingTextures)
 {
 #if PLATFORM_WINDOWS
+	// send telemetry event with current adapter's memory info
+	FD3D12DynamicRHI* D3D12RHI = FD3D12DynamicRHI::GetD3DRHI();
+	FD3D12Adapter* Adapter = nullptr;
+	D3D12RHI->ForEachDevice(InDevice, [&](FD3D12Device* IterationDevice)
+		{
+			if (InDevice == IterationDevice->GetDevice())
+			{
+				Adapter = IterationDevice->GetParentAdapter();
+			}
+		});
+	// if InDevice == nullptr, just pick the first available adapter
+	if (!Adapter && D3D12RHI->GetNumAdapters() == 1)
+	{
+		check(!InDevice);
+		Adapter = &D3D12RHI->GetAdapter(0);
+	}
+	if (Adapter)
+	{
+		const auto& MemoryInfo = Adapter->GetMemoryInfo().LocalMemoryInfo;
+		FCoreDelegates::GetGPUOutOfMemoryDelegate().Broadcast(MemoryInfo.Budget, MemoryInfo.CurrentUsage);
+	}
+
+
 	if (bCreatingTextures)
 	{
 		FPlatformMisc::MessageBoxExt(EAppMsgType::Ok, *LOCTEXT("OutOfVideoMemoryTextures", "Out of video memory trying to allocate a texture! Make sure your video card has the minimum required memory, try lowering the resolution and/or closing other applications that are running. Exiting...").ToString(), TEXT("Error"));
@@ -697,7 +720,6 @@ static void TerminateOnOutOfMemory(ID3D12Device* InDevice, HRESULT D3DResult, bo
 	static IConsoleVariable* GPUCrashOOM = IConsoleManager::Get().FindConsoleVariable(TEXT("r.GPUCrashOnOutOfMemory"));
 	if (GPUCrashOOM && GPUCrashOOM->GetInt())
 	{
-		FD3D12DynamicRHI* D3D12RHI = FD3D12DynamicRHI::GetD3DRHI();
 		// If no device provided then try and log the DRED status of each device
 		D3D12RHI->ForEachDevice(InDevice, [&](FD3D12Device* IterationDevice)
 			{
@@ -707,7 +729,6 @@ static void TerminateOnOutOfMemory(ID3D12Device* InDevice, HRESULT D3DResult, bo
 					LogMemoryInfo(Adapter);
 				}
 			});
-
 		UE_LOG(LogD3D12RHI, Fatal, TEXT("Out of video memory trying to allocate a rendering resource"));
 	}
 	else
diff --git a/Engine/Source/Runtime/Engine/Private/ChartCreation.cpp b/Engine/Source/Runtime/Engine/Private/ChartCreation.cpp
--- a/Engine/Source/Runtime/Engine/Private/ChartCreation.cpp
+++ b/Engine/Source/Runtime/Engine/Private/ChartCreation.cpp
@@ -601,15 +601,15 @@
 	TotalRHIThreadBoundHitchCount = 0;
 	TotalGPUBoundHitchCount = 0;
 	MaxDrawCalls = 0;
-	MinDrawCalls = INT_MAX;
+	MinDrawCalls = TNumericLimits<decltype(MinDrawCalls)>::Max();
 	TotalDrawCalls = 0;
 	MaxPlayerTicks = 0;
-	MinPlayerTicks = INT_MAX;
+	MinPlayerTicks = TNumericLimits<decltype(MinPlayerTicks)>::Max();
 	TotalPlayerTicks = 0;
 	MaxVehicleTicks = 0;
 	TotalVehicleTicks = 0;
 	MaxDrawnPrimitives = 0;
-	MinDrawnPrimitives = INT_MAX;
+	MinDrawnPrimitives = TNumericLimits<decltype(MinDrawnPrimitives)>::Max();
 	TotalDrawnPrimitives = 0;
 	AccumulatedChartTime = 0.0;
 	TimeDisregarded = 0.0;
@@ -617,9 +617,9 @@
 	NumFramesAtCriticalMemoryPressure = 0;
 	MaxPhysicalMemory = 0;
 	MaxVirtualMemory = 0;
-	MinPhysicalMemory = ULONG_MAX;
-	MinVirtualMemory = ULONG_MAX;
-	MinAvailablePhysicalMemory = ULONG_MAX;
+	MinPhysicalMemory = TNumericLimits<decltype(MinPhysicalMemory)>::Max();
+	MinVirtualMemory = TNumericLimits<decltype(MinVirtualMemory)>::Max();
+	MinAvailablePhysicalMemory = TNumericLimits<decltype(MinAvailablePhysicalMemory)>::Max();
 	TotalPhysicalMemoryUsed = 0;
 	TotalVirtualMemoryUsed = 0;
 
diff --git a/Engine/Source/Runtime/Windows/D3D11RHI/Private/D3D11Util.cpp b/Engine/Source/Runtime/Windows/D3D11RHI/Private/D3D11Util.cpp
--- a/Engine/Source/Runtime/Windows/D3D11RHI/Private/D3D11Util.cpp
+++ b/Engine/Source/Runtime/Windows/D3D11RHI/Private/D3D11Util.cpp
@@ -191,7 +191,7 @@ static void TerminateOnDeviceRemoved(HRESULT D3DResult, ID3D11Device* Direct3DDe
 	}
 }
 
-void LogMemoryInfo(const FD3D11Adapter& InAdapter)
+void GetAndLogMemoryInfo(const FD3D11Adapter& InAdapter, uint64& OutVRAMBudgetBytes, uint64& OutVRAMUsageBytes)
 {
 	TRefCountPtr<IDXGIAdapter3> Adapter3;
 	const HRESULT AdapterHR = InAdapter.DXGIAdapter->QueryInterface(IID_PPV_ARGS(Adapter3.GetInitReference()));
@@ -201,6 +201,8 @@ void LogMemoryInfo(const FD3D11Adapter& InAdapter)
 		VERIFYD3D11RESULT(Adapter3->QueryVideoMemoryInfo(0, DXGI_MEMORY_SEGMENT_GROUP_LOCAL, &LocalMemoryInfo));
 		UE_LOG(LogD3D11RHI, Error, TEXT("\tBudget:\t%7.2f MB"), LocalMemoryInfo.Budget / (1024.0f * 1024));
 		UE_LOG(LogD3D11RHI, Error, TEXT("\tUsed:\t%7.2f MB"), LocalMemoryInfo.CurrentUsage / (1024.0f * 1024));
+		OutVRAMBudgetBytes = LocalMemoryInfo.Budget;
+		OutVRAMUsageBytes = LocalMemoryInfo.CurrentUsage;
 	}
 }
 
@@ -208,6 +210,10 @@ static void TerminateOnOutOfMemory(HRESULT D3DResult, bool bCreatingTextures)
 {
 	if (D3DResult == E_OUTOFMEMORY)
 	{
+		uint64 VRAMBudgetBytes = 0, VRAMUsageBytes = 0;
+		GetAndLogMemoryInfo(GD3D11RHI->GetAdapter(), VRAMBudgetBytes, VRAMUsageBytes);
+		FCoreDelegates::GetGPUOutOfMemoryDelegate().Broadcast(VRAMBudgetBytes, VRAMUsageBytes);
+
 		if (bCreatingTextures)
 		{
 			FPlatformMisc::MessageBoxExt(EAppMsgType::Ok, *LOCTEXT("OutOfVideoMemoryTextures", "Out of video memory trying to allocate a texture! Make sure your video card has the minimum required memory, try lowering the resolution and/or closing other applications that are running. Exiting...").ToString(), TEXT("Error"));
@@ -222,7 +228,6 @@ static void TerminateOnOutOfMemory(HRESULT D3DResult, bool bCreatingTextures)
 		static IConsoleVariable* GPUCrashOOM = IConsoleManager::Get().FindConsoleVariable(TEXT("r.GPUCrashOnOutOfMemory"));
 		if (GPUCrashOOM && GPUCrashOOM->GetInt())
 		{
-			LogMemoryInfo(GD3D11RHI->GetAdapter());
 			UE_LOG(LogD3D11RHI, Fatal, TEXT("Out of video memory trying to allocate a rendering resource"));
 		}
 		else
