Show buffer name in assert message when trying to create empty buffers.

#rnx

diff --git a/Engine/Source/Runtime/D3D12RHI/Private/D3D12Buffer.cpp b/Engine/Source/Runtime/D3D12RHI/Private/D3D12Buffer.cpp
--- a/Engine/Source/Runtime/D3D12RHI/Private/D3D12Buffer.cpp
+++ b/Engine/Source/Runtime/D3D12RHI/Private/D3D12Buffer.cpp
@@ -317,7 +317,7 @@ void FD3D12Adapter::AllocateBuffer(FD3D12Device* Device,
 	TRACE_CPUPROFILER_EVENT_SCOPE(D3D12RHI::AllocateBuffer);
 
 	// Explicitly check that the size is nonzero before allowing CreateBuffer to opaquely fail.
-	check(Size > 0);
+	checkf(Size > 0, TEXT("Attempt to create buffer '%s' with size 0."), InDebugName ? InDebugName : TEXT("(null)"));
 
 	if (EnumHasAnyFlags(InUsage, BUF_AnyDynamic))
 	{
diff --git a/Engine/Source/Runtime/Windows/D3D11RHI/Private/D3D11Buffer.cpp b/Engine/Source/Runtime/Windows/D3D11RHI/Private/D3D11Buffer.cpp
--- a/Engine/Source/Runtime/Windows/D3D11RHI/Private/D3D11Buffer.cpp
+++ b/Engine/Source/Runtime/Windows/D3D11RHI/Private/D3D11Buffer.cpp
@@ -21,7 +21,7 @@ FBufferRHIRef FD3D11DynamicRHI::RHICreateBuffer(FRHICommandListBase&, uint32 Siz
 	}
 
 	// Explicitly check that the size is nonzero before allowing CreateBuffer to opaquely fail.
-	check(Size > 0);
+	checkf(Size > 0, TEXT("Attempt to create buffer '%s' with size 0."), CreateInfo.DebugName ? CreateInfo.DebugName : TEXT("(null)"));
 
 	// Describe the buffer.
 	D3D11_BUFFER_DESC Desc = { Size };
