Fix CIS issues with viewport array initialization.

diff --git a/Engine/Source/Runtime/D3D12RHI/Private/D3D12Commands.cpp b/Engine/Source/Runtime/D3D12RHI/Private/D3D12Commands.cpp
--- a/Engine/Source/Runtime/D3D12RHI/Private/D3D12Commands.cpp
+++ b/Engine/Source/Runtime/D3D12RHI/Private/D3D12Commands.cpp
@@ -744,7 +744,7 @@ void FD3D12CommandContext::RHISetViewport(float MinX, float MinY, float MinZ, fl
 void FD3D12CommandContext::RHISetStereoViewport(float LeftMinX, float RightMinX, float LeftMinY, float RightMinY, float MinZ, float LeftMaxX, float RightMaxX, float LeftMaxY, float RightMaxY, float MaxZ)
 {
 	// Set up both viewports
-	D3D12_VIEWPORT Viewports[2] = {0};
+	D3D12_VIEWPORT Viewports[2] = {};
 
 	Viewports[0].TopLeftX = FMath::FloorToInt(LeftMinX);
 	Viewports[0].TopLeftY = FMath::FloorToInt(LeftMinY);
diff --git a/Engine/Source/Runtime/Windows/D3D11RHI/Private/D3D11Commands.cpp b/Engine/Source/Runtime/Windows/D3D11RHI/Private/D3D11Commands.cpp
--- a/Engine/Source/Runtime/Windows/D3D11RHI/Private/D3D11Commands.cpp
+++ b/Engine/Source/Runtime/Windows/D3D11RHI/Private/D3D11Commands.cpp
@@ -282,7 +282,7 @@ static void ValidateScissorRect(const D3D11_VIEWPORT& Viewport, const D3D11_RECT
 void FD3D11DynamicRHI::RHISetStereoViewport(float LeftMinX, float RightMinX, float LeftMinY, float RightMinY, float MinZ, float LeftMaxX, float RightMaxX, float LeftMaxY, float RightMaxY, float MaxZ)
 {
 	// Set up both viewports
-	D3D11_VIEWPORT StereoViewports[2] = { 0 };
+	D3D11_VIEWPORT StereoViewports[2] = {};
 
 	StereoViewports[0].TopLeftX = FMath::FloorToInt(LeftMinX);
 	StereoViewports[0].TopLeftY = FMath::FloorToInt(LeftMinY);
