Fixed typo in log message.

#rnx

diff --git a/Engine/Source/Runtime/Windows/D3D11RHI/Private/D3D11Util.cpp b/Engine/Source/Runtime/Windows/D3D11RHI/Private/D3D11Util.cpp
--- a/Engine/Source/Runtime/Windows/D3D11RHI/Private/D3D11Util.cpp
+++ b/Engine/Source/Runtime/Windows/D3D11RHI/Private/D3D11Util.cpp
@@ -180,7 +180,7 @@ static void TerminateOnDeviceRemoved(HRESULT D3DResult, ID3D11Device* Direct3DDe
 		}
 		else
 		{
-			UE_LOG(LogD3D11RHI, Fatal, TEXT("Unreal Engine is exiting due to D3D device being lost. D3D device was not available to assertain DXGI cause."));
+			UE_LOG(LogD3D11RHI, Fatal, TEXT("Unreal Engine is exiting due to D3D device being lost. D3D device was not available to determine DXGI cause."));
 		}
 
 		// Workaround for the fact that in non-monolithic builds the exe gets into a weird state and exception handling fails. 
