Report creation parameters when buffer creation fails in D3D11.

#rnx

diff --git a/Engine/Source/Runtime/Windows/D3D11RHI/Private/D3D11Buffer.cpp b/Engine/Source/Runtime/Windows/D3D11RHI/Private/D3D11Buffer.cpp
--- a/Engine/Source/Runtime/Windows/D3D11RHI/Private/D3D11Buffer.cpp
+++ b/Engine/Source/Runtime/Windows/D3D11RHI/Private/D3D11Buffer.cpp
@@ -92,7 +92,15 @@ FBufferRHIRef FD3D11DynamicRHI::RHICreateBuffer(FRHICommandListBase&, uint32 Siz
 	}
 
 	TRefCountPtr<ID3D11Buffer> BufferResource;
-	VERIFYD3D11RESULT_EX(Direct3DDevice->CreateBuffer(&Desc, pInitData, BufferResource.GetInitReference()), Direct3DDevice);
+	{
+		HRESULT hr = Direct3DDevice->CreateBuffer(&Desc, pInitData, BufferResource.GetInitReference());
+		if (FAILED(hr))
+		{
+			UE_LOG(LogD3D11RHI, Error, TEXT("Failed to create buffer '%s' with ByteWidth=%u, Usage=%d, BindFlags=0x%x, CPUAccessFlags=0x%x, MiscFlags=0x%x, StructureByteStride=%u, InitData=0x%p"),
+				CreateInfo.DebugName ? CreateInfo.DebugName : TEXT(""), Desc.ByteWidth, Desc.Usage, Desc.BindFlags, Desc.CPUAccessFlags, Desc.MiscFlags, Desc.StructureByteStride, pInitData);
+			VerifyD3D11Result(hr, "CreateBuffer", __FILE__, __LINE__, Direct3DDevice);
+		}
+	}
 
 	if (CreateInfo.DebugName)
 	{
