Removed erroneous IsInRenderingThread check.

diff --git a/Engine/Source/Runtime/Windows/D3D11RHI/Private/D3D11UniformBuffer.cpp b/Engine/Source/Runtime/Windows/D3D11RHI/Private/D3D11UniformBuffer.cpp
--- a/Engine/Source/Runtime/Windows/D3D11RHI/Private/D3D11UniformBuffer.cpp
+++ b/Engine/Source/Runtime/Windows/D3D11RHI/Private/D3D11UniformBuffer.cpp
@@ -295,7 +295,6 @@ void UpdateUniformBufferContents(FD3D11Device* Direct3DDevice, FD3D11DeviceConte
 
 void FD3D11DynamicRHI::RHIUpdateUniformBuffer(FRHICommandListBase& RHICmdList, FRHIUniformBuffer* UniformBufferRHI, const void* Contents)
 {
-	check(IsInRenderingThread());
 	check(UniformBufferRHI);
 
 	FD3D11UniformBuffer* UniformBuffer = ResourceCast(UniformBufferRHI);
