Removed redundant private include paths from build.cs files.

diff --git a/Engine/Source/Runtime/UMG/Private/WidgetLayoutLibrary.cpp b/Engine/Source/Runtime/UMG/Private/WidgetLayoutLibrary.cpp
--- a/Engine/Source/Runtime/UMG/Private/WidgetLayoutLibrary.cpp
+++ b/Engine/Source/Runtime/UMG/Private/WidgetLayoutLibrary.cpp
@@ -23,8 +23,8 @@
 #include "Engine/UserInterfaceSettings.h"
 #include "Slate/SlateBrushAsset.h"
 #include "Blueprint/WidgetLayoutLibrary.h"
-#include "Runtime/Engine/Classes/Engine/UserInterfaceSettings.h"
-#include "Runtime/Engine/Classes/Engine/RendererSettings.h"
+#include "Engine/UserInterfaceSettings.h"
+#include "Engine/RendererSettings.h"
 #include "Blueprint/SlateBlueprintLibrary.h"
 #include "Slate/SceneViewport.h"
 #include "Slate/SGameLayerManager.h"
diff --git a/Engine/Source/Runtime/VectorVM/Private/Tests/VectorVMTests.cpp b/Engine/Source/Runtime/VectorVM/Private/Tests/VectorVMTests.cpp
--- a/Engine/Source/Runtime/VectorVM/Private/Tests/VectorVMTests.cpp
+++ b/Engine/Source/Runtime/VectorVM/Private/Tests/VectorVMTests.cpp
@@ -3,7 +3,7 @@
 #include "CoreMinimal.h"
 #include "Misc/AutomationTest.h"
 #include "VectorVM.h"
-#include "Runtime/VectorVM/Private/VectorVMPrivate.h"
+#include "VectorVMPrivate.h"
 
 #define OP_REGISTER (0)
 #define OP1_CONST (1 << 1)
diff --git a/Engine/Source/Runtime/VulkanRHI/VulkanRHI.Build.cs b/Engine/Source/Runtime/VulkanRHI/VulkanRHI.Build.cs
--- a/Engine/Source/Runtime/VulkanRHI/VulkanRHI.Build.cs
+++ b/Engine/Source/Runtime/VulkanRHI/VulkanRHI.Build.cs
@@ -12,7 +12,6 @@ public class VulkanRHI : ModuleRules
 	{
 		bLegalToDistributeObjectCode = true;
 
-		PrivateIncludePaths.Add("Runtime/VulkanRHI/Private");
 		if (Target.Platform.IsInGroup(UnrealPlatformGroup.Windows))
 		{
 			PrivateIncludePaths.Add("Runtime/VulkanRHI/Private/Windows");
diff --git a/Engine/Source/Runtime/Windows/D3D11RHI/D3D11RHI.Build.cs b/Engine/Source/Runtime/Windows/D3D11RHI/D3D11RHI.Build.cs
--- a/Engine/Source/Runtime/Windows/D3D11RHI/D3D11RHI.Build.cs
+++ b/Engine/Source/Runtime/Windows/D3D11RHI/D3D11RHI.Build.cs
@@ -9,7 +9,6 @@ public class D3D11RHI : ModuleRules
 
 	public D3D11RHI(ReadOnlyTargetRules Target) : base(Target)
 	{
-		PrivateIncludePaths.Add("Runtime/Windows/D3D11RHI/Private");
 		PrivateIncludePaths.Add("../Shaders/Shared");
 
 		PrivateDependencyModuleNames.AddRange(
@@ -23,6 +22,12 @@ public class D3D11RHI : ModuleRules
 			}
 			);
 
+		PublicIncludePathModuleNames.AddRange(
+			new string[] {
+					"HeadMountedDisplay"
+			}
+			);
+
 		AddEngineThirdPartyPrivateStaticDependencies(Target, "DX11");
 		if (bIncludeExtensions)
 		{ 
