LLM was tracking all buffer types as "Meshes" in d3d11. Remove that tracking so the Meshes tag shows a comparable amount of memory usage as it does in d3d12.

diff --git a/Engine/Source/Runtime/Windows/D3D11RHI/Private/D3D11RHI.cpp b/Engine/Source/Runtime/Windows/D3D11RHI/Private/D3D11RHI.cpp
--- a/Engine/Source/Runtime/Windows/D3D11RHI/Private/D3D11RHI.cpp
+++ b/Engine/Source/Runtime/Windows/D3D11RHI/Private/D3D11RHI.cpp
@@ -725,13 +725,10 @@ void UpdateBufferStats(TRefCountPtr<ID3D11Buffer> Buffer, bool bAllocating)
 			INC_MEMORY_STAT_BY(STAT_StructuredBufferMemory,Desc.ByteWidth);
 		}
 
-#if PLATFORM_WINDOWS
 		// this is a work-around on Windows. Due to the fact that there is no way
 		// to hook the actual d3d allocations we can't track the memory in the normal way.
 		// Instead we simply tell LLM the size of these resources.
-		LLM_SCOPED_PAUSE_TRACKING_WITH_ENUM_AND_AMOUNT(ELLMTag::Meshes, Desc.ByteWidth, ELLMTracker::Default, ELLMAllocType::None);
 		LLM_SCOPED_PAUSE_TRACKING_WITH_ENUM_AND_AMOUNT(ELLMTag::GraphicsPlatform, Desc.ByteWidth, ELLMTracker::Platform, ELLMAllocType::None);
-#endif
 	}
 	else
 	{ //-V523
@@ -752,13 +749,10 @@ void UpdateBufferStats(TRefCountPtr<ID3D11Buffer> Buffer, bool bAllocating)
 			DEC_MEMORY_STAT_BY(STAT_StructuredBufferMemory,Desc.ByteWidth);
 		}
 
-#if PLATFORM_WINDOWS
 		// this is a work-around on Windows. Due to the fact that there is no way
 		// to hook the actual d3d allocations we can't track the memory in the normal way.
 		// Instead we simply tell LLM the size of these resources.
-		LLM_SCOPED_PAUSE_TRACKING_WITH_ENUM_AND_AMOUNT(ELLMTag::Meshes, -(int64)Desc.ByteWidth, ELLMTracker::Default, ELLMAllocType::None);
 		LLM_SCOPED_PAUSE_TRACKING_WITH_ENUM_AND_AMOUNT(ELLMTag::GraphicsPlatform, -(int64)Desc.ByteWidth, ELLMTracker::Platform, ELLMAllocType::None);
-#endif
 	}
 }
 
