Allow HDR to be enabled by default in editor: One needs to set:
Editor.HDRSupport=1
r.AllowHDR=1
bUseHDRDisplayOutput=True (in the game user settings) or r.HDR.EnableHDROutput=1

diff --git a/Engine/Source/Runtime/D3D12RHI/Private/Windows/WindowsD3D12Viewport.cpp b/Engine/Source/Runtime/D3D12RHI/Private/Windows/WindowsD3D12Viewport.cpp
--- a/Engine/Source/Runtime/D3D12RHI/Private/Windows/WindowsD3D12Viewport.cpp
+++ b/Engine/Source/Runtime/D3D12RHI/Private/Windows/WindowsD3D12Viewport.cpp
@@ -407,7 +407,7 @@ void FD3D12Viewport::EnableHDR()
 {
 	if ( GRHISupportsHDROutput && IsHDREnabled() )
 	{
-		const float DisplayMaxOutputNits = (DisplayOutputFormat == EDisplayOutputFormat::HDR_ACES_2000nit_ST2084 || DisplayOutputFormat == EDisplayOutputFormat::HDR_ACES_2000nit_ScRGB) ? 2000.f : 1000.f;
+		const float DisplayMaxOutputNits = HDRGetDisplayMaximumLuminance();
 		const float DisplayMinOutputNits = 0.0f;	// Min output of the display
 		const float DisplayMaxCLL = 0.0f;			// Max content light level in lumens (0.0 == unknown)
 		const float DisplayFALL = 0.0f;				// Frame average light level (0.0 == unknown)
diff --git a/Engine/Source/Runtime/Engine/Private/GameUserSettings.cpp b/Engine/Source/Runtime/Engine/Private/GameUserSettings.cpp
--- a/Engine/Source/Runtime/Engine/Private/GameUserSettings.cpp
+++ b/Engine/Source/Runtime/Engine/Private/GameUserSettings.cpp
@@ -468,15 +468,7 @@ void UGameUserSettings::ApplyNonResolutionSettings()
 
 	IConsoleManager::Get().CallAllConsoleVariableSinks();
 
-	bool bWithEditor = false;
-#if WITH_EDITOR
-	if (GIsEditor)
-	{
-		bWithEditor = true;
-	}
-#endif
-
-	bool bEnableHDR = (IsHDRAllowed() && bUseHDRDisplayOutput && !bWithEditor);
+	bool bEnableHDR = (IsHDRAllowed() && bUseHDRDisplayOutput);
 
 	EnableHDRDisplayOutputInternal(bEnableHDR, HDRDisplayOutputNits, true);
 
diff --git a/Engine/Source/Runtime/Windows/D3D11RHI/Private/Windows/WindowsD3D11Device.cpp b/Engine/Source/Runtime/Windows/D3D11RHI/Private/Windows/WindowsD3D11Device.cpp
--- a/Engine/Source/Runtime/Windows/D3D11RHI/Private/Windows/WindowsD3D11Device.cpp
+++ b/Engine/Source/Runtime/Windows/D3D11RHI/Private/Windows/WindowsD3D11Device.cpp
@@ -477,7 +477,7 @@ void FD3D11DynamicRHI::EnableHDR()
 		const EDisplayOutputFormat OutputDevice = HDRGetDefaultDisplayOutputFormat();
 		const EDisplayColorGamut DisplayGamut = HDRGetDefaultDisplayColorGamut();
 
-		const float DisplayMaxOutputNits = (OutputDevice == EDisplayOutputFormat::HDR_ACES_2000nit_ST2084 || OutputDevice == EDisplayOutputFormat::HDR_ACES_2000nit_ScRGB) ? 2000.f : 1000.f;
+		const float DisplayMaxOutputNits = HDRGetDisplayMaximumLuminance();
 		const float DisplayMinOutputNits = 0.0f;	// Min output of the display
 		const float DisplayMaxCLL = 0.0f;			// Max content light level in lumens (0.0 == unknown)
 		const float DisplayFALL = 0.0f;				// Frame average light level (0.0 == unknown)
diff --git a/Engine/Source/Runtime/Windows/D3D11RHI/Private/Windows/WindowsD3D11Viewport.cpp b/Engine/Source/Runtime/Windows/D3D11RHI/Private/Windows/WindowsD3D11Viewport.cpp
--- a/Engine/Source/Runtime/Windows/D3D11RHI/Private/Windows/WindowsD3D11Viewport.cpp
+++ b/Engine/Source/Runtime/Windows/D3D11RHI/Private/Windows/WindowsD3D11Viewport.cpp
@@ -185,6 +185,23 @@ FD3D11Viewport::FD3D11Viewport(FD3D11DynamicRHI* InD3DRHI,HWND InWindowHandle,ui
 			{
 				SwapChain1->QueryInterface(IID_PPV_ARGS(SwapChain.GetInitReference()));
 
+				RECT WindowRect = {};
+#if PLATFORM_WINDOWS
+				GetWindowRect(WindowHandle, &WindowRect);
+#endif
+				FVector2D WindowTopLeft((float)WindowRect.left, (float)WindowRect.top);
+				FVector2D WindowBottomRight((float)WindowRect.right, (float)WindowRect.bottom);
+				bool bHDREnabled = false;
+				EDisplayColorGamut LocalDisplayColorGamut = DisplayColorGamut;
+				EDisplayOutputFormat LocalDisplayOutputFormat = DisplayOutputFormat;
+
+				HDRGetMetaData(LocalDisplayOutputFormat, LocalDisplayColorGamut, bHDREnabled, WindowTopLeft, WindowBottomRight, (void*)WindowHandle);
+				if (bHDREnabled)
+				{
+					DisplayOutputFormat = LocalDisplayOutputFormat;
+					DisplayColorGamut = LocalDisplayColorGamut;
+				}
+
 				// See if we are running on a HDR monitor 
 				CheckHDRMonitorStatus();
 			}
@@ -330,7 +347,7 @@ inline void EnsureColorSpace(IDXGISwapChain* SwapChain, EDisplayColorGamut Displ
 		}
 		else
 		{
-			UE_LOG(LogD3D11RHI, Error, TEXT("Warning: unabled to set color space %s to the swapchain: verify EDisplayOutputFormat / swapchain format"), *ColorSpaceName);
+			UE_LOG(LogD3D11RHI, Error, TEXT("Warning: unable to set color space %s to the swapchain: verify EDisplayOutputFormat / swapchain format"), *ColorSpaceName);
 		}
 	}
 }
