Added debug logging for streamable textures to D3D11 to catch invalid args from streamable asset.

diff --git a/Engine/Source/Runtime/Windows/D3D11RHI/Private/D3D11Util.cpp b/Engine/Source/Runtime/Windows/D3D11RHI/Private/D3D11Util.cpp
--- a/Engine/Source/Runtime/Windows/D3D11RHI/Private/D3D11Util.cpp
+++ b/Engine/Source/Runtime/Windows/D3D11RHI/Private/D3D11Util.cpp
@@ -7,6 +7,7 @@
 #include "D3D11RHIPrivate.h"
 #include "EngineModule.h"
 #include "RendererInterface.h"
+#include "ProfilingDebugging/ScopedDebugInfo.h"
 
 #define D3DERR(x) case x: ErrorCodeText = TEXT(#x); break;
 #define LOCTEXT_NAMESPACE "Developer.MessageLog"
@@ -291,8 +292,15 @@ void VerifyD3D11CreateTextureResult(HRESULT D3DResult, int32 UEFormat,const ANSI
 	const FString ErrorString = GetD3D11ErrorString(D3DResult, 0);
 	const TCHAR* D3DFormatString = GetD3D11TextureFormatString((DXGI_FORMAT)D3DFormat);
 
+	FString DebugInfoString;
+
+	if (FScopedDebugInfo* DebugInfo = FScopedDebugInfo::GetDebugInfoStack())
+	{
+		DebugInfoString = DebugInfo->GetFunctionName();
+	}
+
 	UE_LOG(LogD3D11RHI, Error,
-		TEXT("%s failed with error %s\n at %s:%u\n Size=%ix%ix%i PF=%d D3DFormat=%s(0x%08X), NumMips=%i, Flags=%s, Usage:0x%x, CPUFlags:0x%x, MiscFlags:0x%x, SampleCount:0x%x, SampleQuality:0x%x, SubresPtr:0x%p, SubresPitch:%i, SubresSlicePitch:%i"),
+		TEXT("%s failed with error %s\n at %s:%u\n Size=%ix%ix%i PF=%d D3DFormat=%s(0x%08X), NumMips=%i, Flags=%s, Usage:0x%x, CPUFlags:0x%x, MiscFlags:0x%x, SampleCount:0x%x, SampleQuality:0x%x, SubresPtr:0x%p, SubresPitch:%i, SubresSlicePitch:%i, DebugInfo: %s"),
 		ANSI_TO_TCHAR(Code),
 		*ErrorString,
 		ANSI_TO_TCHAR(Filename),
@@ -312,7 +320,8 @@ void VerifyD3D11CreateTextureResult(HRESULT D3DResult, int32 UEFormat,const ANSI
 		SampleQuality,
 		SubResPtr,
 		SubResPitch,
-		SubResSlicePitch);
+		SubResSlicePitch,
+		*DebugInfoString);
 
 	TerminateOnDeviceRemoved(D3DResult, Device);
 	TerminateOnOutOfMemory(D3DResult, true);
