Fixed Horde issue #220116 and issue #220129 (using ELLMTag enum when LLM is not enabled).

diff --git a/Engine/Source/Runtime/D3D12RHI/Private/D3D12Texture.cpp b/Engine/Source/Runtime/D3D12RHI/Private/D3D12Texture.cpp
--- a/Engine/Source/Runtime/D3D12RHI/Private/D3D12Texture.cpp
+++ b/Engine/Source/Runtime/D3D12RHI/Private/D3D12Texture.cpp
@@ -392,7 +392,7 @@ void FD3D12TextureStats::UpdateD3D12TextureStats(FD3D12Texture& Texture, const D
 		LLM(FLowLevelMemTracker::Get().OnLowLevelAlloc(ELLMTracker::Platform, Texture.GetResource(), TextureSize, ELLMTag::GraphicsPlatform));
 		LLM(FLowLevelMemTracker::Get().OnLowLevelAlloc(ELLMTracker::Default, Texture.GetResource(), TextureSize, ELLMTag::Textures));
 		{
-			UE_MEMSCOPE_DEFAULT(ELLMTag::Textures);
+			LLM(UE_MEMSCOPE_DEFAULT(ELLMTag::Textures));
 			MemoryTrace_Alloc((uint64)Texture.GetResource(), TextureSize, 1024, EMemoryTraceRootHeap::VideoMemory);
 		}
 #endif
diff --git a/Engine/Source/Runtime/Windows/D3D11RHI/Private/D3D11Texture.cpp b/Engine/Source/Runtime/Windows/D3D11RHI/Private/D3D11Texture.cpp
--- a/Engine/Source/Runtime/Windows/D3D11RHI/Private/D3D11Texture.cpp
+++ b/Engine/Source/Runtime/Windows/D3D11RHI/Private/D3D11Texture.cpp
@@ -131,7 +131,7 @@ void UpdateD3D11TextureStats(FD3D11Texture& Texture, bool bAllocating)
 		LLM(FLowLevelMemTracker::Get().OnLowLevelAlloc(ELLMTracker::Platform, Texture.GetResource(), TextureSize, ELLMTag::GraphicsPlatform));
 		LLM(FLowLevelMemTracker::Get().OnLowLevelAlloc(ELLMTracker::Default , Texture.GetResource(), TextureSize, ELLMTag::Textures));
 		{
-			UE_MEMSCOPE_DEFAULT(ELLMTag::Textures);
+			LLM(UE_MEMSCOPE_DEFAULT(ELLMTag::Textures));
 			MemoryTrace_Alloc((uint64)Texture.GetResource(), TextureSize, 1024, EMemoryTraceRootHeap::VideoMemory);
 		}
 	}
