Removing more D3D11RHI/Private include paths and RHICore dependencies.

diff --git a/Engine/Platforms/Hololens/Plugins/Media/HLMedia/Source/HLMedia/HLMedia_HoloLens.Build.cs b/Engine/Platforms/Hololens/Plugins/Media/HLMedia/Source/HLMedia/HLMedia_HoloLens.Build.cs
--- a/Engine/Platforms/Hololens/Plugins/Media/HLMedia/Source/HLMedia/HLMedia_HoloLens.Build.cs
+++ b/Engine/Platforms/Hololens/Plugins/Media/HLMedia/Source/HLMedia/HLMedia_HoloLens.Build.cs
@@ -8,8 +8,6 @@ namespace UnrealBuildTool.Rules
 	{
 		public HLMedia_HoloLens(ReadOnlyTargetRules Target) : base(Target)
 		{
-			var EngineDir = Path.GetFullPath(Target.RelativeEnginePath);
-			PrivateIncludePaths.Add(Path.Combine(EngineDir, "Platforms/HoloLens/Source/Runtime/Windows/D3D11RHI/Private"));
 		}
 	}
 }
diff --git a/Engine/Plugins/Media/HLMedia/Source/HLMedia/HLMedia.Build.cs b/Engine/Plugins/Media/HLMedia/Source/HLMedia/HLMedia.Build.cs
--- a/Engine/Plugins/Media/HLMedia/Source/HLMedia/HLMedia.Build.cs
+++ b/Engine/Plugins/Media/HLMedia/Source/HLMedia/HLMedia.Build.cs
@@ -59,10 +59,5 @@ public class HLMedia : ModuleRules
 			
 		// For D3D11on12
 		AddEngineThirdPartyPrivateStaticDependencies(Target, "DX12", "DX11");
-
-        if (Target.Platform == UnrealTargetPlatform.Win64)
-        {
-            PrivateIncludePaths.Add(Path.Combine(EngineDir, "Source/Runtime/Windows/D3D11RHI/Private/Windows"));
-        }
     }
 }
diff --git a/Engine/Source/Programs/UnrealBuildTool/Platform/Windows/UEBuildWindows.cs b/Engine/Source/Programs/UnrealBuildTool/Platform/Windows/UEBuildWindows.cs
--- a/Engine/Source/Programs/UnrealBuildTool/Platform/Windows/UEBuildWindows.cs
+++ b/Engine/Source/Programs/UnrealBuildTool/Platform/Windows/UEBuildWindows.cs
@@ -1326,12 +1326,6 @@ namespace UnrealBuildTool
 				}
 			}
 
-			if (ModuleName == "D3D11RHI")
-			{
-				// To enable platform specific D3D11 RHI Types
-				Rules.PrivateIncludePaths.Add("Runtime/Windows/D3D11RHI/Private/Windows");
-			}
-
 			// Delay-load D3D12 so we can use the latest features and still run on downlevel versions of the OS
 			Rules.PublicDelayLoadDLLs.Add("d3d12.dll");
 		}
diff --git a/Engine/Source/Runtime/Windows/D3D11RHI/D3D11RHI.Build.cs b/Engine/Source/Runtime/Windows/D3D11RHI/D3D11RHI.Build.cs
--- a/Engine/Source/Runtime/Windows/D3D11RHI/D3D11RHI.Build.cs
+++ b/Engine/Source/Runtime/Windows/D3D11RHI/D3D11RHI.Build.cs
@@ -1,5 +1,6 @@
 // Copyright Epic Games, Inc. All Rights Reserved.
 
+using System.IO;
 using UnrealBuildTool;
 
 [SupportedPlatforms("Win64")]
@@ -19,14 +20,12 @@ public class D3D11RHI : ModuleRules
 				"RHI",
 				"RHICore",
 				"RenderCore"
-			}
-			);
+			});
 
 		PublicIncludePathModuleNames.AddRange(
 			new string[] {
-					"HeadMountedDisplay"
-			}
-			);
+				"HeadMountedDisplay"
+			});
 
 		AddEngineThirdPartyPrivateStaticDependencies(Target, "DX11");
 		if (bIncludeExtensions)
@@ -38,6 +37,10 @@ public class D3D11RHI : ModuleRules
 			AddEngineThirdPartyPrivateStaticDependencies(Target, "IntelExtensionsFramework");
 		}
 
+		if (Target.Platform == UnrealTargetPlatform.Win64)
+		{
+			PrivateIncludePaths.Add(Path.Combine(this.ModuleDirectory, "Private", "Windows"));
+		}
 
         if (Target.Configuration != UnrealTargetConfiguration.Shipping)
 		{
