Removing ControlPointPatchList enum values. Tesselation support was removed in 5.0 so these don't need deprecation.

diff --git a/Engine/Source/Runtime/Apple/MetalRHI/Private/MetalCommands.cpp b/Engine/Source/Runtime/Apple/MetalRHI/Private/MetalCommands.cpp
--- a/Engine/Source/Runtime/Apple/MetalRHI/Private/MetalCommands.cpp
+++ b/Engine/Source/Runtime/Apple/MetalRHI/Private/MetalCommands.cpp
@@ -49,42 +49,6 @@ mtlpp::PrimitiveType TranslatePrimitiveType(uint32 PrimitiveType)
 		case PT_TriangleStrip:	return mtlpp::PrimitiveType::TriangleStrip;
 		case PT_LineList:		return mtlpp::PrimitiveType::Line;
 		case PT_PointList:		return mtlpp::PrimitiveType::Point;
-		// Metal doesn't actually actually draw in control-point patch-lists because of the way the compute shader stage works - it can handle any arbitrary patch size and will output triangles.
-		case PT_1_ControlPointPatchList:
-		case PT_2_ControlPointPatchList:
-		case PT_3_ControlPointPatchList:
-		case PT_4_ControlPointPatchList:
-		case PT_5_ControlPointPatchList:
-		case PT_6_ControlPointPatchList:
-		case PT_7_ControlPointPatchList:
-		case PT_8_ControlPointPatchList:
-		case PT_9_ControlPointPatchList:
-		case PT_10_ControlPointPatchList:
-		case PT_11_ControlPointPatchList:
-		case PT_12_ControlPointPatchList:
-		case PT_13_ControlPointPatchList:
-		case PT_14_ControlPointPatchList:
-		case PT_15_ControlPointPatchList:
-		case PT_16_ControlPointPatchList:
-		case PT_17_ControlPointPatchList:
-		case PT_18_ControlPointPatchList:
-		case PT_19_ControlPointPatchList:
-		case PT_20_ControlPointPatchList:
-		case PT_21_ControlPointPatchList:
-		case PT_22_ControlPointPatchList:
-		case PT_23_ControlPointPatchList:
-		case PT_24_ControlPointPatchList:
-		case PT_25_ControlPointPatchList:
-		case PT_26_ControlPointPatchList:
-		case PT_27_ControlPointPatchList:
-		case PT_28_ControlPointPatchList:
-		case PT_29_ControlPointPatchList:
-		case PT_30_ControlPointPatchList:
-		case PT_31_ControlPointPatchList:
-		case PT_32_ControlPointPatchList:
-		{
-			return mtlpp::PrimitiveType::Triangle;
-		}
 		default:
 			METAL_FATAL_ERROR(TEXT("Unsupported primitive type %d"), (int32)PrimitiveType);
 			return mtlpp::PrimitiveType::Triangle;
diff --git a/Engine/Source/Runtime/Apple/MetalRHI/Private/MetalContext.cpp b/Engine/Source/Runtime/Apple/MetalRHI/Private/MetalContext.cpp
--- a/Engine/Source/Runtime/Apple/MetalRHI/Private/MetalContext.cpp
+++ b/Engine/Source/Runtime/Apple/MetalRHI/Private/MetalContext.cpp
@@ -303,41 +303,6 @@ mtlpp::PrimitiveTopologyClass TranslatePrimitiveTopology(uint32 PrimitiveType)
 			return mtlpp::PrimitiveTopologyClass::Line;
 		case PT_PointList:
 			return mtlpp::PrimitiveTopologyClass::Point;
-		case PT_1_ControlPointPatchList:
-		case PT_2_ControlPointPatchList:
-		case PT_3_ControlPointPatchList:
-		case PT_4_ControlPointPatchList:
-		case PT_5_ControlPointPatchList:
-		case PT_6_ControlPointPatchList:
-		case PT_7_ControlPointPatchList:
-		case PT_8_ControlPointPatchList:
-		case PT_9_ControlPointPatchList:
-		case PT_10_ControlPointPatchList:
-		case PT_11_ControlPointPatchList:
-		case PT_12_ControlPointPatchList:
-		case PT_13_ControlPointPatchList:
-		case PT_14_ControlPointPatchList:
-		case PT_15_ControlPointPatchList:
-		case PT_16_ControlPointPatchList:
-		case PT_17_ControlPointPatchList:
-		case PT_18_ControlPointPatchList:
-		case PT_19_ControlPointPatchList:
-		case PT_20_ControlPointPatchList:
-		case PT_21_ControlPointPatchList:
-		case PT_22_ControlPointPatchList:
-		case PT_23_ControlPointPatchList:
-		case PT_24_ControlPointPatchList:
-		case PT_25_ControlPointPatchList:
-		case PT_26_ControlPointPatchList:
-		case PT_27_ControlPointPatchList:
-		case PT_28_ControlPointPatchList:
-		case PT_29_ControlPointPatchList:
-		case PT_30_ControlPointPatchList:
-		case PT_31_ControlPointPatchList:
-		case PT_32_ControlPointPatchList:
-		{
-			return mtlpp::PrimitiveTopologyClass::Triangle;
-		}
 		default:
 			UE_LOG(LogMetal, Fatal, TEXT("Unsupported primitive topology %d"), (int32)PrimitiveType);
 			return mtlpp::PrimitiveTopologyClass::Triangle;
diff --git a/Engine/Source/Runtime/D3D12RHI/Private/D3D12StateCache.cpp b/Engine/Source/Runtime/D3D12RHI/Private/D3D12StateCache.cpp
--- a/Engine/Source/Runtime/D3D12RHI/Private/D3D12StateCache.cpp
+++ b/Engine/Source/Runtime/D3D12RHI/Private/D3D12StateCache.cpp
@@ -1355,8 +1355,8 @@ void FD3D12StateCache::SetGraphicsPipelineState(FD3D12GraphicsPipelineState* Gra
 			PipelineState.Graphics.CurrentPrimitiveTopology = GetD3D12PrimitiveType(PrimitiveType);
 			bNeedSetPrimitiveTopology = true;
 
-			static_assert(PT_Num == 38, "This computation needs to be updated, matching that of GetVertexCountForPrimitiveCount()");
-			PipelineState.Graphics.PrimitiveTypeFactor = (PrimitiveType == PT_TriangleList) ? 3 : (PrimitiveType == PT_LineList) ? 2 : (PrimitiveType == PT_RectList) ? 3 : (PrimitiveType >= PT_1_ControlPointPatchList) ? (PrimitiveType - PT_1_ControlPointPatchList + 1) : 1;
+			static_assert(PT_Num == 6, "This computation needs to be updated, matching that of GetVertexCountForPrimitiveCount()");
+			PipelineState.Graphics.PrimitiveTypeFactor = (PrimitiveType == PT_TriangleList) ? 3 : (PrimitiveType == PT_LineList) ? 2 : (PrimitiveType == PT_RectList) ? 3 : 1;
 			PipelineState.Graphics.PrimitiveTypeOffset = (PrimitiveType == PT_TriangleStrip) ? 2 : 0;
 		}
 
diff --git a/Engine/Source/Runtime/D3D12RHI/Private/D3D12StateCachePrivate.h b/Engine/Source/Runtime/D3D12RHI/Private/D3D12StateCachePrivate.h
--- a/Engine/Source/Runtime/D3D12RHI/Private/D3D12StateCachePrivate.h
+++ b/Engine/Source/Runtime/D3D12RHI/Private/D3D12StateCachePrivate.h
@@ -269,38 +269,6 @@ static inline D3D_PRIMITIVE_TOPOLOGY GetD3D12PrimitiveType(uint32 PrimitiveType)
 #else
 		0,
 #endif
-		D3D_PRIMITIVE_TOPOLOGY_1_CONTROL_POINT_PATCHLIST,  // PT_1_ControlPointPatchList
-		D3D_PRIMITIVE_TOPOLOGY_2_CONTROL_POINT_PATCHLIST,  // PT_2_ControlPointPatchList
-		D3D_PRIMITIVE_TOPOLOGY_3_CONTROL_POINT_PATCHLIST,  // PT_3_ControlPointPatchList
-		D3D_PRIMITIVE_TOPOLOGY_4_CONTROL_POINT_PATCHLIST,  // PT_4_ControlPointPatchList
-		D3D_PRIMITIVE_TOPOLOGY_5_CONTROL_POINT_PATCHLIST,  // PT_5_ControlPointPatchList
-		D3D_PRIMITIVE_TOPOLOGY_6_CONTROL_POINT_PATCHLIST,  // PT_6_ControlPointPatchList
-		D3D_PRIMITIVE_TOPOLOGY_7_CONTROL_POINT_PATCHLIST,  // PT_7_ControlPointPatchList
-		D3D_PRIMITIVE_TOPOLOGY_8_CONTROL_POINT_PATCHLIST,  // PT_8_ControlPointPatchList
-		D3D_PRIMITIVE_TOPOLOGY_9_CONTROL_POINT_PATCHLIST,  // PT_9_ControlPointPatchList
-		D3D_PRIMITIVE_TOPOLOGY_10_CONTROL_POINT_PATCHLIST, // PT_10_ControlPointPatchList
-		D3D_PRIMITIVE_TOPOLOGY_11_CONTROL_POINT_PATCHLIST, // PT_11_ControlPointPatchList
-		D3D_PRIMITIVE_TOPOLOGY_12_CONTROL_POINT_PATCHLIST, // PT_12_ControlPointPatchList
-		D3D_PRIMITIVE_TOPOLOGY_13_CONTROL_POINT_PATCHLIST, // PT_13_ControlPointPatchList
-		D3D_PRIMITIVE_TOPOLOGY_14_CONTROL_POINT_PATCHLIST, // PT_14_ControlPointPatchList
-		D3D_PRIMITIVE_TOPOLOGY_15_CONTROL_POINT_PATCHLIST, // PT_15_ControlPointPatchList
-		D3D_PRIMITIVE_TOPOLOGY_16_CONTROL_POINT_PATCHLIST, // PT_16_ControlPointPatchList
-		D3D_PRIMITIVE_TOPOLOGY_17_CONTROL_POINT_PATCHLIST, // PT_17_ControlPointPatchList
-		D3D_PRIMITIVE_TOPOLOGY_18_CONTROL_POINT_PATCHLIST, // PT_18_ControlPointPatchList
-		D3D_PRIMITIVE_TOPOLOGY_19_CONTROL_POINT_PATCHLIST, // PT_19_ControlPointPatchList
-		D3D_PRIMITIVE_TOPOLOGY_20_CONTROL_POINT_PATCHLIST, // PT_20_ControlPointPatchList
-		D3D_PRIMITIVE_TOPOLOGY_21_CONTROL_POINT_PATCHLIST, // PT_21_ControlPointPatchList
-		D3D_PRIMITIVE_TOPOLOGY_22_CONTROL_POINT_PATCHLIST, // PT_22_ControlPointPatchList
-		D3D_PRIMITIVE_TOPOLOGY_23_CONTROL_POINT_PATCHLIST, // PT_23_ControlPointPatchList
-		D3D_PRIMITIVE_TOPOLOGY_24_CONTROL_POINT_PATCHLIST, // PT_24_ControlPointPatchList
-		D3D_PRIMITIVE_TOPOLOGY_25_CONTROL_POINT_PATCHLIST, // PT_25_ControlPointPatchList
-		D3D_PRIMITIVE_TOPOLOGY_26_CONTROL_POINT_PATCHLIST, // PT_26_ControlPointPatchList
-		D3D_PRIMITIVE_TOPOLOGY_27_CONTROL_POINT_PATCHLIST, // PT_27_ControlPointPatchList
-		D3D_PRIMITIVE_TOPOLOGY_28_CONTROL_POINT_PATCHLIST, // PT_28_ControlPointPatchList
-		D3D_PRIMITIVE_TOPOLOGY_29_CONTROL_POINT_PATCHLIST, // PT_29_ControlPointPatchList
-		D3D_PRIMITIVE_TOPOLOGY_30_CONTROL_POINT_PATCHLIST, // PT_30_ControlPointPatchList
-		D3D_PRIMITIVE_TOPOLOGY_31_CONTROL_POINT_PATCHLIST, // PT_31_ControlPointPatchList
-		D3D_PRIMITIVE_TOPOLOGY_32_CONTROL_POINT_PATCHLIST, // PT_32_ControlPointPatchList
 	};
 	static_assert(UE_ARRAY_COUNT(D3D12PrimitiveType) == PT_Num, "Primitive lookup table is wrong size");
 
diff --git a/Engine/Source/Runtime/D3D12RHI/Public/D3D12Util.h b/Engine/Source/Runtime/D3D12RHI/Public/D3D12Util.h
--- a/Engine/Source/Runtime/D3D12RHI/Public/D3D12Util.h
+++ b/Engine/Source/Runtime/D3D12RHI/Public/D3D12Util.h
@@ -761,38 +761,6 @@ FORCEINLINE_DEBUGGABLE D3D_PRIMITIVE_TOPOLOGY TranslatePrimitiveType(EPrimitiveT
 		case PT_RectList:				return D3D12RHI_PRIMITIVE_TOPOLOGY_RECTLIST;
 	#endif
 
-	case PT_1_ControlPointPatchList:	return D3D_PRIMITIVE_TOPOLOGY_1_CONTROL_POINT_PATCHLIST;
-	case PT_2_ControlPointPatchList:	return D3D_PRIMITIVE_TOPOLOGY_2_CONTROL_POINT_PATCHLIST;
-	case PT_3_ControlPointPatchList:	return D3D_PRIMITIVE_TOPOLOGY_3_CONTROL_POINT_PATCHLIST;
-	case PT_4_ControlPointPatchList:	return D3D_PRIMITIVE_TOPOLOGY_4_CONTROL_POINT_PATCHLIST;
-	case PT_5_ControlPointPatchList:	return D3D_PRIMITIVE_TOPOLOGY_5_CONTROL_POINT_PATCHLIST;
-	case PT_6_ControlPointPatchList:	return D3D_PRIMITIVE_TOPOLOGY_6_CONTROL_POINT_PATCHLIST;
-	case PT_7_ControlPointPatchList:	return D3D_PRIMITIVE_TOPOLOGY_7_CONTROL_POINT_PATCHLIST;
-	case PT_8_ControlPointPatchList:	return D3D_PRIMITIVE_TOPOLOGY_8_CONTROL_POINT_PATCHLIST;
-	case PT_9_ControlPointPatchList:	return D3D_PRIMITIVE_TOPOLOGY_9_CONTROL_POINT_PATCHLIST;
-	case PT_10_ControlPointPatchList:	return D3D_PRIMITIVE_TOPOLOGY_10_CONTROL_POINT_PATCHLIST;
-	case PT_11_ControlPointPatchList:	return D3D_PRIMITIVE_TOPOLOGY_11_CONTROL_POINT_PATCHLIST;
-	case PT_12_ControlPointPatchList:	return D3D_PRIMITIVE_TOPOLOGY_12_CONTROL_POINT_PATCHLIST;
-	case PT_13_ControlPointPatchList:	return D3D_PRIMITIVE_TOPOLOGY_13_CONTROL_POINT_PATCHLIST;
-	case PT_14_ControlPointPatchList:	return D3D_PRIMITIVE_TOPOLOGY_14_CONTROL_POINT_PATCHLIST;
-	case PT_15_ControlPointPatchList:	return D3D_PRIMITIVE_TOPOLOGY_15_CONTROL_POINT_PATCHLIST;
-	case PT_16_ControlPointPatchList:	return D3D_PRIMITIVE_TOPOLOGY_16_CONTROL_POINT_PATCHLIST;
-	case PT_17_ControlPointPatchList:	return D3D_PRIMITIVE_TOPOLOGY_17_CONTROL_POINT_PATCHLIST;
-	case PT_18_ControlPointPatchList:	return D3D_PRIMITIVE_TOPOLOGY_18_CONTROL_POINT_PATCHLIST;
-	case PT_19_ControlPointPatchList:	return D3D_PRIMITIVE_TOPOLOGY_19_CONTROL_POINT_PATCHLIST;
-	case PT_20_ControlPointPatchList:	return D3D_PRIMITIVE_TOPOLOGY_20_CONTROL_POINT_PATCHLIST;
-	case PT_21_ControlPointPatchList:	return D3D_PRIMITIVE_TOPOLOGY_21_CONTROL_POINT_PATCHLIST;
-	case PT_22_ControlPointPatchList:	return D3D_PRIMITIVE_TOPOLOGY_22_CONTROL_POINT_PATCHLIST;
-	case PT_23_ControlPointPatchList:	return D3D_PRIMITIVE_TOPOLOGY_23_CONTROL_POINT_PATCHLIST;
-	case PT_24_ControlPointPatchList:	return D3D_PRIMITIVE_TOPOLOGY_24_CONTROL_POINT_PATCHLIST;
-	case PT_25_ControlPointPatchList:	return D3D_PRIMITIVE_TOPOLOGY_25_CONTROL_POINT_PATCHLIST;
-	case PT_26_ControlPointPatchList:	return D3D_PRIMITIVE_TOPOLOGY_26_CONTROL_POINT_PATCHLIST;
-	case PT_27_ControlPointPatchList:	return D3D_PRIMITIVE_TOPOLOGY_27_CONTROL_POINT_PATCHLIST;
-	case PT_28_ControlPointPatchList:	return D3D_PRIMITIVE_TOPOLOGY_28_CONTROL_POINT_PATCHLIST;
-	case PT_29_ControlPointPatchList:	return D3D_PRIMITIVE_TOPOLOGY_29_CONTROL_POINT_PATCHLIST;
-	case PT_30_ControlPointPatchList:	return D3D_PRIMITIVE_TOPOLOGY_30_CONTROL_POINT_PATCHLIST;
-	case PT_31_ControlPointPatchList:	return D3D_PRIMITIVE_TOPOLOGY_31_CONTROL_POINT_PATCHLIST;
-	case PT_32_ControlPointPatchList:	return D3D_PRIMITIVE_TOPOLOGY_32_CONTROL_POINT_PATCHLIST;
 	default:
 		ensure(0);
 		return D3D_PRIMITIVE_TOPOLOGY_UNDEFINED;
diff --git a/Engine/Source/Runtime/OpenGLDrv/Private/OpenGLCommands.cpp b/Engine/Source/Runtime/OpenGLDrv/Private/OpenGLCommands.cpp
--- a/Engine/Source/Runtime/OpenGLDrv/Private/OpenGLCommands.cpp
+++ b/Engine/Source/Runtime/OpenGLDrv/Private/OpenGLCommands.cpp
@@ -2285,8 +2285,7 @@ void FOpenGLDynamicRHI::RHIDrawPrimitive(uint32 BaseVertexIndex,uint32 NumPrimit
 
 	GLenum DrawMode = GL_TRIANGLES;
 	GLsizei NumElements = 0;
-	GLint PatchSize = 0;
-	FindPrimitiveType(PrimitiveType, NumPrimitives, DrawMode, NumElements, PatchSize);
+	FindPrimitiveType(PrimitiveType, NumPrimitives, DrawMode, NumElements);
 
 	GPUProfilingData.RegisterGPUWork(NumPrimitives * NumInstances, VertexCount * NumInstances);
 	if (NumInstances == 1)
@@ -2331,8 +2330,7 @@ void FOpenGLDynamicRHI::RHIDrawPrimitiveIndirect(FRHIBuffer* ArgumentBufferRHI,
 
 		GLenum DrawMode = GL_TRIANGLES;
 		GLsizei NumElements = 0;
-		GLint PatchSize = 0;
-		FindPrimitiveType(PrimitiveType, 0, DrawMode, NumElements, PatchSize);
+		FindPrimitiveType(PrimitiveType, 0, DrawMode, NumElements);
 
 		FOpenGLBuffer* ArgumentBuffer = ResourceCast(ArgumentBufferRHI);
 
@@ -2384,8 +2382,7 @@ void FOpenGLDynamicRHI::RHIDrawIndexedIndirect(FRHIBuffer* IndexBufferRHI, FRHIB
 
 		GLenum DrawMode = GL_TRIANGLES;
 		GLsizei NumElements = 0;
-		GLint PatchSize = 0;
-		FindPrimitiveType(PrimitiveType, 0, DrawMode, NumElements, PatchSize);
+		FindPrimitiveType(PrimitiveType, 0, DrawMode, NumElements);
 
 		GLenum IndexType = IndexBuffer->GetStride() == sizeof(uint32) ? GL_UNSIGNED_INT : GL_UNSIGNED_SHORT;
 
@@ -2472,8 +2469,7 @@ void FOpenGLDynamicRHI::RHIDrawIndexedPrimitive(FRHIBuffer* IndexBufferRHI, int3
 
 	GLenum DrawMode = GL_TRIANGLES;
 	GLsizei NumElements = 0;
-	GLint PatchSize = 0;
-	FindPrimitiveType(PrimitiveType, NumPrimitives, DrawMode, NumElements, PatchSize);
+	FindPrimitiveType(PrimitiveType, NumPrimitives, DrawMode, NumElements);
 
 	GLenum IndexType = IndexBuffer->GetStride() == sizeof(uint32) ? GL_UNSIGNED_INT : GL_UNSIGNED_SHORT;
 	StartIndex *= IndexBuffer->GetStride() == sizeof(uint32) ? sizeof(uint32) : sizeof(uint16);
@@ -2531,8 +2527,7 @@ void FOpenGLDynamicRHI::RHIDrawIndexedPrimitiveIndirect(FRHIBuffer* IndexBufferR
 
 		GLenum DrawMode = GL_TRIANGLES;
 		GLsizei NumElements = 0;
-		GLint PatchSize = 0;
-		FindPrimitiveType(PrimitiveType, 0, DrawMode, NumElements, PatchSize);
+		FindPrimitiveType(PrimitiveType, 0, DrawMode, NumElements);
 
 		GLenum IndexType = IndexBuffer->GetStride() == sizeof(uint32) ? GL_UNSIGNED_INT : GL_UNSIGNED_SHORT;
 
diff --git a/Engine/Source/Runtime/OpenGLDrv/Private/OpenGLDrvPrivate.h b/Engine/Source/Runtime/OpenGLDrv/Private/OpenGLDrvPrivate.h
--- a/Engine/Source/Runtime/OpenGLDrv/Private/OpenGLDrvPrivate.h
+++ b/Engine/Source/Runtime/OpenGLDrv/Private/OpenGLDrvPrivate.h
@@ -346,11 +346,10 @@ inline uint32 FindMaxMipmapLevel(uint32 Width, uint32 Height, uint32 Depth)
 	return FindMaxMipmapLevel((Width > Height) ? Width : Height, Depth);
 }
 
-inline void FindPrimitiveType(uint32 InPrimitiveType, uint32 InNumPrimitives, GLenum &DrawMode, GLsizei &NumElements, GLint &PatchSize)
+inline void FindPrimitiveType(uint32 InPrimitiveType, uint32 InNumPrimitives, GLenum &DrawMode, GLsizei &NumElements)
 {
 	DrawMode = GL_TRIANGLES;
 	NumElements = InNumPrimitives;
-	PatchSize = 0;
 
 	switch (InPrimitiveType)
 	{
@@ -370,42 +369,6 @@ inline void FindPrimitiveType(uint32 InPrimitiveType, uint32 InNumPrimitives, GL
 		DrawMode = GL_POINTS;
 		NumElements = InNumPrimitives;
 		break;
-	case PT_1_ControlPointPatchList:
-	case PT_2_ControlPointPatchList:
-	case PT_3_ControlPointPatchList:
-	case PT_4_ControlPointPatchList:
-	case PT_5_ControlPointPatchList:
-	case PT_6_ControlPointPatchList:
-	case PT_7_ControlPointPatchList:
-	case PT_8_ControlPointPatchList:
-	case PT_9_ControlPointPatchList:
-	case PT_10_ControlPointPatchList:
-	case PT_11_ControlPointPatchList:
-	case PT_12_ControlPointPatchList:
-	case PT_13_ControlPointPatchList:
-	case PT_14_ControlPointPatchList:
-	case PT_15_ControlPointPatchList:
-	case PT_16_ControlPointPatchList:
-	case PT_17_ControlPointPatchList:
-	case PT_18_ControlPointPatchList:
-	case PT_19_ControlPointPatchList:
-	case PT_20_ControlPointPatchList:
-	case PT_21_ControlPointPatchList:
-	case PT_22_ControlPointPatchList:
-	case PT_23_ControlPointPatchList:
-	case PT_24_ControlPointPatchList:
-	case PT_25_ControlPointPatchList:
-	case PT_26_ControlPointPatchList:
-	case PT_27_ControlPointPatchList:
-	case PT_28_ControlPointPatchList:
-	case PT_29_ControlPointPatchList:
-	case PT_30_ControlPointPatchList:
-	case PT_31_ControlPointPatchList:
-	case PT_32_ControlPointPatchList:
-		DrawMode = GL_PATCHES;
-		PatchSize = InPrimitiveType - uint32(PT_1_ControlPointPatchList) + 1;
-		NumElements = InNumPrimitives * PatchSize;
-		break;
 	default:
 		UE_LOG(LogRHI, Fatal,TEXT("Unsupported primitive type %u"), InPrimitiveType);
 		break;
diff --git a/Engine/Source/Runtime/RHI/Private/PipelineFileCache.cpp b/Engine/Source/Runtime/RHI/Private/PipelineFileCache.cpp
--- a/Engine/Source/Runtime/RHI/Private/PipelineFileCache.cpp
+++ b/Engine/Source/Runtime/RHI/Private/PipelineFileCache.cpp
@@ -952,12 +952,6 @@ bool FPipelineCacheFileFormatPSO::Verify() const
 		}
 #endif
 		
-		if(GraphicsDesc.PrimitiveType >= PT_1_ControlPointPatchList && GraphicsDesc.PrimitiveType <= PT_32_ControlPointPatchList)
-		{
-			// Define says we don't support tessellation - can't draw patches - not a valid PSO for target platform
-			return false;
-		}
-
 #if PLATFORM_SUPPORTS_GEOMETRY_SHADERS
 		// Is there anything to actually test here?
 #endif
diff --git a/Engine/Source/Runtime/RHI/Public/RHIDefinitions.h b/Engine/Source/Runtime/RHI/Public/RHIDefinitions.h
--- a/Engine/Source/Runtime/RHI/Public/RHIDefinitions.h
+++ b/Engine/Source/Runtime/RHI/Public/RHIDefinitions.h
@@ -1624,41 +1624,8 @@ enum EPrimitiveType
 	// Supported only if GRHISupportsRectTopology == true.
 	PT_RectList,
 
-	// Tesselation patch list. Supported only if tesselation is supported.
-	PT_1_ControlPointPatchList,
-	PT_2_ControlPointPatchList,
-	PT_3_ControlPointPatchList,
-	PT_4_ControlPointPatchList,
-	PT_5_ControlPointPatchList,
-	PT_6_ControlPointPatchList,
-	PT_7_ControlPointPatchList,
-	PT_8_ControlPointPatchList,
-	PT_9_ControlPointPatchList,
-	PT_10_ControlPointPatchList,
-	PT_11_ControlPointPatchList,
-	PT_12_ControlPointPatchList,
-	PT_13_ControlPointPatchList,
-	PT_14_ControlPointPatchList,
-	PT_15_ControlPointPatchList,
-	PT_16_ControlPointPatchList,
-	PT_17_ControlPointPatchList,
-	PT_18_ControlPointPatchList,
-	PT_19_ControlPointPatchList,
-	PT_20_ControlPointPatchList,
-	PT_21_ControlPointPatchList,
-	PT_22_ControlPointPatchList,
-	PT_23_ControlPointPatchList,
-	PT_24_ControlPointPatchList,
-	PT_25_ControlPointPatchList,
-	PT_26_ControlPointPatchList,
-	PT_27_ControlPointPatchList,
-	PT_28_ControlPointPatchList,
-	PT_29_ControlPointPatchList,
-	PT_30_ControlPointPatchList,
-	PT_31_ControlPointPatchList,
-	PT_32_ControlPointPatchList,
 	PT_Num,
-	PT_NumBits = 6
+	PT_NumBits = 3
 };
 static_assert(PT_Num <= (1 << 8), "EPrimitiveType doesn't fit in a byte");
 static_assert(PT_Num <= (1 << PT_NumBits), "PT_NumBits is too small");
diff --git a/Engine/Source/Runtime/RHI/Public/RHIUtilities.h b/Engine/Source/Runtime/RHI/Public/RHIUtilities.h
--- a/Engine/Source/Runtime/RHI/Public/RHIUtilities.h
+++ b/Engine/Source/Runtime/RHI/Public/RHIUtilities.h
@@ -803,9 +803,9 @@ inline void TransitionAndCopyTexture(FRHICommandList& RHICmdList, FRHITexture* S
  */
 inline uint32 GetVertexCountForPrimitiveCount(uint32 NumPrimitives, uint32 PrimitiveType)
 {
-	static_assert(PT_Num == 38, "This function needs to be updated");
-	uint32 Factor = (PrimitiveType == PT_TriangleList)? 3 : (PrimitiveType == PT_LineList)? 2 : (PrimitiveType == PT_RectList)? 3 : (PrimitiveType >= PT_1_ControlPointPatchList)? (PrimitiveType - PT_1_ControlPointPatchList + 1) : 1;
-	uint32 Offset = (PrimitiveType == PT_TriangleStrip)? 2 : 0;
+	static_assert(PT_Num == 6, "This function needs to be updated");
+	uint32 Factor = (PrimitiveType == PT_TriangleList) ? 3 : (PrimitiveType == PT_LineList) ? 2 : (PrimitiveType == PT_RectList) ? 3 : 1;
+	uint32 Offset = (PrimitiveType == PT_TriangleStrip) ? 2 : 0;
 
 	return NumPrimitives * Factor + Offset;
 
diff --git a/Engine/Source/Runtime/VulkanRHI/Private/VulkanPipeline.cpp b/Engine/Source/Runtime/VulkanRHI/Private/VulkanPipeline.cpp
--- a/Engine/Source/Runtime/VulkanRHI/Private/VulkanPipeline.cpp
+++ b/Engine/Source/Runtime/VulkanRHI/Private/VulkanPipeline.cpp
@@ -1626,39 +1626,6 @@ static inline VkPrimitiveTopology UEToVulkanTopologyType(const FVulkanDevice* In
 		return VK_PRIMITIVE_TOPOLOGY_TRIANGLE_LIST;
 	case PT_TriangleStrip:
 		return VK_PRIMITIVE_TOPOLOGY_TRIANGLE_STRIP;
-	case PT_1_ControlPointPatchList:
-	case PT_2_ControlPointPatchList:
-	case PT_3_ControlPointPatchList:
-	case PT_4_ControlPointPatchList:
-	case PT_5_ControlPointPatchList:
-	case PT_6_ControlPointPatchList:
-	case PT_7_ControlPointPatchList:
-	case PT_8_ControlPointPatchList:
-	case PT_9_ControlPointPatchList:
-	case PT_10_ControlPointPatchList:
-	case PT_12_ControlPointPatchList:
-	case PT_13_ControlPointPatchList:
-	case PT_14_ControlPointPatchList:
-	case PT_15_ControlPointPatchList:
-	case PT_16_ControlPointPatchList:
-	case PT_17_ControlPointPatchList:
-	case PT_18_ControlPointPatchList:
-	case PT_19_ControlPointPatchList:
-	case PT_20_ControlPointPatchList:
-	case PT_22_ControlPointPatchList:
-	case PT_23_ControlPointPatchList:
-	case PT_24_ControlPointPatchList:
-	case PT_25_ControlPointPatchList:
-	case PT_26_ControlPointPatchList:
-	case PT_27_ControlPointPatchList:
-	case PT_28_ControlPointPatchList:
-	case PT_29_ControlPointPatchList:
-	case PT_30_ControlPointPatchList:
-	case PT_31_ControlPointPatchList:
-	case PT_32_ControlPointPatchList:
-		OutControlPoints = (PrimitiveType - PT_1_ControlPointPatchList + 1);
-		checkf(false, TEXT("Missing tessellation shaders, however tried to use EPrimitiveType %d (%d control points)"), (uint32)PrimitiveType, OutControlPoints);
-		break;
 	default:
 		checkf(false, TEXT("Unsupported EPrimitiveType %d"), (uint32)PrimitiveType);
 		break;
diff --git a/Engine/Source/Runtime/Windows/D3D11RHI/Private/D3D11Commands.cpp b/Engine/Source/Runtime/Windows/D3D11RHI/Private/D3D11Commands.cpp
--- a/Engine/Source/Runtime/Windows/D3D11RHI/Private/D3D11Commands.cpp
+++ b/Engine/Source/Runtime/Windows/D3D11RHI/Private/D3D11Commands.cpp
@@ -1093,41 +1093,6 @@ static D3D11_PRIMITIVE_TOPOLOGY GetD3D11PrimitiveType(EPrimitiveType PrimitiveTy
 	case PT_LineList: return D3D11_PRIMITIVE_TOPOLOGY_LINELIST;
 	case PT_PointList: return D3D11_PRIMITIVE_TOPOLOGY_POINTLIST;
 
-	// ControlPointPatchList types will pretend to be TRIANGLELISTS with a stride of N 
-	// (where N is the number of control points specified), so we can return them for
-	// tessellation and non-tessellation. This functionality is only used when rendering a 
-	// default material with something that claims to be tessellated, generally because the 
-	// tessellation material failed to compile for some reason.
-	case PT_3_ControlPointPatchList: return D3D11_PRIMITIVE_TOPOLOGY_3_CONTROL_POINT_PATCHLIST;
-	case PT_4_ControlPointPatchList: return D3D11_PRIMITIVE_TOPOLOGY_4_CONTROL_POINT_PATCHLIST;
-	case PT_5_ControlPointPatchList: return D3D11_PRIMITIVE_TOPOLOGY_5_CONTROL_POINT_PATCHLIST;
-	case PT_6_ControlPointPatchList: return D3D11_PRIMITIVE_TOPOLOGY_6_CONTROL_POINT_PATCHLIST;
-	case PT_7_ControlPointPatchList: return D3D11_PRIMITIVE_TOPOLOGY_7_CONTROL_POINT_PATCHLIST;
-	case PT_8_ControlPointPatchList: return D3D11_PRIMITIVE_TOPOLOGY_8_CONTROL_POINT_PATCHLIST; 
-	case PT_9_ControlPointPatchList: return D3D11_PRIMITIVE_TOPOLOGY_9_CONTROL_POINT_PATCHLIST; 
-	case PT_10_ControlPointPatchList: return D3D11_PRIMITIVE_TOPOLOGY_10_CONTROL_POINT_PATCHLIST; 
-	case PT_11_ControlPointPatchList: return D3D11_PRIMITIVE_TOPOLOGY_11_CONTROL_POINT_PATCHLIST; 
-	case PT_12_ControlPointPatchList: return D3D11_PRIMITIVE_TOPOLOGY_12_CONTROL_POINT_PATCHLIST; 
-	case PT_13_ControlPointPatchList: return D3D11_PRIMITIVE_TOPOLOGY_13_CONTROL_POINT_PATCHLIST; 
-	case PT_14_ControlPointPatchList: return D3D11_PRIMITIVE_TOPOLOGY_14_CONTROL_POINT_PATCHLIST; 
-	case PT_15_ControlPointPatchList: return D3D11_PRIMITIVE_TOPOLOGY_15_CONTROL_POINT_PATCHLIST; 
-	case PT_16_ControlPointPatchList: return D3D11_PRIMITIVE_TOPOLOGY_16_CONTROL_POINT_PATCHLIST; 
-	case PT_17_ControlPointPatchList: return D3D11_PRIMITIVE_TOPOLOGY_17_CONTROL_POINT_PATCHLIST; 
-	case PT_18_ControlPointPatchList: return D3D11_PRIMITIVE_TOPOLOGY_18_CONTROL_POINT_PATCHLIST; 
-	case PT_19_ControlPointPatchList: return D3D11_PRIMITIVE_TOPOLOGY_19_CONTROL_POINT_PATCHLIST; 
-	case PT_20_ControlPointPatchList: return D3D11_PRIMITIVE_TOPOLOGY_20_CONTROL_POINT_PATCHLIST; 
-	case PT_21_ControlPointPatchList: return D3D11_PRIMITIVE_TOPOLOGY_21_CONTROL_POINT_PATCHLIST; 
-	case PT_22_ControlPointPatchList: return D3D11_PRIMITIVE_TOPOLOGY_22_CONTROL_POINT_PATCHLIST; 
-	case PT_23_ControlPointPatchList: return D3D11_PRIMITIVE_TOPOLOGY_23_CONTROL_POINT_PATCHLIST; 
-	case PT_24_ControlPointPatchList: return D3D11_PRIMITIVE_TOPOLOGY_24_CONTROL_POINT_PATCHLIST; 
-	case PT_25_ControlPointPatchList: return D3D11_PRIMITIVE_TOPOLOGY_25_CONTROL_POINT_PATCHLIST; 
-	case PT_26_ControlPointPatchList: return D3D11_PRIMITIVE_TOPOLOGY_26_CONTROL_POINT_PATCHLIST; 
-	case PT_27_ControlPointPatchList: return D3D11_PRIMITIVE_TOPOLOGY_27_CONTROL_POINT_PATCHLIST; 
-	case PT_28_ControlPointPatchList: return D3D11_PRIMITIVE_TOPOLOGY_28_CONTROL_POINT_PATCHLIST; 
-	case PT_29_ControlPointPatchList: return D3D11_PRIMITIVE_TOPOLOGY_29_CONTROL_POINT_PATCHLIST; 
-	case PT_30_ControlPointPatchList: return D3D11_PRIMITIVE_TOPOLOGY_30_CONTROL_POINT_PATCHLIST; 
-	case PT_31_ControlPointPatchList: return D3D11_PRIMITIVE_TOPOLOGY_31_CONTROL_POINT_PATCHLIST; 
-	case PT_32_ControlPointPatchList: return D3D11_PRIMITIVE_TOPOLOGY_32_CONTROL_POINT_PATCHLIST; 
 	default: UE_LOG(LogD3D11RHI, Fatal,TEXT("Unknown primitive type: %u"),PrimitiveType);
 	};
 
