Always make a copy of texture update data in D3D11. Fixes issue where memory is stale when the command is later run.

diff --git a/Engine/Source/Runtime/Windows/D3D11RHI/Private/D3D11Texture.cpp b/Engine/Source/Runtime/Windows/D3D11RHI/Private/D3D11Texture.cpp
--- a/Engine/Source/Runtime/Windows/D3D11RHI/Private/D3D11Texture.cpp
+++ b/Engine/Source/Runtime/Windows/D3D11RHI/Private/D3D11Texture.cpp
@@ -1516,19 +1516,14 @@ void FD3D11DynamicRHI::RHIUnlockTexture2DArray(FRHITexture2DArray* TextureRHI, u
 
 void FD3D11DynamicRHI::RHIUpdateTexture2D(FRHICommandListBase& RHICmdList, FRHITexture2D* TextureRHI, uint32 MipIndex, const FUpdateTextureRegion2D& UpdateRegion, uint32 SourcePitch, const uint8* SourceData)
 {
-	bool bFreeMemory = !(RHICmdList.Bypass() || !IsRunningRHIInSeparateThread());
+	const FPixelFormatInfo& FormatInfo = GPixelFormats[TextureRHI->GetFormat()];
+	const size_t UpdateHeightInTiles = FMath::DivideAndRoundUp(UpdateRegion.Height, (uint32)FormatInfo.BlockSizeY);
+	const size_t SourceDataSize = static_cast<size_t>(SourcePitch) * UpdateHeightInTiles;
+	uint8* SourceDataCopy = (uint8*)FMemory::Malloc(SourceDataSize);
+	FMemory::Memcpy(SourceDataCopy, SourceData, SourceDataSize);
+	SourceData = SourceDataCopy;
 
-	if (bFreeMemory)
-	{
-		const FPixelFormatInfo& FormatInfo = GPixelFormats[TextureRHI->GetFormat()];
-		const size_t UpdateHeightInTiles = FMath::DivideAndRoundUp(UpdateRegion.Height, (uint32)FormatInfo.BlockSizeY);
-		const size_t SourceDataSize = static_cast<size_t>(SourcePitch) * UpdateHeightInTiles;
-		uint8* SourceDataCopy = (uint8*)FMemory::Malloc(SourceDataSize);
-		FMemory::Memcpy(SourceDataCopy, SourceData, SourceDataSize);
-		SourceData = SourceDataCopy;
-	}
-
-	RHICmdList.EnqueueLambda([this, TextureRHI, MipIndex, UpdateRegion, SourcePitch, SourceData, bFreeMemory] (FRHICommandListBase&)
+	RHICmdList.EnqueueLambda([this, TextureRHI, MipIndex, UpdateRegion, SourcePitch, SourceData] (FRHICommandListBase&)
 	{
 		FD3D11Texture* Texture = ResourceCast(TextureRHI);
 
@@ -1547,22 +1542,18 @@ void FD3D11DynamicRHI::RHIUpdateTexture2D(FRHICommandListBase& RHICmdList, FRHIT
 
 		Direct3DDeviceIMContext->UpdateSubresource(Texture->GetResource(), MipIndex, &DestBox, SourceData, SourcePitch, 0);
 
-		if (bFreeMemory) { FMemory::Free((void*)SourceData); }
+		FMemory::Free((void*)SourceData);
 	});
 }
 
 void FD3D11DynamicRHI::RHIUpdateTexture3D(FRHICommandListBase& RHICmdList, FRHITexture3D* TextureRHI,uint32 MipIndex,const FUpdateTextureRegion3D& UpdateRegion,uint32 SourceRowPitch,uint32 SourceDepthPitch,const uint8* SourceData)
 {
-	bool bFreeMemory = !(RHICmdList.Bypass() || !IsRunningRHIInSeparateThread());
-	if (bFreeMemory)
-	{
-		const SIZE_T SourceDataSize = static_cast<SIZE_T>(SourceDepthPitch) * UpdateRegion.Depth;
-		uint8* SourceDataCopy = (uint8*)FMemory::Malloc(SourceDataSize);
-		FMemory::Memcpy(SourceDataCopy, SourceData, SourceDataSize);
-		SourceData = SourceDataCopy;
-	}
+	const SIZE_T SourceDataSize = static_cast<SIZE_T>(SourceDepthPitch) * UpdateRegion.Depth;
+	uint8* SourceDataCopy = (uint8*)FMemory::Malloc(SourceDataSize);
+	FMemory::Memcpy(SourceDataCopy, SourceData, SourceDataSize);
+	SourceData = SourceDataCopy;
 
-	RHICmdList.EnqueueLambda([this, TextureRHI, MipIndex, UpdateRegion, SourceRowPitch, SourceDepthPitch, SourceData, bFreeMemory] (FRHICommandListBase&)
+	RHICmdList.EnqueueLambda([this, TextureRHI, MipIndex, UpdateRegion, SourceRowPitch, SourceDepthPitch, SourceData] (FRHICommandListBase&)
 	{
 		FD3D11Texture* Texture = ResourceCast(TextureRHI);
 
@@ -1580,7 +1571,7 @@ void FD3D11DynamicRHI::RHIUpdateTexture3D(FRHICommandListBase& RHICmdList, FRHIT
 
 		Direct3DDeviceIMContext->UpdateSubresource(Texture->GetResource(), MipIndex, &DestBox, SourceData, SourceRowPitch, SourceDepthPitch);
 
-		if (bFreeMemory) { FMemory::Free((void*)SourceData); }
+		FMemory::Free((void*)SourceData);
 	});
 }
 
