[WindowsPlatformFeatures] removed `WINDOWSPLATFORMFEATURES_DEBUG`,  `WINDOWSPLATFORMFEATURES_START`, and `WINDOWSPLATFORMFEATURES_END`
no need for more macro to enable debugging use regular macros to disable optimizations

diff --git a/Engine/Source/Runtime/Windows/WindowsPlatformFeatures/Private/HighlightRecorder.cpp b/Engine/Source/Runtime/Windows/WindowsPlatformFeatures/Private/HighlightRecorder.cpp
--- a/Engine/Source/Runtime/Windows/WindowsPlatformFeatures/Private/HighlightRecorder.cpp
+++ b/Engine/Source/Runtime/Windows/WindowsPlatformFeatures/Private/HighlightRecorder.cpp
@@ -20,8 +20,6 @@
 DEFINE_LOG_CATEGORY(WMF);
 DEFINE_LOG_CATEGORY(HighlightRecorder);
 
-WINDOWSPLATFORMFEATURES_START
-
 //////////////////////////////////////////////////////////////////////////
 // console commands for testing
 
@@ -413,6 +411,3 @@ bool FHighlightRecorder::GetSavingStart(const TArray<AVEncoder::FMediaPacket>& S
 
 	return true;
 }
-
-WINDOWSPLATFORMFEATURES_END
-
diff --git a/Engine/Source/Runtime/Windows/WindowsPlatformFeatures/Private/WindowsPlatformFeatures.cpp b/Engine/Source/Runtime/Windows/WindowsPlatformFeatures/Private/WindowsPlatformFeatures.cpp
--- a/Engine/Source/Runtime/Windows/WindowsPlatformFeatures/Private/WindowsPlatformFeatures.cpp
+++ b/Engine/Source/Runtime/Windows/WindowsPlatformFeatures/Private/WindowsPlatformFeatures.cpp
@@ -7,8 +7,6 @@
 
 IMPLEMENT_MODULE(FWindowsPlatformFeaturesModule, WindowsPlatformFeatures);
 
-WINDOWSPLATFORMFEATURES_START
-
 FWindowsPlatformFeaturesModule::FWindowsPlatformFeaturesModule()
 {
 }
@@ -23,5 +21,3 @@ void FWindowsPlatformFeaturesModule::StartupModule()
 {
 	FModuleManager::Get().LoadModule(TEXT("GameplayMediaEncoder"));
 }
-
-WINDOWSPLATFORMFEATURES_END
diff --git a/Engine/Source/Runtime/Windows/WindowsPlatformFeatures/Private/WindowsPlatformFeatures.h b/Engine/Source/Runtime/Windows/WindowsPlatformFeatures/Private/WindowsPlatformFeatures.h
--- a/Engine/Source/Runtime/Windows/WindowsPlatformFeatures/Private/WindowsPlatformFeatures.h
+++ b/Engine/Source/Runtime/Windows/WindowsPlatformFeatures/Private/WindowsPlatformFeatures.h
@@ -4,7 +4,6 @@
 
 #include "CoreMinimal.h"
 #include "PlatformFeatures.h"
-#include "WindowsPlatformFeaturesCommon.h"
 
 class FWindowsPlatformFeaturesModule : public IPlatformFeaturesModule
 {
diff --git a/Engine/Source/Runtime/Windows/WindowsPlatformFeatures/Private/WindowsPlatformFeaturesCommon.h b/Engine/Source/Runtime/Windows/WindowsPlatformFeatures/Private/WindowsPlatformFeaturesCommon.h
deleted file mode 100644
diff --git a/Engine/Source/Runtime/Windows/WindowsPlatformFeatures/Private/WindowsVideoRecordingSystem.cpp b/Engine/Source/Runtime/Windows/WindowsPlatformFeatures/Private/WindowsVideoRecordingSystem.cpp
--- a/Engine/Source/Runtime/Windows/WindowsPlatformFeatures/Private/WindowsVideoRecordingSystem.cpp
+++ b/Engine/Source/Runtime/Windows/WindowsPlatformFeatures/Private/WindowsVideoRecordingSystem.cpp
@@ -19,8 +19,6 @@ DEFINE_VIDEOSYSTEMRECORDING_STATS
 DEFINE_LOG_CATEGORY(WindowsVideoRecordingSystem);
 CSV_DEFINE_CATEGORY(WindowsVideoRecordingSystem, true);
 
-WINDOWSPLATFORMFEATURES_START
-
 /**
  * This internal helper class handles disabling of the Windows built-in screenshot (PrtScr key) and clip
  * recorder (via Win-G key menu on Windows 10)
@@ -414,6 +412,3 @@ EVideoRecordingState FWindowsVideoRecordingSystem::GetRecordingState() const
 {
 	return RecordState;
 }
-
-WINDOWSPLATFORMFEATURES_END
-
diff --git a/Engine/Source/Runtime/Windows/WindowsPlatformFeatures/Private/WindowsVideoRecordingSystem.h b/Engine/Source/Runtime/Windows/WindowsPlatformFeatures/Private/WindowsVideoRecordingSystem.h
--- a/Engine/Source/Runtime/Windows/WindowsPlatformFeatures/Private/WindowsVideoRecordingSystem.h
+++ b/Engine/Source/Runtime/Windows/WindowsPlatformFeatures/Private/WindowsVideoRecordingSystem.h
@@ -13,7 +13,6 @@
 #include "RHI.h"
 #include "RHIResources.h"
 #include "PipelineStateCache.h"
-#include "WindowsPlatformFeaturesCommon.h"
 
 DECLARE_LOG_CATEGORY_EXTERN(WindowsVideoRecordingSystem, VeryVerbose, VeryVerbose);
 
diff --git a/Engine/Source/Runtime/Windows/WindowsPlatformFeatures/Private/WmfMp4Writer.cpp b/Engine/Source/Runtime/Windows/WindowsPlatformFeatures/Private/WmfMp4Writer.cpp
--- a/Engine/Source/Runtime/Windows/WindowsPlatformFeatures/Private/WmfMp4Writer.cpp
+++ b/Engine/Source/Runtime/Windows/WindowsPlatformFeatures/Private/WmfMp4Writer.cpp
@@ -14,8 +14,6 @@ DECLARE_LOG_CATEGORY_EXTERN(MP4, Log, VeryVerbose);
 
 DEFINE_LOG_CATEGORY(MP4);
 
-WINDOWSPLATFORMFEATURES_START
-
 bool FWmfMp4Writer::Initialize(const TCHAR* Filename)
 {
 	CHECK_HR(MFCreateSinkWriterFromURL(Filename, nullptr, nullptr, Writer.GetInitReference()));
@@ -138,7 +136,3 @@ bool FWmfMp4Writer::Finalize()
 
 	return true;
 }
-
-WINDOWSPLATFORMFEATURES_END
-
-
diff --git a/Engine/Source/Runtime/Windows/WindowsPlatformFeatures/Private/WmfPrivate.h b/Engine/Source/Runtime/Windows/WindowsPlatformFeatures/Private/WmfPrivate.h
--- a/Engine/Source/Runtime/Windows/WindowsPlatformFeatures/Private/WmfPrivate.h
+++ b/Engine/Source/Runtime/Windows/WindowsPlatformFeatures/Private/WmfPrivate.h
@@ -21,8 +21,6 @@ THIRD_PARTY_INCLUDES_START
 	#include <shlwapi.h>
 THIRD_PARTY_INCLUDES_END
 
-#include "WindowsPlatformFeaturesCommon.h"
-
 inline const FString GetComErrorDescription(HRESULT Res)
 {
 	const uint32 BufSize = 4096;
diff --git a/Engine/Source/Runtime/Windows/WindowsPlatformFeatures/Private/WmfRingBuffer.cpp b/Engine/Source/Runtime/Windows/WindowsPlatformFeatures/Private/WmfRingBuffer.cpp
--- a/Engine/Source/Runtime/Windows/WindowsPlatformFeatures/Private/WmfRingBuffer.cpp
+++ b/Engine/Source/Runtime/Windows/WindowsPlatformFeatures/Private/WmfRingBuffer.cpp
@@ -6,8 +6,6 @@
 DECLARE_LOG_CATEGORY_EXTERN(WmfRingBuffer, Log, VeryVerbose);
 DEFINE_LOG_CATEGORY(WmfRingBuffer);
 
-WINDOWSPLATFORMFEATURES_START
-
 void FWmfRingBuffer::Push(AVEncoder::FMediaPacket&& Sample)
 {
 	check(MaxDuration != 0);
@@ -136,6 +134,3 @@ void FWmfRingBuffer::Reset()
 	FScopeLock Lock(&Mutex);
 	Samples.Empty();
 }
-
-WINDOWSPLATFORMFEATURES_END
-
