Fix texture array resolve in D3D11.

diff --git a/Engine/Source/Runtime/Windows/D3D11RHI/Private/D3D11RenderTarget.cpp b/Engine/Source/Runtime/Windows/D3D11RHI/Private/D3D11RenderTarget.cpp
--- a/Engine/Source/Runtime/Windows/D3D11RHI/Private/D3D11RenderTarget.cpp
+++ b/Engine/Source/Runtime/Windows/D3D11RHI/Private/D3D11RenderTarget.cpp
@@ -1194,8 +1194,8 @@ void FD3D11DynamicRHI::ResolveTexture(UE::RHICore::FResolveTextureInfo Info)
 
 		for (int32 ArraySlice = ArraySliceBegin; ArraySlice < ArraySliceEnd; ArraySlice++)
 		{
-			int32 DestSubresource   = D3D11CalcSubresource(Info.MipLevel, ArraySlice, DestDesc.ArraySize);
-			int32 SourceSubresource = D3D11CalcSubresource(Info.MipLevel, ArraySlice, SourceDesc.ArraySize);
+			int32 DestSubresource   = D3D11CalcSubresource(Info.MipLevel, ArraySlice, DestDesc.NumMips);
+			int32 SourceSubresource = D3D11CalcSubresource(Info.MipLevel, ArraySlice, SourceDesc.NumMips);
 
 			Direct3DDeviceIMContext->ResolveSubresource(DestTexture->GetResource(), DestSubresource, SourceTexture->GetResource(), SourceSubresource, DestFormatTypeless);
 		}
