PR #9851: Add missing parameter to FWindowsMMDeviceCache::OnDeviceRemoved log which causes a crash (Contributed by [REDACTED_PERSON])

diff --git a/Engine/Source/Runtime/Windows/AudioMixerXAudio2/Private/WindowsMMDeviceInfoCache.cpp b/Engine/Source/Runtime/Windows/AudioMixerXAudio2/Private/WindowsMMDeviceInfoCache.cpp
--- a/Engine/Source/Runtime/Windows/AudioMixerXAudio2/Private/WindowsMMDeviceInfoCache.cpp
+++ b/Engine/Source/Runtime/Windows/AudioMixerXAudio2/Private/WindowsMMDeviceInfoCache.cpp
@@ -403,7 +403,7 @@ void FWindowsMMDeviceCache::OnDeviceRemoved(const FString& DeviceId, bool)
 {
 	FWriteScopeLock WriteLock(CacheMutationLock);
 	FName DeviceIdName = *DeviceId;
-	UE_CLOG(!Cache.Contains(DeviceIdName), LogAudioMixer, Warning, TEXT("FWindowsMMDeviceCache::OnDeviceRemoved: DeviceId='%s' was not in the cache. "));
+	UE_CLOG(!Cache.Contains(DeviceIdName), LogAudioMixer, Warning, TEXT("FWindowsMMDeviceCache::OnDeviceRemoved: DeviceId='%s' was not in the cache. "), *DeviceId);
 	Cache.Remove(DeviceIdName);
 }
 
