Fix to add DX11 support for AMF encoding. NOTE: This change was originally approved in stream Release-5.1 (https://p4-swarm.epicgames.net/reviews/21212427) but was not approved in time for 5.1, hence it's re-submission into Main for 5.2.

We create a temporary d3d device before calling AMD AGS as it will not check the DriverStore without an initialized device when agsDriverExtensionsDX11_CreateDevice is called. We then destroy the temporary device. According to AMD this requirement is currently by design.

diff --git a/Engine/Plugins/Media/HardwareEncoders/Source/EncoderAMF/Private/Amf_EncoderH264.cpp b/Engine/Plugins/Media/HardwareEncoders/Source/EncoderAMF/Private/Amf_EncoderH264.cpp
--- a/Engine/Plugins/Media/HardwareEncoders/Source/EncoderAMF/Private/Amf_EncoderH264.cpp
+++ b/Engine/Plugins/Media/HardwareEncoders/Source/EncoderAMF/Private/Amf_EncoderH264.cpp
@@ -168,13 +168,6 @@ namespace AVEncoder
 			return false;
 		}
 
-		//TODO(sandor.hadas) see if the current issues with AMF can be resolved then remove this error message
-		if (RHIType == ERHIInterfaceType::D3D11)
-		{
-			UE_LOG(LogEncoderAMF, Error, TEXT("AMF with DX11 is not currently supported try DX12 or Vulkan."));
-			return false;
-		}
-
 		if (!EncoderDevice)
 		{
 			UE_LOG(LogEncoderAMF, Error, TEXT("Amf needs an encoder device."));
diff --git a/Engine/Source/Runtime/Windows/D3D11RHI/Private/Windows/WindowsD3D11Device.cpp b/Engine/Source/Runtime/Windows/D3D11RHI/Private/Windows/WindowsD3D11Device.cpp
--- a/Engine/Source/Runtime/Windows/D3D11RHI/Private/Windows/WindowsD3D11Device.cpp
+++ b/Engine/Source/Runtime/Windows/D3D11RHI/Private/Windows/WindowsD3D11Device.cpp
@@ -1909,6 +1909,30 @@ void FD3D11DynamicRHI::InitD3DDevice()
 			AmdExtensionParams.pAppName = bDisableAppRegistration ? TEXT("") : FApp::GetProjectName();
 			AmdExtensionParams.appVersion = AGS_UNSPECIFIED_VERSION;
 
+			// agsDriverExtensionsDX11_CreateDevice will not check for the DriverStore if there is not already a D3D11 device
+			// initialized. In order to use AMF we require libraries to also be loaded from the driver store so we temporarily
+			// initialize a device here and destroy it once the agsDriverExtensionsDX11_CreateDevice is run.
+			ID3D11Device* D3DDevicePreload = nullptr;
+			ID3D11DeviceContext* D3DDeviceContextPreload = nullptr;
+
+			int32 NumAllowedFeatureLevels = 1;
+			D3D_FEATURE_LEVEL OutFeatureLevel = FeatureLevel;
+			HRESULT Result = D3D11CreateDevice(
+				Adapter.DXGIAdapter,
+				D3D_DRIVER_TYPE_UNKNOWN,
+				nullptr,
+				DeviceFlags,
+				&FeatureLevel,
+				NumAllowedFeatureLevels,
+				D3D11_SDK_VERSION,
+				&D3DDevicePreload,
+				&ActualFeatureLevel,
+				&D3DDeviceContextPreload);
+			if (FAILED(Result))
+			{
+				UE_LOG(LogD3D11RHI, Error, TEXT("Failed to load the AMD DriverStore library"));
+			}
+
 			AGSDX11ReturnedParams DeviceCreationReturnedParams;
 			FMemory::Memzero(&DeviceCreationReturnedParams, sizeof(DeviceCreationReturnedParams));
 			AGSReturnCode DeviceCreation =
@@ -1918,6 +1942,13 @@ void FD3D11DynamicRHI::InitD3DDevice()
 					&AmdExtensionParams,
 					&DeviceCreationReturnedParams);
 
+			// Destroy temporary device and context			
+			D3DDevicePreload->Release();
+			D3DDevicePreload = nullptr;
+
+			D3DDeviceContextPreload->Release();
+			D3DDeviceContextPreload = nullptr;
+
 			if (DeviceCreation == AGS_SUCCESS)
 			{
 				Direct3DDevice = DeviceCreationReturnedParams.pDevice;
