Adding header to fix nonunity

diff --git a/Engine/Source/Runtime/Windows/D3D11RHI/Public/D3D11ConstantBuffer.h b/Engine/Source/Runtime/Windows/D3D11RHI/Public/D3D11ConstantBuffer.h
--- a/Engine/Source/Runtime/Windows/D3D11RHI/Public/D3D11ConstantBuffer.h
+++ b/Engine/Source/Runtime/Windows/D3D11RHI/Public/D3D11ConstantBuffer.h
@@ -7,6 +7,7 @@
 #pragma once
 
 #include "RenderResource.h"
+#include "Stats/Stats2.h"
 
 class FD3D11DynamicRHI;
 
