Add the ability for plug-ins to hook DXGI swapchain creation.

This is the simplest way to integrate some third party upscaling/frame interpolation libraries.

#rnx

diff --git a/Engine/Source/Runtime/D3D12RHI/Private/Windows/WindowsD3D12Viewport.cpp b/Engine/Source/Runtime/D3D12RHI/Private/Windows/WindowsD3D12Viewport.cpp
--- a/Engine/Source/Runtime/D3D12RHI/Private/Windows/WindowsD3D12Viewport.cpp
+++ b/Engine/Source/Runtime/D3D12RHI/Private/Windows/WindowsD3D12Viewport.cpp
@@ -13,6 +13,9 @@
 
 #include "HAL/ThreadHeartBeat.h"
 
+#include "Windows/IDXGISwapchainProvider.h"
+#include "Features/IModularFeatures.h"
+
 static const uint32 WindowsDefaultNumBackBuffers = 3;
 
 extern FD3D12Texture* GetSwapChainSurface(FD3D12Device* Parent, EPixelFormat PixelFormat, uint32 SizeX, uint32 SizeY, IDXGISwapChain* SwapChain, uint32 BackBufferIndex, TRefCountPtr<ID3D12Resource> BackBufferResourceOverride);
@@ -64,6 +67,27 @@ void FD3D12Viewport::Init()
 	FD3D12Adapter* Adapter = GetParentAdapter();
 	IDXGIFactory2* Factory2 = Adapter->GetDXGIFactory2();
 
+	TArray<IDXGISwapchainProvider*> DXGISwapchainProviderModules = IModularFeatures::Get().GetModularFeatureImplementations<IDXGISwapchainProvider>(IDXGISwapchainProvider::GetModularFeatureName());
+	IDXGISwapchainProvider* DXGISwapchainProvider = nullptr;
+	for (IDXGISwapchainProvider* ProviderModule : DXGISwapchainProviderModules)
+	{
+		if (ProviderModule->SupportsRHI(ERHIInterfaceType::D3D12))
+		{
+			DXGISwapchainProvider = ProviderModule;
+			break;
+		}
+	}
+
+	if (DXGISwapchainProvider)
+	{
+		static bool bCustomSwapchainLogged = false;
+		if (!bCustomSwapchainLogged)
+		{
+			UE_LOG(LogD3D12RHI, Log, TEXT("Found a custom swapchain provider: '%s'."), DXGISwapchainProvider->GetName());
+			bCustomSwapchainLogged = true;
+		}
+	}
+
 	bAllowTearing = false;
 	if (GD3D12UseAllowTearing)
 	{
@@ -148,7 +172,11 @@ void FD3D12Viewport::Init()
 				SwapChainDesc1.Width = SizeX;
 				SwapChainDesc1.Height = SizeY;
 
-				VERIFYD3D12RESULT(Factory2->CreateSwapChainForHwnd(CommandQueue, WindowHandle, &SwapChainDesc1, nullptr, nullptr, SwapChain1.GetInitReference()));
+				HRESULT hr = DXGISwapchainProvider ?
+					DXGISwapchainProvider->CreateSwapChainForHwnd(Factory2, CommandQueue, WindowHandle, &SwapChainDesc1, nullptr, nullptr, SwapChain1.GetInitReference()) :
+					Factory2->CreateSwapChainForHwnd(CommandQueue, WindowHandle, &SwapChainDesc1, nullptr, nullptr, SwapChain1.GetInitReference());
+
+				VERIFYD3D12RESULT(hr);
 			}
 			else
 			{
@@ -176,7 +204,11 @@ void FD3D12Viewport::Init()
 			SwapChainDesc.Flags = SwapChainFlags;
 
 			TRefCountPtr<IDXGISwapChain> SwapChain;
-			HRESULT hr = Factory2->CreateSwapChain(CommandQueue, &SwapChainDesc, SwapChain.GetInitReference());
+			
+			HRESULT hr = DXGISwapchainProvider ?
+				DXGISwapchainProvider->CreateSwapChain(Factory2, CommandQueue, &SwapChainDesc, SwapChain.GetInitReference()) :
+				Factory2->CreateSwapChain(CommandQueue, &SwapChainDesc, SwapChain.GetInitReference());
+			
 			if (FAILED(hr))
 			{
 				UE_LOG(LogD3D12RHI, Warning, TEXT("Failed to create swapchain with the following parameters:"));
@@ -185,6 +217,7 @@ void FD3D12Viewport::Init()
 				UE_LOG(LogD3D12RHI, Warning, TEXT("\tWindows handle: 0x%x (IsWindow: %s)"), WindowHandle, IsWindow(WindowHandle) ? TEXT("true") : TEXT("false"));
 				UE_LOG(LogD3D12RHI, Warning, TEXT("\tFullscreen: %s"), bIsFullscreen ? TEXT("true") : TEXT("false"));
 				UE_LOG(LogD3D12RHI, Warning, TEXT("\tSwapchain flags: %d"), SwapChainFlags);
+				UE_LOG(LogD3D12RHI, Warning, TEXT("\tCustom swapchain provider: %s"), DXGISwapchainProvider ? DXGISwapchainProvider->GetName() : TEXT("none"));
 
 				VERIFYD3D12RESULT(hr);
 			}
diff --git a/Engine/Source/Runtime/RHI/Public/Windows/IDXGISwapchainProvider.h b/Engine/Source/Runtime/RHI/Public/Windows/IDXGISwapchainProvider.h
new file mode 100644
--- /dev/null
+++ b/Engine/Source/Runtime/RHI/Public/Windows/IDXGISwapchainProvider.h
@@ -0,0 +1,22 @@
+// Copyright Epic Games, Inc. All Rights Reserved.
+
+#pragma once
+
+#include "Features/IModularFeature.h"
+#include "RHIDefinitions.h"
+
+class IDXGISwapchainProvider : public IModularFeature
+{
+public:
+	static FName GetModularFeatureName()
+	{
+		static FName FeatureName(TEXT("DXGISwapchainProvider"));
+		return FeatureName;
+	}
+
+	virtual bool SupportsRHI(ERHIInterfaceType RHIType) const = 0;
+	virtual TCHAR* GetName() const = 0;
+
+	virtual HRESULT CreateSwapChainForHwnd(IDXGIFactory2* pFactory, IUnknown* pDevice, HWND hWnd, const DXGI_SWAP_CHAIN_DESC1* pDesc, const DXGI_SWAP_CHAIN_FULLSCREEN_DESC* pFulScreenDesc, IDXGIOutput* pRestrictToOutput, IDXGISwapChain1** ppSwapChain) = 0;
+	virtual HRESULT CreateSwapChain(IDXGIFactory* pFactory, IUnknown* pDevice, DXGI_SWAP_CHAIN_DESC* pDesc, IDXGISwapChain** ppSwapChain) = 0;
+};
diff --git a/Engine/Source/Runtime/Windows/D3D11RHI/Private/Windows/WindowsD3D11Viewport.cpp b/Engine/Source/Runtime/Windows/D3D11RHI/Private/Windows/WindowsD3D11Viewport.cpp
--- a/Engine/Source/Runtime/Windows/D3D11RHI/Private/Windows/WindowsD3D11Viewport.cpp
+++ b/Engine/Source/Runtime/Windows/D3D11RHI/Private/Windows/WindowsD3D11Viewport.cpp
@@ -11,6 +11,9 @@
 #include "Windows/AllowWindowsPlatformTypes.h"
 #include <dwmapi.h>
 
+#include "Windows/IDXGISwapchainProvider.h"
+#include "Features/IModularFeatures.h"
+
 #ifndef DXGI_SWAP_CHAIN_FLAG_ALLOW_TEARING
 #define DXGI_SWAP_CHAIN_FLAG_ALLOW_TEARING  2048
 #endif
@@ -58,6 +61,27 @@ FD3D11Viewport::FD3D11Viewport(FD3D11DynamicRHI* InD3DRHI,HWND InWindowHandle,ui
 
 	PixelFormat = InD3DRHI->GetDisplayFormat(InPreferredPixelFormat);
 
+	TArray<IDXGISwapchainProvider*> DXGISwapchainProviderModules = IModularFeatures::Get().GetModularFeatureImplementations<IDXGISwapchainProvider>(IDXGISwapchainProvider::GetModularFeatureName());
+	IDXGISwapchainProvider* DXGISwapchainProvider = nullptr;
+	for (IDXGISwapchainProvider* ProviderModule : DXGISwapchainProviderModules)
+	{
+		if (ProviderModule->SupportsRHI(ERHIInterfaceType::D3D11))
+		{
+			DXGISwapchainProvider = ProviderModule;
+			break;
+		}
+	}
+
+	if (DXGISwapchainProvider)
+	{
+		static bool bCustomSwapchainLogged = false;
+		if (!bCustomSwapchainLogged)
+		{
+			UE_LOG(LogD3D11RHI, Log, TEXT("Found a custom swapchain provider: '%s'."), DXGISwapchainProvider->GetName());
+			bCustomSwapchainLogged = true;
+		}
+	}
+
 	// Create a backbuffer/swapchain for each viewport
 	TRefCountPtr<IDXGIDevice> DXGIDevice;
 	VERIFYD3D11RESULT_EX( D3DRHI->GetDevice()->QueryInterface(IID_PPV_ARGS(DXGIDevice.GetInitReference())), D3DRHI->GetDevice() );
@@ -146,7 +170,11 @@ FD3D11Viewport::FD3D11Viewport(FD3D11DynamicRHI* InD3DRHI,HWND InWindowHandle,ui
 				SwapChainDesc1.Flags = GSwapChainFlags;
 
 				IDXGISwapChain1* SwapChain1 = nullptr;
-				VERIFYD3D11RESULT_EX((Factory2->CreateSwapChainForHwnd(D3DRHI->GetDevice(), WindowHandle, &SwapChainDesc1, nullptr, nullptr, &SwapChain1)), D3DRHI->GetDevice());
+				HRESULT CreateSwapChainForHwndResult = DXGISwapchainProvider ?
+					DXGISwapchainProvider->CreateSwapChainForHwnd(Factory2, D3DRHI->GetDevice(), WindowHandle, &SwapChainDesc1, nullptr, nullptr, &SwapChain1) :
+					Factory2->CreateSwapChainForHwnd(D3DRHI->GetDevice(), WindowHandle, &SwapChainDesc1, nullptr, nullptr, &SwapChain1);
+
+				VERIFYD3D11RESULT_EX(CreateSwapChainForHwndResult, D3DRHI->GetDevice());
 				SwapChain = SwapChain1;
 			}
 			else
@@ -180,7 +208,9 @@ FD3D11Viewport::FD3D11Viewport(FD3D11DynamicRHI* InD3DRHI,HWND InWindowHandle,ui
 			SwapChainDesc.Scaling = GSwapScaling;
 
 			IDXGISwapChain1* SwapChain1 = nullptr;
-			HRESULT CreateSwapChainForHwndResult = Factory2->CreateSwapChainForHwnd(D3DRHI->GetDevice(), WindowHandle, &SwapChainDesc, &FSSwapChainDesc, nullptr, &SwapChain1);
+			HRESULT CreateSwapChainForHwndResult = DXGISwapchainProvider ?
+				DXGISwapchainProvider->CreateSwapChainForHwnd(Factory2, D3DRHI->GetDevice(), WindowHandle, &SwapChainDesc, &FSSwapChainDesc, nullptr, &SwapChain1) :
+				Factory2->CreateSwapChainForHwnd(D3DRHI->GetDevice(), WindowHandle, &SwapChainDesc, &FSSwapChainDesc, nullptr, &SwapChain1);
 			if(SUCCEEDED(CreateSwapChainForHwndResult))
 			{
 				SwapChain1->QueryInterface(IID_PPV_ARGS(SwapChain.GetInitReference()));
@@ -234,7 +264,9 @@ FD3D11Viewport::FD3D11Viewport(FD3D11DynamicRHI* InD3DRHI,HWND InWindowHandle,ui
 			SwapChainDesc.SwapEffect = GSwapEffect;
 			SwapChainDesc.Flags = GSwapChainFlags;
 
-			HRESULT CreateSwapChainResult = D3DRHI->GetFactory()->CreateSwapChain(D3DRHI->GetDevice(), &SwapChainDesc, SwapChain.GetInitReference());
+			HRESULT CreateSwapChainResult = DXGISwapchainProvider ?
+				DXGISwapchainProvider->CreateSwapChain(D3DRHI->GetFactory(), D3DRHI->GetDevice(), &SwapChainDesc, SwapChain.GetInitReference()) :
+				D3DRHI->GetFactory()->CreateSwapChain(D3DRHI->GetDevice(), &SwapChainDesc, SwapChain.GetInitReference());
 			if (CreateSwapChainResult == E_INVALIDARG)
 			{
 				const TCHAR* D3DFormatString = GetD3D11TextureFormatString(SwapChainDesc.BufferDesc.Format);
