Fix for crash in D3D12 due to RHIPerFrameRHIFlushComplete being called mid-pass, invaliding bound static uniform buffers. Also added clear of bound uniform buffers in D3D11.

diff --git a/Engine/Source/Runtime/D3D12RHI/Private/D3D12RHI.cpp b/Engine/Source/Runtime/D3D12RHI/Private/D3D12RHI.cpp
--- a/Engine/Source/Runtime/D3D12RHI/Private/D3D12RHI.cpp
+++ b/Engine/Source/Runtime/D3D12RHI/Private/D3D12RHI.cpp
@@ -514,7 +514,7 @@ void FD3D12DynamicRHI::RHIPerFrameRHIFlushComplete()
 
 		for (FD3D12Device* Device : Adapter.GetDevices())
 		{
-			Device->GetDefaultCommandContext().ClearState(FD3D12ContextCommon::EClearStateMode::All);
+			Device->GetDefaultCommandContext().ClearState(FD3D12ContextCommon::EClearStateMode::TransientOnly);
 		}
 	}
 }
diff --git a/Engine/Source/Runtime/Windows/D3D11RHI/Private/Windows/WindowsD3D11Device.cpp b/Engine/Source/Runtime/Windows/D3D11RHI/Private/Windows/WindowsD3D11Device.cpp
--- a/Engine/Source/Runtime/Windows/D3D11RHI/Private/Windows/WindowsD3D11Device.cpp
+++ b/Engine/Source/Runtime/Windows/D3D11RHI/Private/Windows/WindowsD3D11Device.cpp
@@ -2243,6 +2243,16 @@ void FD3D11DynamicRHI::RHIPerFrameRHIFlushComplete()
 		}
 	}
 #endif
+
+	for (int32 Frequency = 0; Frequency < SF_NumStandardFrequencies; ++Frequency)
+	{
+		DirtyUniformBuffers[Frequency] = 0;
+
+		for (int32 BindIndex = 0; BindIndex < MAX_UNIFORM_BUFFERS_PER_SHADER_STAGE; ++BindIndex)
+		{
+			BoundUniformBuffers[Frequency][BindIndex] = nullptr;
+		}
+	}
 }
 
 /**
