Don't assert when we can't find any display modes.

diff --git a/Engine/Source/Runtime/D3D12RHI/Private/Windows/WindowsD3D12Device.cpp b/Engine/Source/Runtime/D3D12RHI/Private/Windows/WindowsD3D12Device.cpp
--- a/Engine/Source/Runtime/D3D12RHI/Private/Windows/WindowsD3D12Device.cpp
+++ b/Engine/Source/Runtime/D3D12RHI/Private/Windows/WindowsD3D12Device.cpp
@@ -1867,67 +1867,72 @@ bool FD3D12DynamicRHI::RHIGetAvailableResolutions(FScreenResolutionArray& Resolu
 			return false;
 		}
 
-		checkf(NumModes > 0, TEXT("No display modes found for the standard format DXGI_FORMAT_R8G8B8A8_UNORM!"));
-
-		DXGI_MODE_DESC* ModeList = new DXGI_MODE_DESC[NumModes];
-		VERIFYD3D12RESULT(Output->GetDisplayModeList(Format, 0, &NumModes, ModeList));
-
-		for (uint32 m = 0; m < NumModes; m++)
+		if (NumModes > 0)
 		{
-			CA_SUPPRESS(6385);
-			if (((int32)ModeList[m].Width >= MinAllowableResolutionX) &&
-				((int32)ModeList[m].Width <= MaxAllowableResolutionX) &&
-				((int32)ModeList[m].Height >= MinAllowableResolutionY) &&
-				((int32)ModeList[m].Height <= MaxAllowableResolutionY)
-				)
+			DXGI_MODE_DESC* ModeList = new DXGI_MODE_DESC[NumModes];
+			VERIFYD3D12RESULT(Output->GetDisplayModeList(Format, 0, &NumModes, ModeList));
+
+			for (uint32 m = 0; m < NumModes; m++)
 			{
-				bool bAddIt = true;
-				if (bIgnoreRefreshRate == false)
+				CA_SUPPRESS(6385);
+				if (((int32)ModeList[m].Width >= MinAllowableResolutionX) &&
+					((int32)ModeList[m].Width <= MaxAllowableResolutionX) &&
+					((int32)ModeList[m].Height >= MinAllowableResolutionY) &&
+					((int32)ModeList[m].Height <= MaxAllowableResolutionY)
+					)
 				{
-					if (((int32)ModeList[m].RefreshRate.Numerator < MinAllowableRefreshRate * ModeList[m].RefreshRate.Denominator) ||
-						((int32)ModeList[m].RefreshRate.Numerator > MaxAllowableRefreshRate * ModeList[m].RefreshRate.Denominator)
-						)
+					bool bAddIt = true;
+					if (bIgnoreRefreshRate == false)
 					{
-						continue;
+						if (((int32)ModeList[m].RefreshRate.Numerator < MinAllowableRefreshRate * ModeList[m].RefreshRate.Denominator) ||
+							((int32)ModeList[m].RefreshRate.Numerator > MaxAllowableRefreshRate * ModeList[m].RefreshRate.Denominator)
+							)
+						{
+							continue;
+						}
 					}
-				}
-				else
-				{
-					// See if it is in the list already
-					for (int32 CheckIndex = 0; CheckIndex < Resolutions.Num(); CheckIndex++)
+					else
 					{
-						FScreenResolutionRHI& CheckResolution = Resolutions[CheckIndex];
-						if ((CheckResolution.Width == ModeList[m].Width) &&
-							(CheckResolution.Height == ModeList[m].Height))
+						// See if it is in the list already
+						for (int32 CheckIndex = 0; CheckIndex < Resolutions.Num(); CheckIndex++)
 						{
-							// Already in the list...
-							bAddIt = false;
-							break;
+							FScreenResolutionRHI& CheckResolution = Resolutions[CheckIndex];
+							if ((CheckResolution.Width == ModeList[m].Width) &&
+								(CheckResolution.Height == ModeList[m].Height))
+							{
+								// Already in the list...
+								bAddIt = false;
+								break;
+							}
 						}
 					}
-				}
 
-				if (bAddIt)
-				{
-					// Add the mode to the list
-					int32 Temp2Index = Resolutions.AddZeroed();
-					FScreenResolutionRHI& ScreenResolution = Resolutions[Temp2Index];
+					if (bAddIt)
+					{
+						// Add the mode to the list
+						int32 Temp2Index = Resolutions.AddZeroed();
+						FScreenResolutionRHI& ScreenResolution = Resolutions[Temp2Index];
 
-					ScreenResolution.Width = ModeList[m].Width;
-					ScreenResolution.Height = ModeList[m].Height;
-					ScreenResolution.RefreshRate = ModeList[m].RefreshRate.Numerator / ModeList[m].RefreshRate.Denominator;
+						ScreenResolution.Width = ModeList[m].Width;
+						ScreenResolution.Height = ModeList[m].Height;
+						ScreenResolution.RefreshRate = ModeList[m].RefreshRate.Numerator / ModeList[m].RefreshRate.Denominator;
+					}
 				}
 			}
-		}
 
-		delete[] ModeList;
+			delete[] ModeList;
+		}
+		else
+		{
+			UE_LOG(LogD3D12RHI, Warning, TEXT("No display modes found for the standard format DXGI_FORMAT_R8G8B8A8_UNORM!"));
+		}
 
 		++CurrentOutput;
 
 	// TODO: Cap at 1 for default output
 	} while (CurrentOutput < 1);
 
-	return true;
+	return Resolutions.Num() > 0;
 }
 
 void FWindowsD3D12Adapter::CreateCommandSignatures()
diff --git a/Engine/Source/Runtime/Windows/D3D11RHI/Private/Windows/WindowsD3D11Device.cpp b/Engine/Source/Runtime/Windows/D3D11RHI/Private/Windows/WindowsD3D11Device.cpp
--- a/Engine/Source/Runtime/Windows/D3D11RHI/Private/Windows/WindowsD3D11Device.cpp
+++ b/Engine/Source/Runtime/Windows/D3D11RHI/Private/Windows/WindowsD3D11Device.cpp
@@ -2360,66 +2360,70 @@ bool FD3D11DynamicRHI::RHIGetAvailableResolutions(FScreenResolutionArray& Resolu
 			continue;
 		}
 
-		// It's still invalid to "succeed" and be given no modes.
-		checkf(NumModes > 0, TEXT("No display modes found for DXGI_FORMAT_R8G8B8A8_UNORM or DXGI_FORMAT_B8G8R8A8_UNORM formats!"));
-
-		DXGI_MODE_DESC* ModeList = new DXGI_MODE_DESC[ NumModes ];
-		VERIFYD3D11RESULT(Output->GetDisplayModeList(Format, 0, &NumModes, ModeList));
-
-		for(uint32 m = 0;m < NumModes;m++)
+		if (NumModes > 0)
 		{
-			CA_SUPPRESS(6385);
-			if (((int32)ModeList[m].Width >= MinAllowableResolutionX) &&
-				((int32)ModeList[m].Width <= MaxAllowableResolutionX) &&
-				((int32)ModeList[m].Height >= MinAllowableResolutionY) &&
-				((int32)ModeList[m].Height <= MaxAllowableResolutionY)
-				)
+			DXGI_MODE_DESC* ModeList = new DXGI_MODE_DESC[NumModes];
+			VERIFYD3D11RESULT(Output->GetDisplayModeList(Format, 0, &NumModes, ModeList));
+
+			for (uint32 m = 0; m < NumModes; m++)
 			{
-				bool bAddIt = true;
-				if (bIgnoreRefreshRate == false)
+				CA_SUPPRESS(6385);
+				if (((int32)ModeList[m].Width >= MinAllowableResolutionX) &&
+					((int32)ModeList[m].Width <= MaxAllowableResolutionX) &&
+					((int32)ModeList[m].Height >= MinAllowableResolutionY) &&
+					((int32)ModeList[m].Height <= MaxAllowableResolutionY)
+					)
 				{
-					if (((int32)ModeList[m].RefreshRate.Numerator < MinAllowableRefreshRate * ModeList[m].RefreshRate.Denominator) ||
-						((int32)ModeList[m].RefreshRate.Numerator > MaxAllowableRefreshRate * ModeList[m].RefreshRate.Denominator)
-						)
+					bool bAddIt = true;
+					if (bIgnoreRefreshRate == false)
 					{
-						continue;
+						if (((int32)ModeList[m].RefreshRate.Numerator < MinAllowableRefreshRate * ModeList[m].RefreshRate.Denominator) ||
+							((int32)ModeList[m].RefreshRate.Numerator > MaxAllowableRefreshRate * ModeList[m].RefreshRate.Denominator)
+							)
+						{
+							continue;
+						}
 					}
-				}
-				else
-				{
-					// See if it is in the list already
-					for (int32 CheckIndex = 0; CheckIndex < Resolutions.Num(); CheckIndex++)
+					else
 					{
-						FScreenResolutionRHI& CheckResolution = Resolutions[CheckIndex];
-						if ((CheckResolution.Width == ModeList[m].Width) &&
-							(CheckResolution.Height == ModeList[m].Height))
+						// See if it is in the list already
+						for (int32 CheckIndex = 0; CheckIndex < Resolutions.Num(); CheckIndex++)
 						{
-							// Already in the list...
-							bAddIt = false;
-							break;
+							FScreenResolutionRHI& CheckResolution = Resolutions[CheckIndex];
+							if ((CheckResolution.Width == ModeList[m].Width) &&
+								(CheckResolution.Height == ModeList[m].Height))
+							{
+								// Already in the list...
+								bAddIt = false;
+								break;
+							}
 						}
 					}
-				}
 
-				if (bAddIt)
-				{
-					// Add the mode to the list
-					int32 Temp2Index = Resolutions.AddZeroed();
-					FScreenResolutionRHI& ScreenResolution = Resolutions[Temp2Index];
+					if (bAddIt)
+					{
+						// Add the mode to the list
+						int32 Temp2Index = Resolutions.AddZeroed();
+						FScreenResolutionRHI& ScreenResolution = Resolutions[Temp2Index];
 
-					ScreenResolution.Width = ModeList[m].Width;
-					ScreenResolution.Height = ModeList[m].Height;
-					ScreenResolution.RefreshRate = ModeList[m].RefreshRate.Numerator / ModeList[m].RefreshRate.Denominator;
+						ScreenResolution.Width = ModeList[m].Width;
+						ScreenResolution.Height = ModeList[m].Height;
+						ScreenResolution.RefreshRate = ModeList[m].RefreshRate.Numerator / ModeList[m].RefreshRate.Denominator;
+					}
 				}
 			}
-		}
 
-		delete[] ModeList;
+			delete[] ModeList;
+		}
+		else
+		{
+			UE_LOG(LogD3D11RHI, Warning, TEXT("No display modes found for the standard format DXGI_FORMAT_R8G8B8A8_UNORM!"));
+		}
 
 		++CurrentOutput;
 
 	// TODO: Cap at 1 for default output
 	} while(CurrentOutput < 1); //-V654
 
-	return true;
+	return Resolutions.Num() > 0;
 }
