Make sampler state creation on certain platforms threadsafe.

diff --git a/Engine/Source/Runtime/Apple/MetalRHI/Private/MetalState.cpp b/Engine/Source/Runtime/Apple/MetalRHI/Private/MetalState.cpp
--- a/Engine/Source/Runtime/Apple/MetalRHI/Private/MetalState.cpp
+++ b/Engine/Source/Runtime/Apple/MetalRHI/Private/MetalState.cpp
@@ -223,9 +223,11 @@ private:
 };
 
 static FMetalStateObjectCache<FSamplerStateInitializerRHI, FMetalSampler> Samplers;
+static FCriticalSection SamplersCS;
 
 static FMetalSampler FindOrCreateSamplerState(mtlpp::Device Device, const FSamplerStateInitializerRHI& Initializer)
 {
+	FScopeLock Lock(&SamplersCS);
 	FMetalSampler State = Samplers.Find(Initializer);
 	if (!State.GetPtr())
 	{
diff --git a/Engine/Source/Runtime/Windows/D3D11RHI/Private/D3D11State.cpp b/Engine/Source/Runtime/Windows/D3D11RHI/Private/D3D11State.cpp
--- a/Engine/Source/Runtime/Windows/D3D11RHI/Private/D3D11State.cpp
+++ b/Engine/Source/Runtime/Windows/D3D11RHI/Private/D3D11State.cpp
@@ -10,10 +10,13 @@
 // Cache of Sampler States; we store pointers to both as we don't want the TMap to be artificially
 // modifying ref counts if not needed; so we manage that ourselves
 static TMap<ID3D11SamplerState*, FD3D11SamplerState*> GSamplerStateCache;
+static FCriticalSection GSamplerStateCacheCS;
 
 
 void EmptyD3DSamplerStateCache()
 {
+	FScopeLock Lock(&GSamplerStateCacheCS);
+
 	for (auto Iter = GSamplerStateCache.CreateIterator(); Iter; ++Iter )
 	{
 		auto* State = Iter.Value();
@@ -181,6 +184,7 @@ FSamplerStateRHIRef FD3D11DynamicRHI::RHICreateSamplerState(const FSamplerStateI
 	TRefCountPtr<ID3D11SamplerState> SamplerStateHandle;
 	VERIFYD3D11RESULT_EX(Direct3DDevice->CreateSamplerState(&SamplerDesc, SamplerStateHandle.GetInitReference()), Direct3DDevice);
 
+	FScopeLock Lock(&GSamplerStateCacheCS);
 	FD3D11SamplerState** Found = GSamplerStateCache.Find(SamplerStateHandle);
 	if (Found)
 	{
