Fix RHI non-unity errors

#horde 265705

diff --git a/Engine/Source/Runtime/OpenGLDrv/Private/OpenGLDevice.cpp b/Engine/Source/Runtime/OpenGLDrv/Private/OpenGLDevice.cpp
--- a/Engine/Source/Runtime/OpenGLDrv/Private/OpenGLDevice.cpp
+++ b/Engine/Source/Runtime/OpenGLDrv/Private/OpenGLDevice.cpp
@@ -22,6 +22,7 @@
 #include "OpenGLDrv.h"
 #include "OpenGLDrvPrivate.h"
 #include "SceneUtils.h"
+#include "GlobalShader.h"
 
 #include "HardwareInfo.h"
 #include "OpenGLProgramBinaryFileCache.h"
diff --git a/Engine/Source/Runtime/Windows/D3D11RHI/Private/Windows/WindowsD3D11Device.cpp b/Engine/Source/Runtime/Windows/D3D11RHI/Private/Windows/WindowsD3D11Device.cpp
--- a/Engine/Source/Runtime/Windows/D3D11RHI/Private/Windows/WindowsD3D11Device.cpp
+++ b/Engine/Source/Runtime/Windows/D3D11RHI/Private/Windows/WindowsD3D11Device.cpp
@@ -26,6 +26,7 @@
 #include "RHIValidation.h"
 #include "HAL/ExceptionHandling.h"
 #include "HDRHelper.h"
+#include "GlobalShader.h"
 
 #if NV_AFTERMATH
 bool GDX11NVAfterMathEnabled = false;
