- Fix d3d11 max constant buffer size to be 4k float4s rather than 4k bytes

diff --git a/Engine/Source/Runtime/Windows/D3D11RHI/Private/D3D11ConstantBuffer.cpp b/Engine/Source/Runtime/Windows/D3D11RHI/Private/D3D11ConstantBuffer.cpp
--- a/Engine/Source/Runtime/Windows/D3D11RHI/Private/D3D11ConstantBuffer.cpp
+++ b/Engine/Source/Runtime/Windows/D3D11RHI/Private/D3D11ConstantBuffer.cpp
@@ -16,7 +16,7 @@
 const uint32 GConstantBufferSizes[MAX_CONSTANT_BUFFER_SLOTS] = 
 {
 	// CBs must be a multiple of 16
-	(uint32)Align(MAX_GLOBAL_CONSTANT_BUFFER_SIZE, 16),
+	(uint32)Align(MAX_GLOBAL_CONSTANT_BUFFER_BYTE_SIZE, 16),
 };
 
 // New circular buffer system for faster constant uploads.  Avoids CopyResource and speeds things up considerably
diff --git a/Engine/Source/Runtime/Windows/D3D11RHI/Private/D3D11Device.cpp b/Engine/Source/Runtime/Windows/D3D11RHI/Private/D3D11Device.cpp
--- a/Engine/Source/Runtime/Windows/D3D11RHI/Private/D3D11Device.cpp
+++ b/Engine/Source/Runtime/Windows/D3D11RHI/Private/D3D11Device.cpp
@@ -208,7 +208,7 @@ FD3D11DynamicRHI::FD3D11DynamicRHI(IDXGIFactory1* InDXGIFactory1, D3D_FEATURE_LE
 	GMaxTextureDimensions = D3D11_REQ_TEXTURE2D_U_OR_V_DIMENSION;
 	GMaxCubeTextureDimensions = D3D11_REQ_TEXTURECUBE_DIMENSION;
 	GMaxTextureArrayLayers = D3D11_REQ_TEXTURE2D_ARRAY_AXIS_DIMENSION;
-	GRHIMaxConstantBufferByteSize = MAX_GLOBAL_CONSTANT_BUFFER_SIZE;
+	GRHIMaxConstantBufferByteSize = MAX_GLOBAL_CONSTANT_BUFFER_BYTE_SIZE;
 	GRHISupportsMSAADepthSampleAccess = true;
 	GRHISupportsRHIThread = !!EXPERIMENTAL_D3D11_RHITHREAD;
 
diff --git a/Engine/Source/Runtime/Windows/D3D11RHI/Private/Windows/WindowsD3D11ConstantBuffer.cpp b/Engine/Source/Runtime/Windows/D3D11RHI/Private/Windows/WindowsD3D11ConstantBuffer.cpp
--- a/Engine/Source/Runtime/Windows/D3D11RHI/Private/Windows/WindowsD3D11ConstantBuffer.cpp
+++ b/Engine/Source/Runtime/Windows/D3D11RHI/Private/Windows/WindowsD3D11ConstantBuffer.cpp
@@ -13,7 +13,7 @@ void FWinD3D11ConstantBuffer::InitDynamicRHI()
 	D3D11_BUFFER_DESC BufferDesc;
 	ZeroMemory( &BufferDesc, sizeof( D3D11_BUFFER_DESC ) );
 	// Verify constant buffer ByteWidth requirements
-	check(MaxSize <= D3D11_REQ_CONSTANT_BUFFER_ELEMENT_COUNT && (MaxSize % 16) == 0);
+	check(MaxSize <= (D3D11_REQ_CONSTANT_BUFFER_ELEMENT_COUNT * 16) && (MaxSize % 16) == 0);
 	BufferDesc.ByteWidth = MaxSize;
 
 	BufferDesc.Usage = D3D11_USAGE_DEFAULT;
diff --git a/Engine/Source/Runtime/Windows/D3D11RHI/Public/D3D11ConstantBuffer.h b/Engine/Source/Runtime/Windows/D3D11RHI/Public/D3D11ConstantBuffer.h
--- a/Engine/Source/Runtime/Windows/D3D11RHI/Public/D3D11ConstantBuffer.h
+++ b/Engine/Source/Runtime/Windows/D3D11RHI/Public/D3D11ConstantBuffer.h
@@ -12,7 +12,7 @@
 class FD3D11DynamicRHI;
 
 /** Size of the default constant buffer. */
-#define MAX_GLOBAL_CONSTANT_BUFFER_SIZE		4096
+#define MAX_GLOBAL_CONSTANT_BUFFER_BYTE_SIZE		(16*4096)
 
 // !!! These offsets must match the cbuffer register definitions in Common.usf !!!
 enum ED3D11ShaderOffsetBuffer
