redirect D3D log messages to the log file.

#rnx

diff --git a/Engine/Source/Runtime/D3D12RHI/Private/D3D12Adapter.cpp b/Engine/Source/Runtime/D3D12RHI/Private/D3D12Adapter.cpp
--- a/Engine/Source/Runtime/D3D12RHI/Private/D3D12Adapter.cpp
+++ b/Engine/Source/Runtime/D3D12RHI/Private/D3D12Adapter.cpp
@@ -98,6 +98,8 @@ static bool CheckD3DStoredMessages()
 			D3D12_MESSAGE* d3dMessage = nullptr;
 			SIZE_T AllocateSize = 0;
 
+			static const bool bBreakOnWarning = FParse::Param(FCommandLine::Get(), TEXT("d3dbreakonwarning"));
+
 			int StoredMessageCount = d3dInfoQueue->GetNumStoredMessagesAllowedByRetrievalFilter();
 			for (int MessageIndex = 0; MessageIndex < StoredMessageCount; MessageIndex++)
 			{
@@ -130,24 +132,32 @@ static bool CheckD3DStoredMessages()
 					// get the actual message data from the queue
 					hr = d3dInfoQueue->GetMessage(MessageIndex, d3dMessage, &MessageLength);
 
-					if (d3dMessage->Severity == D3D12_MESSAGE_SEVERITY_ERROR)
-					{
-						UE_LOG(LogD3D12RHI, Error, TEXT("%s"), ANSI_TO_TCHAR(d3dMessage->pDescription));
-					}
-					else if (d3dMessage->Severity == D3D12_MESSAGE_SEVERITY_WARNING)
-					{
-						UE_LOG(LogD3D12RHI, Warning, TEXT("%s"), ANSI_TO_TCHAR(d3dMessage->pDescription));
-					}
-					else 
+
+					switch (d3dMessage->Severity)
 					{
-						UE_LOG(LogD3D12RHI, Log, TEXT("%s"), ANSI_TO_TCHAR(d3dMessage->pDescription));
+					case D3D12_MESSAGE_SEVERITY_CORRUPTION:
+					case D3D12_MESSAGE_SEVERITY_ERROR:
+						{
+							UE_LOG(LogD3D12RHI, Error, TEXT("[D3DDebug] %s"), ANSI_TO_TCHAR(d3dMessage->pDescription));
+							bResult = true;
+							break;
+						}
+					case D3D12_MESSAGE_SEVERITY_WARNING:
+						{
+							UE_LOG(LogD3D12RHI, Warning, TEXT("[D3DDebug] %s"), ANSI_TO_TCHAR(d3dMessage->pDescription));
+							if (bBreakOnWarning)
+								bResult = true;
+							break;
+						}
+					default:
+						{
+							UE_LOG(LogD3D12RHI, Log, TEXT("[D3DDebug] %s"), ANSI_TO_TCHAR(d3dMessage->pDescription));
+							break;
+						}
 					}
 				}
-
-				// we got messages
-				bResult = true;
 			}
-
+			d3dInfoQueue->ClearStoredMessages();
 			if (AllocateSize > 0)
 			{
 				FMemory::Free(d3dMessage);
@@ -611,10 +621,10 @@ void FD3D12Adapter::CreateRootDevice(bool bWithDebug)
 			FMemory::Memzero(&NewFilter, sizeof(NewFilter));
 
 			// Turn off info msgs as these get really spewy
-			D3D12_MESSAGE_SEVERITY DenySeverity = D3D12_MESSAGE_SEVERITY_INFO;
-			NewFilter.DenyList.NumSeverities = 1;
-			NewFilter.DenyList.pSeverityList = &DenySeverity;
-
+			const bool bLogWarnings = FParse::Param(FCommandLine::Get(), TEXT("d3dbreakonwarning")) || FParse::Param(FCommandLine::Get(), TEXT("d3dlogwarnings"));
+			D3D12_MESSAGE_SEVERITY DenySeverity[] = { D3D12_MESSAGE_SEVERITY_INFO, D3D12_MESSAGE_SEVERITY_WARNING };
+			NewFilter.DenyList.NumSeverities = 1 + (bLogWarnings ? 0 : 1);
+			NewFilter.DenyList.pSeverityList = DenySeverity;
 			// Be sure to carefully comment the reason for any additions here!  Someone should be able to look at it later and get an idea of whether it is still necessary.
 			TArray<D3D12_MESSAGE_ID, TInlineAllocator<16>> DenyIds = {
 
@@ -703,7 +713,8 @@ void FD3D12Adapter::CreateRootDevice(bool bWithDebug)
 			// Enable this to break on a specific id in order to quickly get a callstack
 			//pd3dInfoQueue->SetBreakOnID(D3D12_MESSAGE_ID_DEVICE_DRAW_CONSTANT_BUFFER_TOO_SMALL, true);
 
-			if (FParse::Param(FCommandLine::Get(), TEXT("d3dbreakonwarning")))
+			// Break on D3D warnings if warning log or warning breakpoint is enabled
+			if (bLogWarnings)
 			{
 				pd3dInfoQueue->SetBreakOnSeverity(D3D12_MESSAGE_SEVERITY_WARNING, true);
 			}
@@ -1356,12 +1367,7 @@ FD3D12Adapter::~FD3D12Adapter()
 #endif
 
 #if PLATFORM_WINDOWS
-	if (ExceptionHandlerHandle != INVALID_HANDLE_VALUE)
-	{
-		RemoveVectoredExceptionHandler(ExceptionHandlerHandle);
-	}
-
-	if (D3D12RHI_ShouldCreateWithD3DDebug())
+	if (GetD3DDevice() && D3D12RHI_ShouldCreateWithD3DDebug())
 	{
 		TRefCountPtr<ID3D12DebugDevice> Debug;
 		if (SUCCEEDED(GetD3DDevice()->QueryInterface(IID_PPV_ARGS(Debug.GetInitReference()))))
@@ -1370,6 +1376,11 @@ FD3D12Adapter::~FD3D12Adapter()
 			Debug->ReportLiveDeviceObjects(rldoFlags);
 		}
 	}
+
+	if (ExceptionHandlerHandle != INVALID_HANDLE_VALUE)
+	{
+		RemoveVectoredExceptionHandler(ExceptionHandlerHandle);
+	}
 #endif
 
 #if D3D12_SUPPORTS_DXGI_DEBUG
diff --git a/Engine/Source/Runtime/Windows/D3D11RHI/Private/D3D11Device.cpp b/Engine/Source/Runtime/Windows/D3D11RHI/Private/D3D11Device.cpp
--- a/Engine/Source/Runtime/Windows/D3D11RHI/Private/D3D11Device.cpp
+++ b/Engine/Source/Runtime/Windows/D3D11RHI/Private/D3D11Device.cpp
@@ -668,6 +668,11 @@ void FD3D11DynamicRHI::CleanupD3DDevice()
 			}
 		}
 
+		if (ExceptionHandlerHandle != INVALID_HANDLE_VALUE)
+		{
+			RemoveVectoredExceptionHandler(ExceptionHandlerHandle);
+		}
+
 		// ORION - avoid shutdown crash that is currently present in UE
 #if !PLATFORM_SEH_EXCEPTIONS_DISABLED
 		if (1) // (IsRHIDeviceNVIDIA())
diff --git a/Engine/Source/Runtime/Windows/D3D11RHI/Private/D3D11RHIPrivate.h b/Engine/Source/Runtime/Windows/D3D11RHI/Private/D3D11RHIPrivate.h
--- a/Engine/Source/Runtime/Windows/D3D11RHI/Private/D3D11RHIPrivate.h
+++ b/Engine/Source/Runtime/Windows/D3D11RHI/Private/D3D11RHIPrivate.h
@@ -933,6 +933,8 @@ protected:
 
 	FDisplayInformationArray DisplayList;
 
+	HANDLE ExceptionHandlerHandle = INVALID_HANDLE_VALUE;
+
 	bool bRenderDoc = false;
 
 public:
diff --git a/Engine/Source/Runtime/Windows/D3D11RHI/Private/Windows/WindowsD3D11Device.cpp b/Engine/Source/Runtime/Windows/D3D11RHI/Private/Windows/WindowsD3D11Device.cpp
--- a/Engine/Source/Runtime/Windows/D3D11RHI/Private/Windows/WindowsD3D11Device.cpp
+++ b/Engine/Source/Runtime/Windows/D3D11RHI/Private/Windows/WindowsD3D11Device.cpp
@@ -149,6 +149,106 @@
 	FConsoleCommandDelegate::CreateStatic(&FD3D11DumpLiveObjects)
 );
 
+static bool CheckD3D11StoredMessages()
+{
+	bool bResult = false;
+
+	TRefCountPtr<ID3D11Debug> DebugDevice = nullptr;
+	VERIFYD3D11RESULT_EX(GD3D11RHI->GetDevice()->QueryInterface(__uuidof(ID3D11Debug), (void**)DebugDevice.GetInitReference()), GD3D11RHI->GetDevice());
+	if (DebugDevice)
+	{
+		TRefCountPtr<ID3D11InfoQueue> d3dInfoQueue;
+		if (SUCCEEDED(GD3D11RHI->GetDevice()->QueryInterface(__uuidof(ID3D11InfoQueue), (void**)d3dInfoQueue.GetInitReference())))
+		{
+			D3D11_MESSAGE* d3dMessage = nullptr;
+			SIZE_T AllocateSize = 0;
+
+			static bool bBreakOnWarning = FParse::Param(FCommandLine::Get(), TEXT("d3dbreakonwarning"));
+
+			int StoredMessageCount = d3dInfoQueue->GetNumStoredMessagesAllowedByRetrievalFilter();
+			for (int MessageIndex = 0; MessageIndex < StoredMessageCount; MessageIndex++)
+			{
+				SIZE_T MessageLength = 0;
+				HRESULT hr = d3dInfoQueue->GetMessage(MessageIndex, nullptr, &MessageLength);
+
+				// realloc the message
+				if (MessageLength > AllocateSize)
+				{
+					if (d3dMessage)
+					{
+						FMemory::Free(d3dMessage);
+						d3dMessage = nullptr;
+						AllocateSize = 0;
+					}
+
+					d3dMessage = (D3D11_MESSAGE*)FMemory::Malloc(MessageLength);
+					AllocateSize = MessageLength;
+				}
+
+				if (d3dMessage)
+				{
+					// get the actual message data from the queue
+					hr = d3dInfoQueue->GetMessage(MessageIndex, d3dMessage, &MessageLength);
+
+					switch (d3dMessage->Severity)
+					{
+					case D3D11_MESSAGE_SEVERITY_CORRUPTION:
+					case D3D11_MESSAGE_SEVERITY_ERROR:
+						{
+							UE_LOG(LogD3D11RHI, Error, TEXT("[D3DDebug] %s"), ANSI_TO_TCHAR(d3dMessage->pDescription));
+							bResult = true;
+							break;
+						}
+					case D3D11_MESSAGE_SEVERITY_WARNING:
+						{
+							UE_LOG(LogD3D11RHI, Warning, TEXT("[D3DDebug] %s"), ANSI_TO_TCHAR(d3dMessage->pDescription));
+							if (bBreakOnWarning)
+							{
+								bResult = true;
+							}
+							break;
+						}
+					default:
+						{
+							UE_LOG(LogD3D11RHI, Log, TEXT("[D3DDebug] %s"), ANSI_TO_TCHAR(d3dMessage->pDescription));
+						}
+					}
+				}
+			}
+			d3dInfoQueue->ClearStoredMessages();
+			if (AllocateSize > 0)
+			{
+				FMemory::Free(d3dMessage);
+			}
+		}
+	}
+
+	return bResult;
+}
+
+static LONG __stdcall D3D11VectoredExceptionHandler(EXCEPTION_POINTERS* InInfo)
+{
+	// Only handle D3D error codes here
+	if (InInfo->ExceptionRecord->ExceptionCode == _FACDXGI)
+	{
+		if (CheckD3D11StoredMessages())
+		{
+			if (FPlatformMisc::IsDebuggerPresent())
+			{
+				// when we get here, then it means that BreakOnSeverity was set for this error message, so request the debug break here as well
+				// when the debugger is attached
+				UE_DEBUG_BREAK();
+			}
+		}
+
+		// Handles the exception
+		return EXCEPTION_CONTINUE_EXECUTION;
+	}
+
+	// continue searching
+	return EXCEPTION_CONTINUE_SEARCH;
+}
+
 /** This function is used as a SEH filter to catch only delay load exceptions. */
 static bool IsDelayLoadException(PEXCEPTION_POINTERS ExceptionPointers)
 {
@@ -1989,6 +2089,70 @@
 		// We should get the feature level we asked for as earlier we checked to ensure it is supported.
 		check(ActualFeatureLevel == FeatureLevel);
 
+		if (bWithD3DDebug)
+		{
+			TRefCountPtr<ID3D11InfoQueue> d3dInfoQueue;
+			if (SUCCEEDED(GD3D11RHI->GetDevice()->QueryInterface(__uuidof(ID3D11InfoQueue), (void**)d3dInfoQueue.GetInitReference())))
+			{
+				/* install callback */
+				ExceptionHandlerHandle = AddVectoredExceptionHandler(1, D3D11VectoredExceptionHandler);
+
+				/* filter messages */
+				const bool bLogWarnings = FParse::Param(FCommandLine::Get(), TEXT("d3dbreakonwarning")) || FParse::Param(FCommandLine::Get(), TEXT("d3dlogwarnings"));
+				D3D11_INFO_QUEUE_FILTER NewFilter;
+				FMemory::Memzero(&NewFilter, sizeof(NewFilter));
+
+				D3D11_MESSAGE_SEVERITY DenySeverity[] = { D3D11_MESSAGE_SEVERITY_INFO, D3D11_MESSAGE_SEVERITY_WARNING };
+				NewFilter.DenyList.NumSeverities = 1 + (bLogWarnings ? 0 : 1);
+				NewFilter.DenyList.pSeverityList = DenySeverity;
+
+
+				// Be sure to carefully comment the reason for any additions here!  Someone should be able to look at it later and get an idea of whether it is still necessary.
+				D3D11_MESSAGE_ID DenyIds[]  = {
+					// OMSETRENDERTARGETS_INVALIDVIEW - d3d will complain if depth and color targets don't have the exact same dimensions, but actually
+					//	if the color target is smaller then things are ok.  So turn off this error.  There is a manual check in FD3D11DynamicRHI::SetRenderTarget
+					//	that tests for depth smaller than color and MSAA settings to match.
+					D3D11_MESSAGE_ID_OMSETRENDERTARGETS_INVALIDVIEW, 
+
+					// QUERY_BEGIN_ABANDONING_PREVIOUS_RESULTS - The RHI exposes the interface to make and issue queries and a separate interface to use that data.
+					//		Currently there is a situation where queries are issued and the results may be ignored on purpose.  Filtering out this message so it doesn't
+					//		swarm the debug spew and mask other important warnings
+					D3D11_MESSAGE_ID_QUERY_BEGIN_ABANDONING_PREVIOUS_RESULTS,
+					D3D11_MESSAGE_ID_QUERY_END_ABANDONING_PREVIOUS_RESULTS,
+
+					// D3D11_MESSAGE_ID_CREATEINPUTLAYOUT_EMPTY_LAYOUT - This is a warning that gets triggered if you use a null vertex declaration,
+					//       which we want to do when the vertex shader is generating vertices based on ID.
+					D3D11_MESSAGE_ID_CREATEINPUTLAYOUT_EMPTY_LAYOUT,
+
+					// D3D11_MESSAGE_ID_DEVICE_DRAW_INDEX_BUFFER_TOO_SMALL - This warning gets triggered by Slate draws which are actually using a valid index range.
+					//		The invalid warning seems to only happen when VS 2012 is installed.  Reported to MS.  
+					//		There is now an assert in DrawIndexedPrimitive to catch any valid errors reading from the index buffer outside of range.
+					D3D11_MESSAGE_ID_DEVICE_DRAW_INDEX_BUFFER_TOO_SMALL,
+
+					// D3D11_MESSAGE_ID_DEVICE_DRAW_RENDERTARGETVIEW_NOT_SET - This warning gets triggered by shadow depth rendering because the shader outputs
+					//		a color but we don't bind a color render target. That is safe as writes to unbound render targets are discarded.
+					//		Also, batched elements triggers it when rendering outside of scene rendering as it outputs to the GBuffer containing normals which is not bound.
+					(D3D11_MESSAGE_ID)3146081, // D3D11_MESSAGE_ID_DEVICE_DRAW_RENDERTARGETVIEW_NOT_SET,
+
+					// Spams constantly as we change the debug name on rendertargets that get reused.
+					D3D11_MESSAGE_ID_SETPRIVATEDATA_CHANGINGPARAMS, 
+				};
+
+				NewFilter.DenyList.NumIDs = sizeof(DenyIds)/sizeof(D3D11_MESSAGE_ID);
+				NewFilter.DenyList.pIDList = (D3D11_MESSAGE_ID*)&DenyIds;
+
+				d3dInfoQueue->PushStorageFilter(&NewFilter);
+
+				/* ensure callback is called */
+				d3dInfoQueue->SetBreakOnSeverity(D3D11_MESSAGE_SEVERITY_CORRUPTION, true);
+				d3dInfoQueue->SetBreakOnSeverity(D3D11_MESSAGE_SEVERITY_ERROR, true);
+				if (bLogWarnings)
+				{
+					d3dInfoQueue->SetBreakOnSeverity(D3D11_MESSAGE_SEVERITY_WARNING, true);
+				}
+			}
+		}
+
 		GRHIPersistentThreadGroupCount = 1440; // TODO: Revisit based on vendor/adapter/perf query
 
 		StateCache.Init(Direct3DDeviceIMContext);
@@ -2115,76 +2279,7 @@
 		}
 
 		SetupAfterDeviceCreation();
+		GRHISupportsHDROutput = SetupDisplayHDRMetaData();
 
-#if !(UE_BUILD_SHIPPING && WITH_EDITOR)
-		// Add some filter outs for known debug spew messages (that we don't care about)
-		if(DeviceFlags & D3D11_CREATE_DEVICE_DEBUG)
-		{
-			TRefCountPtr<ID3D11InfoQueue> InfoQueue;
-			VERIFYD3D11RESULT_EX(Direct3DDevice->QueryInterface( IID_ID3D11InfoQueue, (void**)InfoQueue.GetInitReference()), Direct3DDevice);
-			if (InfoQueue)
-			{
-				D3D11_INFO_QUEUE_FILTER NewFilter;
-				FMemory::Memzero(&NewFilter,sizeof(NewFilter));
-
-				// Turn off info msgs as these get really spewy
-				D3D11_MESSAGE_SEVERITY DenySeverity = D3D11_MESSAGE_SEVERITY_INFO;
-				NewFilter.DenyList.NumSeverities = 1;
-				NewFilter.DenyList.pSeverityList = &DenySeverity;
-
-				// Be sure to carefully comment the reason for any additions here!  Someone should be able to look at it later and get an idea of whether it is still necessary.
-				D3D11_MESSAGE_ID DenyIds[]  = {
-					// OMSETRENDERTARGETS_INVALIDVIEW - d3d will complain if depth and color targets don't have the exact same dimensions, but actually
-					//	if the color target is smaller then things are ok.  So turn off this error.  There is a manual check in FD3D11DynamicRHI::SetRenderTarget
-					//	that tests for depth smaller than color and MSAA settings to match.
-					D3D11_MESSAGE_ID_OMSETRENDERTARGETS_INVALIDVIEW, 
-
-					// QUERY_BEGIN_ABANDONING_PREVIOUS_RESULTS - The RHI exposes the interface to make and issue queries and a separate interface to use that data.
-					//		Currently there is a situation where queries are issued and the results may be ignored on purpose.  Filtering out this message so it doesn't
-					//		swarm the debug spew and mask other important warnings
-					D3D11_MESSAGE_ID_QUERY_BEGIN_ABANDONING_PREVIOUS_RESULTS,
-					D3D11_MESSAGE_ID_QUERY_END_ABANDONING_PREVIOUS_RESULTS,
-
-					// D3D11_MESSAGE_ID_CREATEINPUTLAYOUT_EMPTY_LAYOUT - This is a warning that gets triggered if you use a null vertex declaration,
-					//       which we want to do when the vertex shader is generating vertices based on ID.
-					D3D11_MESSAGE_ID_CREATEINPUTLAYOUT_EMPTY_LAYOUT,
-
-					// D3D11_MESSAGE_ID_DEVICE_DRAW_INDEX_BUFFER_TOO_SMALL - This warning gets triggered by Slate draws which are actually using a valid index range.
-					//		The invalid warning seems to only happen when VS 2012 is installed.  Reported to MS.  
-					//		There is now an assert in DrawIndexedPrimitive to catch any valid errors reading from the index buffer outside of range.
-					D3D11_MESSAGE_ID_DEVICE_DRAW_INDEX_BUFFER_TOO_SMALL,
-
-					// D3D11_MESSAGE_ID_DEVICE_DRAW_RENDERTARGETVIEW_NOT_SET - This warning gets triggered by shadow depth rendering because the shader outputs
-					//		a color but we don't bind a color render target. That is safe as writes to unbound render targets are discarded.
-					//		Also, batched elements triggers it when rendering outside of scene rendering as it outputs to the GBuffer containing normals which is not bound.
-					(D3D11_MESSAGE_ID)3146081, // D3D11_MESSAGE_ID_DEVICE_DRAW_RENDERTARGETVIEW_NOT_SET,
-
-					// Spams constantly as we change the debug name on rendertargets that get reused.
-					D3D11_MESSAGE_ID_SETPRIVATEDATA_CHANGINGPARAMS, 
-				};
-
-				NewFilter.DenyList.NumIDs = sizeof(DenyIds)/sizeof(D3D11_MESSAGE_ID);
-				NewFilter.DenyList.pIDList = (D3D11_MESSAGE_ID*)&DenyIds;
-
-				InfoQueue->PushStorageFilter(&NewFilter);
-
-				// Break on D3D debug errors.
-				InfoQueue->SetBreakOnSeverity(D3D11_MESSAGE_SEVERITY_ERROR,true);
-
-				// Enable this to break on a specific id in order to quickly get a callstack
-				//InfoQueue->SetBreakOnID(D3D11_MESSAGE_ID_DEVICE_DRAW_CONSTANT_BUFFER_TOO_SMALL, true);
-
-				if (FParse::Param(FCommandLine::Get(),TEXT("d3dbreakonwarning")))
-				{
-					InfoQueue->SetBreakOnSeverity(D3D11_MESSAGE_SEVERITY_WARNING,true);
-				}
-			}
-		}
-#endif
-		
-		{
-			GRHISupportsHDROutput = SetupDisplayHDRMetaData();
-		}
-
 		// Add device overclock state to crash context
 		const bool bIsGPUOverclocked = IsDeviceOverclocked();
