Adding missing platform checks.

diff --git a/Engine/Source/Runtime/Windows/D3D11RHI/Private/D3D11Device.cpp b/Engine/Source/Runtime/Windows/D3D11RHI/Private/D3D11Device.cpp
--- a/Engine/Source/Runtime/Windows/D3D11RHI/Private/D3D11Device.cpp
+++ b/Engine/Source/Runtime/Windows/D3D11RHI/Private/D3D11Device.cpp
@@ -668,10 +668,12 @@ void FD3D11DynamicRHI::CleanupD3DDevice()
 			}
 		}
 
+#if PLATFORM_WINDOWS
 		if (ExceptionHandlerHandle != INVALID_HANDLE_VALUE)
 		{
 			RemoveVectoredExceptionHandler(ExceptionHandlerHandle);
 		}
+#endif
 
 		// ORION - avoid shutdown crash that is currently present in UE
 #if !PLATFORM_SEH_EXCEPTIONS_DISABLED
diff --git a/Engine/Source/Runtime/Windows/D3D11RHI/Private/D3D11RHIPrivate.h b/Engine/Source/Runtime/Windows/D3D11RHI/Private/D3D11RHIPrivate.h
--- a/Engine/Source/Runtime/Windows/D3D11RHI/Private/D3D11RHIPrivate.h
+++ b/Engine/Source/Runtime/Windows/D3D11RHI/Private/D3D11RHIPrivate.h
@@ -933,7 +933,9 @@ protected:
 
 	FDisplayInformationArray DisplayList;
 
+#if PLATFORM_WINDOWS
 	HANDLE ExceptionHandlerHandle = INVALID_HANDLE_VALUE;
+#endif
 
 	bool bRenderDoc = false;
 
diff --git a/Engine/Source/Runtime/Windows/D3D11RHI/Private/Windows/WindowsD3D11Device.cpp b/Engine/Source/Runtime/Windows/D3D11RHI/Private/Windows/WindowsD3D11Device.cpp
--- a/Engine/Source/Runtime/Windows/D3D11RHI/Private/Windows/WindowsD3D11Device.cpp
+++ b/Engine/Source/Runtime/Windows/D3D11RHI/Private/Windows/WindowsD3D11Device.cpp
@@ -2094,8 +2094,10 @@ void FD3D11DynamicRHI::InitD3DDevice()
 			TRefCountPtr<ID3D11InfoQueue> d3dInfoQueue;
 			if (SUCCEEDED(GD3D11RHI->GetDevice()->QueryInterface(__uuidof(ID3D11InfoQueue), (void**)d3dInfoQueue.GetInitReference())))
 			{
+#if PLATFORM_WINDOWS
 				/* install callback */
 				ExceptionHandlerHandle = AddVectoredExceptionHandler(1, D3D11VectoredExceptionHandler);
+#endif
 
 				/* filter messages */
 				const bool bLogWarnings = FParse::Param(FCommandLine::Get(), TEXT("d3dbreakonwarning")) || FParse::Param(FCommandLine::Get(), TEXT("d3dlogwarnings"));
