Fix for creating empty uniform buffer off the render thread in D3D11.

diff --git a/Engine/Source/Runtime/Windows/D3D11RHI/Private/D3D11UniformBuffer.cpp b/Engine/Source/Runtime/Windows/D3D11RHI/Private/D3D11UniformBuffer.cpp
--- a/Engine/Source/Runtime/Windows/D3D11RHI/Private/D3D11UniformBuffer.cpp
+++ b/Engine/Source/Runtime/Windows/D3D11RHI/Private/D3D11UniformBuffer.cpp
@@ -106,7 +106,7 @@ void UniformBufferBeginFrame()
 static bool IsPoolingEnabled()
 {
 	static const auto CVar = IConsoleManager::Get().FindTConsoleVariableDataInt(TEXT("r.UniformBufferPooling"));
-	int32 CVarValue = CVar->GetValueOnRenderThread();
+	int32 CVarValue = CVar->GetValueOnAnyThread();
 	return CVarValue != 0;
 };
 
@@ -240,7 +240,7 @@ FUniformBufferRHIRef FD3D11DynamicRHI::RHICreateUniformBuffer(const void* Conten
 			ImmutableData.SysMemPitch = ImmutableData.SysMemSlicePitch = 0;
 
 			TRefCountPtr<ID3D11Buffer> UniformBufferResource;
-			VERIFYD3D11RESULT_EX(Direct3DDevice->CreateBuffer(&Desc,&ImmutableData,UniformBufferResource.GetInitReference()), Direct3DDevice);
+			VERIFYD3D11RESULT_EX(Direct3DDevice->CreateBuffer(&Desc, Contents ? &ImmutableData : nullptr,UniformBufferResource.GetInitReference()), Direct3DDevice);
 
 			const bool bAllocatedFromPool = false;
 			NewUniformBuffer = new FD3D11UniformBuffer(this, Layout, UniformBufferResource, FRingAllocation(), bAllocatedFromPool);
