Add structured buffer stride validation for d3d11

#rnx

diff --git a/Engine/Source/Runtime/Windows/D3D11RHI/Private/D3D11Shaders.cpp b/Engine/Source/Runtime/Windows/D3D11RHI/Private/D3D11Shaders.cpp
--- a/Engine/Source/Runtime/Windows/D3D11RHI/Private/D3D11Shaders.cpp
+++ b/Engine/Source/Runtime/Windows/D3D11RHI/Private/D3D11Shaders.cpp
@@ -54,6 +54,8 @@ static inline void ReadShaderOptionalData(FShaderCodeReader& InShaderCode, TShad
 	int32 IsSm6ShaderSize = 1;
 	const uint8* IsSm6Shader = InShaderCode.FindOptionalData('6', IsSm6ShaderSize);
 	OutShader.bIsSm6Shader = IsSm6Shader && IsSm6ShaderSize && *IsSm6Shader;
+
+	UE::RHICore::SetupShaderCodeValidationData(&OutShader, InShaderCode);
 }
 
 static bool ApplyVendorExtensions(ID3D11Device* Direct3DDevice, EShaderFrequency Frequency, const FD3D11ShaderData* ShaderData, bool& OutNeedsReset)
