Windows platform extensions extends the normal Windows platform features module.

#rnx

diff --git a/Engine/Source/Runtime/Windows/WindowsPlatformFeatures/Private/WindowsPlatformFeatures.h b/Engine/Source/Runtime/Windows/WindowsPlatformFeatures/Private/WindowsPlatformFeatures.h
deleted file mode 100644
diff --git a/Engine/Source/Runtime/Windows/WindowsPlatformFeatures/Public/WindowsPlatformFeatures.h b/Engine/Source/Runtime/Windows/WindowsPlatformFeatures/Public/WindowsPlatformFeatures.h
new file mode 100644
--- /dev/null
+++ b/Engine/Source/Runtime/Windows/WindowsPlatformFeatures/Public/WindowsPlatformFeatures.h
@@ -0,0 +1,22 @@
+// Copyright Epic Games, Inc. All Rights Reserved.
+
+#pragma once
+
+#include "CoreMinimal.h"
+#include "PlatformFeatures.h"
+
+class WINDOWSPLATFORMFEATURES_API FWindowsPlatformFeaturesModule : public IPlatformFeaturesModule
+{
+public:
+
+	/** Creates a new instance of the audio device implemented by the module. */
+	FWindowsPlatformFeaturesModule();
+
+	virtual IVideoRecordingSystem* GetVideoRecordingSystem() override;
+
+private:
+	void StartupModule() override;
+
+};
+
+
