- Reduce d3d11 constant buffer size to allow fluids to run and reduce impact on performance

#rnx

diff --git a/Engine/Source/Runtime/Windows/D3D11RHI/Public/D3D11ConstantBuffer.h b/Engine/Source/Runtime/Windows/D3D11RHI/Public/D3D11ConstantBuffer.h
--- a/Engine/Source/Runtime/Windows/D3D11RHI/Public/D3D11ConstantBuffer.h
+++ b/Engine/Source/Runtime/Windows/D3D11RHI/Public/D3D11ConstantBuffer.h
@@ -11,8 +11,11 @@
 
 class FD3D11DynamicRHI;
 
-/** Size of the default constant buffer. */
-#define MAX_GLOBAL_CONSTANT_BUFFER_BYTE_SIZE		(16*4096)
+/**
+Size of the default constant buffer.
+Note: D3D11 allows for 64k of storage but increasing this has a negative impact on performance.
+*/
+#define MAX_GLOBAL_CONSTANT_BUFFER_BYTE_SIZE		(2*4096)
 
 // !!! These offsets must match the cbuffer register definitions in Common.usf !!!
 enum ED3D11ShaderOffsetBuffer
