Log the name of the texture when the D3D11 CreateTexture call fails.

#rnx

diff --git a/Engine/Source/Runtime/Windows/D3D11RHI/Private/D3D11Texture.cpp b/Engine/Source/Runtime/Windows/D3D11RHI/Private/D3D11Texture.cpp
--- a/Engine/Source/Runtime/Windows/D3D11RHI/Private/D3D11Texture.cpp
+++ b/Engine/Source/Runtime/Windows/D3D11RHI/Private/D3D11Texture.cpp
@@ -295,7 +295,7 @@ void ApplyBC7SoftwareAdapterWorkaround(bool bSoftwareAdapter, D3D11_TEXTURE2D_DE
 /**
  * Creates a 2D texture optionally guarded by a structured exception handler.
  */
-void SafeCreateTexture2D(ID3D11Device* Direct3DDevice, int32 UEFormat, const D3D11_TEXTURE2D_DESC* TextureDesc, const D3D11_SUBRESOURCE_DATA* SubResourceData, ID3D11Texture2D** OutTexture2D)
+static void SafeCreateTexture2D(ID3D11Device* Direct3DDevice, int32 UEFormat, const D3D11_TEXTURE2D_DESC* TextureDesc, const D3D11_SUBRESOURCE_DATA* SubResourceData, ID3D11Texture2D** OutTexture2D, const TCHAR* DebugName)
 {
 #if GUARDED_TEXTURE_CREATES
 	bool bDriverCrash = true;
@@ -319,7 +319,8 @@ void SafeCreateTexture2D(ID3D11Device* Direct3DDevice, int32 UEFormat, const D3D
 			SubResourceData ? SubResourceData->pSysMem : nullptr,
 			SubResourceData ? SubResourceData->SysMemPitch : 0,
 			SubResourceData ? SubResourceData->SysMemSlicePitch : 0,
-			Direct3DDevice
+			Direct3DDevice,
+			DebugName
 			);
 #if GUARDED_TEXTURE_CREATES
 		bDriverCrash = false;
@@ -361,18 +362,18 @@ FD3D11Texture* FD3D11DynamicRHI::CreateD3D11Texture2D(FRHITextureCreateDesc cons
 
 	if (bCubeTexture)
 	{
-		checkf(SizeX <= GetMaxCubeTextureDimension(), TEXT("Requested cube texture size too large: %i, Max: %i"), SizeX, GetMaxCubeTextureDimension());
+		checkf(SizeX <= GetMaxCubeTextureDimension(), TEXT("Requested cube texture size too large: %i, Max: %i, DebugName: '%s'"), SizeX, GetMaxCubeTextureDimension(), CreateDesc.DebugName ? CreateDesc.DebugName : TEXT(""));
 		check(SizeX == SizeY);
 	}
 	else
 	{
-		checkf(SizeX <= GetMax2DTextureDimension(), TEXT("Requested texture2d x size too large: %i, Max: %i"), SizeX, GetMax2DTextureDimension());
-		checkf(SizeY <= GetMax2DTextureDimension(), TEXT("Requested texture2d y size too large: %i, Max: %i"), SizeY, GetMax2DTextureDimension());
+		checkf(SizeX <= GetMax2DTextureDimension(), TEXT("Requested texture2d x size too large: %i, Max: %i, DebugName: '%s'"), SizeX, GetMax2DTextureDimension(), CreateDesc.DebugName ? CreateDesc.DebugName : TEXT(""));
+		checkf(SizeY <= GetMax2DTextureDimension(), TEXT("Requested texture2d y size too large: %i, Max: %i, DebugName: '%s'"), SizeY, GetMax2DTextureDimension(), CreateDesc.DebugName ? CreateDesc.DebugName : TEXT(""));
 	}
 
 	if (bTextureArray)
 	{
-		checkf(SizeZ <= GetMaxTextureArrayLayers(), TEXT("Requested texture array size too large: %i, Max: %i"), SizeZ, GetMaxTextureArrayLayers());
+		checkf(SizeZ <= GetMaxTextureArrayLayers(), TEXT("Requested texture array size too large: %i, Max: %i, DebugName: '%s'"), SizeZ, GetMaxTextureArrayLayers(), CreateDesc.DebugName ? CreateDesc.DebugName : TEXT(""));
 	}
 
 	SCOPE_CYCLE_COUNTER(STAT_D3D11CreateTextureTime);
@@ -564,7 +565,7 @@ FD3D11Texture* FD3D11DynamicRHI::CreateD3D11Texture2D(FRHITextureCreateDesc cons
 	else
 #endif
 	{
-		SafeCreateTexture2D(Direct3DDevice, Format, &TextureDesc, pSubresourceData, TextureResource.GetInitReference());
+		SafeCreateTexture2D(Direct3DDevice, Format, &TextureDesc, pSubresourceData, TextureResource.GetInitReference(), CreateDesc.DebugName);
 	}
 
 	if (bCreateRTV)
@@ -878,7 +879,8 @@ FD3D11Texture* FD3D11DynamicRHI::CreateD3D11Texture3D(FRHITextureCreateDesc cons
 		SubResData ? SubResData->pSysMem : nullptr,
 		SubResData ? SubResData->SysMemPitch : 0,
 		SubResData ? SubResData->SysMemSlicePitch : 0,
-		Direct3DDevice
+		Direct3DDevice,
+		CreateDesc.DebugName
 		);
 
 	// Create a shader resource view for the texture.
@@ -1256,7 +1258,8 @@ void* FD3D11Texture::Lock(FD3D11DynamicRHI* D3DRHI, uint32 MipIndex, uint32 Arra
 			nullptr,
 			0,
 			0,
-			D3DRHI->GetDevice()
+			D3DRHI->GetDevice(),
+			*FString::Printf(TEXT("%s_Staging"), *GetName().ToString())
 			);
 		LockedData.StagingResource = StagingTexture;
 
diff --git a/Engine/Source/Runtime/Windows/D3D11RHI/Private/D3D11Util.cpp b/Engine/Source/Runtime/Windows/D3D11RHI/Private/D3D11Util.cpp
--- a/Engine/Source/Runtime/Windows/D3D11RHI/Private/D3D11Util.cpp
+++ b/Engine/Source/Runtime/Windows/D3D11RHI/Private/D3D11Util.cpp
@@ -276,7 +276,7 @@ void VerifyD3D11ShaderResult(FRHIShader* Shader, HRESULT D3DResult, const ANSICH
 
 void VerifyD3D11CreateTextureResult(HRESULT D3DResult, int32 UEFormat,const ANSICHAR* Code,const ANSICHAR* Filename,uint32 Line,uint32 SizeX,uint32 SizeY,uint32 SizeZ,uint8 D3DFormat,uint32 NumMips,uint32 Flags,
 	D3D11_USAGE Usage, uint32 CPUAccessFlags, uint32 MiscFlags, uint32 SampleCount, uint32 SampleQuality,
-	const void* SubResPtr, uint32 SubResPitch, uint32 SubResSlicePitch, ID3D11Device* Device)
+	const void* SubResPtr, uint32 SubResPitch, uint32 SubResSlicePitch, ID3D11Device* Device, const TCHAR* DebugName)
 {
 	check(FAILED(D3DResult));
 
@@ -291,7 +291,7 @@ void VerifyD3D11CreateTextureResult(HRESULT D3DResult, int32 UEFormat,const ANSI
 	}
 
 	UE_LOG(LogD3D11RHI, Error,
-		TEXT("%s failed with error %s\n at %s:%u\n Size=%ix%ix%i PF=%d D3DFormat=%s(0x%08X), NumMips=%i, Flags=%s, Usage:0x%x, CPUFlags:0x%x, MiscFlags:0x%x, SampleCount:0x%x, SampleQuality:0x%x, SubresPtr:0x%p, SubresPitch:%i, SubresSlicePitch:%i, DebugInfo: %s"),
+		TEXT("%s failed with error %s\n at %s:%u\n Size=%ix%ix%i PF=%d D3DFormat=%s(0x%08X), NumMips=%i, Flags=%s, Usage:0x%x, CPUFlags:0x%x, MiscFlags:0x%x, SampleCount:0x%x, SampleQuality:0x%x, SubresPtr:0x%p, SubresPitch:%i, SubresSlicePitch:%i, Name:'%s', DebugInfo: %s"),
 		ANSI_TO_TCHAR(Code),
 		*ErrorString,
 		ANSI_TO_TCHAR(Filename),
@@ -312,6 +312,7 @@ void VerifyD3D11CreateTextureResult(HRESULT D3DResult, int32 UEFormat,const ANSI
 		SubResPtr,
 		SubResPitch,
 		SubResSlicePitch,
+		DebugName ? DebugName : TEXT(""),
 		*DebugInfoString);
 
 	TerminateOnDeviceRemoved(D3DResult, Device);
diff --git a/Engine/Source/Runtime/Windows/D3D11RHI/Public/D3D11Util.h b/Engine/Source/Runtime/Windows/D3D11RHI/Public/D3D11Util.h
--- a/Engine/Source/Runtime/Windows/D3D11RHI/Public/D3D11Util.h
+++ b/Engine/Source/Runtime/Windows/D3D11RHI/Public/D3D11Util.h
@@ -49,7 +49,7 @@ extern D3D11RHI_API void VerifyD3D11ShaderResult(class FRHIShader* Shader, HRESU
 extern D3D11RHI_API void VerifyD3D11CreateTextureResult(HRESULT D3DResult, int32 UEFormat,const ANSICHAR* Code,const ANSICHAR* Filename,uint32 Line,
 										 uint32 SizeX,uint32 SizeY,uint32 SizeZ,uint8 D3DFormat,uint32 NumMips,uint32 Flags, D3D11_USAGE Usage,
 										 uint32 CPUAccessFlags, uint32 MiscFlags, uint32 SampleCount, uint32 SampleQuality,
-										 const void* SubResPtr, uint32 SubResPitch, uint32 SubResSlicePitch,ID3D11Device* Device);
+										 const void* SubResPtr, uint32 SubResPitch, uint32 SubResSlicePitch, ID3D11Device* Device, const TCHAR* DebugName);
 
 struct FD3D11ResizeViewportState
 {
@@ -73,7 +73,7 @@ extern D3D11RHI_API void VerifyD3D11CreateViewResult(HRESULT D3DResult, const AN
 #define VERIFYD3D11RESULT_NOEXIT(x)		{HRESULT hr = x; if (FAILED(hr)) { VerifyD3D11ResultNoExit(hr,#x,__FILE__,__LINE__, 0); }}
 #define VERIFYD3D11SHADERRESULT(Result, Shader, Device) {HRESULT hr = (Result); if (FAILED(hr)) { VerifyD3D11ShaderResult(Shader, hr, #Result,__FILE__,__LINE__, Device); }}
 #define VERIFYD3D11RESULT_NOEXIT(x)		{HRESULT hr = x; if (FAILED(hr)) { VerifyD3D11ResultNoExit(hr,#x,__FILE__,__LINE__, 0); }}
-#define VERIFYD3D11CREATETEXTURERESULT(x,UEFormat,SizeX,SizeY,SizeZ,Format,NumMips,Flags,Usage,CPUAccessFlags,MiscFlags,SampleCount,SampleQuality,SubResPtr,SubResPitch,SubResSlicePitch,Device) {HRESULT hr = x; if (FAILED(hr)) { VerifyD3D11CreateTextureResult(hr, UEFormat,#x,__FILE__,__LINE__,SizeX,SizeY,SizeZ,Format,NumMips,Flags,Usage,CPUAccessFlags,MiscFlags,SampleCount,SampleQuality,SubResPtr,SubResPitch,SubResSlicePitch,Device); }}
+#define VERIFYD3D11CREATETEXTURERESULT(x,UEFormat,SizeX,SizeY,SizeZ,Format,NumMips,Flags,Usage,CPUAccessFlags,MiscFlags,SampleCount,SampleQuality,SubResPtr,SubResPitch,SubResSlicePitch,Device,DebugName) {HRESULT hr = x; if (FAILED(hr)) { VerifyD3D11CreateTextureResult(hr, UEFormat,#x,__FILE__,__LINE__,SizeX,SizeY,SizeZ,Format,NumMips,Flags,Usage,CPUAccessFlags,MiscFlags,SampleCount,SampleQuality,SubResPtr,SubResPitch,SubResSlicePitch,Device,DebugName); }}
 #define VERIFYD3D11RESIZEVIEWPORTRESULT(x, OldState, NewState, Device) { HRESULT hr = x; if (FAILED(hr)) { VerifyD3D11ResizeViewportResult(hr, #x, __FILE__, __LINE__, OldState, NewState, Device); }}
 #define VERIFYD3D11CREATEVIEWRESULT(x, Device, Resource, Desc) {HRESULT hr = x; if (FAILED(hr)) { VerifyD3D11CreateViewResult(hr, #x, __FILE__, __LINE__, Device, Resource, Desc); }}
 
