Adding ISaveGameSystemModule module, comparable to the platform chunk install module.

Example usage added to Windows platform, allowing Windows platform extensions to override the default save system.
```
[PlatformFeatures]
SaveGameSystemModule=MySaveModule
```

diff --git a/Engine/Source/Runtime/Engine/Public/SaveGameSystem.h b/Engine/Source/Runtime/Engine/Public/SaveGameSystem.h
--- a/Engine/Source/Runtime/Engine/Public/SaveGameSystem.h
+++ b/Engine/Source/Runtime/Engine/Public/SaveGameSystem.h
@@ -9,6 +9,7 @@
 #include "Templates/Function.h"
 #include "Templates/SharedPointer.h"
 #include "Tasks/Pipe.h"
+#include "Modules/ModuleInterface.h"
 
 /**
  * Interface for platform feature modules
@@ -207,3 +208,14 @@ protected:
 
 };
 
+/**
+ * Save game system Module Interface
+ */
+class ISaveGameSystemModule : public IModuleInterface
+{
+public:
+
+	virtual ISaveGameSystem* GetSaveGameSystem() = 0;
+};
+
+
diff --git a/Engine/Source/Runtime/Windows/WindowsPlatformFeatures/Private/WindowsPlatformFeatures.cpp b/Engine/Source/Runtime/Windows/WindowsPlatformFeatures/Private/WindowsPlatformFeatures.cpp
--- a/Engine/Source/Runtime/Windows/WindowsPlatformFeatures/Private/WindowsPlatformFeatures.cpp
+++ b/Engine/Source/Runtime/Windows/WindowsPlatformFeatures/Private/WindowsPlatformFeatures.cpp
@@ -3,7 +3,9 @@
 #include "WindowsPlatformFeatures.h"
 #include "WmfPrivate.h"
 #include "WindowsVideoRecordingSystem.h"
+#include "SaveGameSystem.h"
 #include "Misc/CommandLine.h"
+#include "Misc/ConfigCacheIni.h"
 
 IMPLEMENT_MODULE(FWindowsPlatformFeaturesModule, WindowsPlatformFeatures);
 
@@ -17,6 +19,41 @@ IVideoRecordingSystem* FWindowsPlatformFeaturesModule::GetVideoRecordingSystem()
 	return &VideoRecordingSystem;
 }
 
+ISaveGameSystem* FWindowsPlatformFeaturesModule::GetSaveGameSystem()
+{
+	static ISaveGameSystem* SaveGameSystem = nullptr;
+	static bool bIniChecked = false;
+	if (!SaveGameSystem || !bIniChecked)
+	{
+		ISaveGameSystemModule* SaveGameSystemModule = nullptr;
+		if (!GEngineIni.IsEmpty())
+		{
+			FString SaveGameModule;
+			GConfig->GetString(TEXT("PlatformFeatures"), TEXT("SaveGameSystemModule"), SaveGameModule, GEngineIni);
+
+			if (!SaveGameModule.IsEmpty())
+			{
+				SaveGameSystemModule = FModuleManager::LoadModulePtr<ISaveGameSystemModule>(*SaveGameModule);
+				if (SaveGameSystemModule != nullptr)
+				{
+					// Attempt to grab the save game system
+					SaveGameSystem = SaveGameSystemModule->GetSaveGameSystem();
+				}
+			}
+			bIniChecked = true;
+		}
+
+		if (SaveGameSystem == nullptr)
+		{
+			// Placeholder/default instance
+			SaveGameSystem = IPlatformFeaturesModule::GetSaveGameSystem();
+		}
+	}
+
+	return SaveGameSystem;
+}
+
+
 void FWindowsPlatformFeaturesModule::StartupModule()
 {
 	FModuleManager::Get().LoadModule(TEXT("GameplayMediaEncoder"));
diff --git a/Engine/Source/Runtime/Windows/WindowsPlatformFeatures/Public/WindowsPlatformFeatures.h b/Engine/Source/Runtime/Windows/WindowsPlatformFeatures/Public/WindowsPlatformFeatures.h
--- a/Engine/Source/Runtime/Windows/WindowsPlatformFeatures/Public/WindowsPlatformFeatures.h
+++ b/Engine/Source/Runtime/Windows/WindowsPlatformFeatures/Public/WindowsPlatformFeatures.h
@@ -13,6 +13,7 @@ public:
 	FWindowsPlatformFeaturesModule();
 
 	virtual IVideoRecordingSystem* GetVideoRecordingSystem() override;
+	virtual class ISaveGameSystem* GetSaveGameSystem() override;
 
 private:
 	void StartupModule() override;
