Fixed mismatched copy constructor/assignment operator in FD3D11LockedKey and FD3D12LockedKey.

diff --git a/Engine/Source/Runtime/D3D12RHI/Private/D3D12Util.h b/Engine/Source/Runtime/D3D12RHI/Private/D3D12Util.h
--- a/Engine/Source/Runtime/D3D12RHI/Private/D3D12Util.h
+++ b/Engine/Source/Runtime/D3D12RHI/Private/D3D12Util.h
@@ -175,12 +175,6 @@ public:
 	{
 		return SourceObject != Other.SourceObject || Subresource != Other.Subresource;
 	}
-	FD3D12LockedKey& operator=(const FD3D12LockedKey& Other)
-	{
-		SourceObject = Other.SourceObject;
-		Subresource = Other.Subresource;
-		return *this;
-	}
 	uint32 GetHash() const
 	{
 		return PointerHash(SourceObject);
diff --git a/Engine/Source/Runtime/Windows/D3D11RHI/Public/D3D11Util.h b/Engine/Source/Runtime/Windows/D3D11RHI/Public/D3D11Util.h
--- a/Engine/Source/Runtime/Windows/D3D11RHI/Public/D3D11Util.h
+++ b/Engine/Source/Runtime/Windows/D3D11RHI/Public/D3D11Util.h
@@ -137,13 +137,6 @@ public:
 		return SourceObject != Other.SourceObject || Subresource != Other.Subresource;
 	}
 
-	FD3D11LockedKey& operator = (const FD3D11LockedKey& Other)
-	{
-		SourceObject = Other.SourceObject;
-		Subresource = Other.Subresource;
-		return *this;
-	}
-
 	uint32 GetHash() const
 	{
 		return PointerHash(SourceObject);
