PR #10329: Add on screen warning Debug device (eg d3ddebug) is enabled during stats (Contributed by [REDACTED_PERSON])

Original PR only added functionality for d3d12. This CL includes support in d3d11 and Vulkan, and makes d3d11 and d3d12 use a cvar to store d3ddebug settings, rather than parsing the commandline directly into a bool. This makes debug device settings more consistent across RHIs and lets you enable d3ddebug from an ini file now. Also added GRHIGlobals.IsDebugLayerEnabled to make tracking whether or not we're using a debug device consistent and easy across all RHIs.

diff --git a/Engine/Plugins/Experimental/NNE/Source/NNERuntimeRDG/Private/NNERuntimeRDGDml.cpp b/Engine/Plugins/Experimental/NNE/Source/NNERuntimeRDG/Private/NNERuntimeRDGDml.cpp
--- a/Engine/Plugins/Experimental/NNE/Source/NNERuntimeRDG/Private/NNERuntimeRDGDml.cpp
+++ b/Engine/Plugins/Experimental/NNE/Source/NNERuntimeRDG/Private/NNERuntimeRDGDml.cpp
@@ -54,7 +54,7 @@ bool FRuntimeDmlStartup()
 		DirectMLDLLPaths[0] = DirectMLRuntimeBinPath / TEXT("DirectML.dll");
 
 #ifdef WITH_DIRECTML_DEBUG
-		if (GetID3D12PlatformDynamicRHI()->IsD3DDebugEnabled())
+		if (GRHIGlobals.IsDebugLayerEnabled)
 		{
 			DirectMLDLLPaths[1] = DirectMLRuntimeBinPath / TEXT("DirectML.Debug.dll");
 			++NumPaths;
@@ -259,7 +259,7 @@ bool UNNERuntimeRDGDmlImpl::Init(bool bRegisterOnlyOperators)
 
 	// Set debugging flags
 #ifdef WITH_DIRECTML_DEBUG
-	if (RHI->IsD3DDebugEnabled())
+	if (GRHIGlobals.IsDebugLayerEnabled)
 	{
 		DmlCreateFlags |= DML_CREATE_DEVICE_FLAG_DEBUG;
 	}
diff --git a/Engine/Plugins/Experimental/NNERuntimeORTGpu/Source/NNERuntimeORT/Private/NNERuntimeORTModel.cpp b/Engine/Plugins/Experimental/NNERuntimeORTGpu/Source/NNERuntimeORT/Private/NNERuntimeORTModel.cpp
--- a/Engine/Plugins/Experimental/NNERuntimeORTGpu/Source/NNERuntimeORT/Private/NNERuntimeORTModel.cpp
+++ b/Engine/Plugins/Experimental/NNERuntimeORTGpu/Source/NNERuntimeORT/Private/NNERuntimeORTModel.cpp
@@ -371,7 +371,7 @@ namespace UE::NNERuntimeORT::Private
 		DML_CREATE_DEVICE_FLAGS DmlCreateFlags = DML_CREATE_DEVICE_FLAG_NONE;
 
 		// Set debugging flags
-		if (RHI->IsD3DDebugEnabled())
+		if (GRHIGlobals.IsDebugLayerEnabled)
 		{
 			DmlCreateFlags |= DML_CREATE_DEVICE_FLAG_DEBUG;
 		}
diff --git a/Engine/Source/Runtime/D3D12RHI/Private/D3D12Adapter.cpp b/Engine/Source/Runtime/D3D12RHI/Private/D3D12Adapter.cpp
--- a/Engine/Source/Runtime/D3D12RHI/Private/D3D12Adapter.cpp
+++ b/Engine/Source/Runtime/D3D12RHI/Private/D3D12Adapter.cpp
@@ -64,6 +64,31 @@ static int32 GD3D12EnableDRED = 0;
 static int32 GD3D12EnableLightweightDRED = 1;
 #endif // UE_BUILD_SHIPPING || UE_BUILD_TEST
 
+TAutoConsoleVariable<int32> GD3D12DebugCvar (
+	TEXT("r.D3D12.EnableD3DDebug"),
+	0,
+	TEXT("0 to disable d3ddebug layer (default)\n")
+	TEXT("1 to enable error logging (-d3ddebug) \n")
+	TEXT("2 to enable error & warning logging (-d3dlogwarnings)\n")
+	TEXT("3 to enable breaking on errors & warnings (-d3dbreakonwarning)\n")
+	TEXT("4 to enable CONTINUING on errors (-d3dcontinueonerrors)\n"),
+	ECVF_RenderThreadSafe | ECVF_ReadOnly);
+
+bool D3D12_ShouldLogD3DDebugWarnings()
+{
+	return GD3D12DebugCvar.GetValueOnAnyThread() > 1;
+}
+
+bool D3D12_ShouldBreakOnD3DDebugErrors()
+{
+	return GD3D12DebugCvar.GetValueOnAnyThread() > 0 && GD3D12DebugCvar.GetValueOnAnyThread() != 4;
+}
+
+bool D3D12_ShouldBreakOnD3DDebugWarnings()
+{
+	return GD3D12DebugCvar.GetValueOnAnyThread() > 3;
+}
+
 static FAutoConsoleVariableRef CVarD3D12EnableGPUBreadCrumbs(
 	TEXT("r.D3D12.BreadCrumbs"),
 	GD3D12EnableGPUBreadCrumbs,
@@ -694,9 +719,9 @@ void FD3D12Adapter::CreateRootDevice(bool bWithDebug)
 		TRefCountPtr<ID3D12InfoQueue> d3dInfoQueue;
 		if (SUCCEEDED(d3dDebug->QueryInterface(__uuidof(ID3D12InfoQueue), (void**)d3dInfoQueue.GetInitReference())))
 		{
-			d3dInfoQueue->SetBreakOnSeverity(D3D12_MESSAGE_SEVERITY_CORRUPTION, true);
-			d3dInfoQueue->SetBreakOnSeverity(D3D12_MESSAGE_SEVERITY_ERROR, true);
-			//d3dInfoQueue->SetBreakOnSeverity(D3D12_MESSAGE_SEVERITY_WARNING, true);
+			d3dInfoQueue->SetBreakOnSeverity(D3D12_MESSAGE_SEVERITY_CORRUPTION, D3D12_ShouldBreakOnD3DDebugErrors());
+			d3dInfoQueue->SetBreakOnSeverity(D3D12_MESSAGE_SEVERITY_ERROR, D3D12_ShouldBreakOnD3DDebugErrors());
+			d3dInfoQueue->SetBreakOnSeverity(D3D12_MESSAGE_SEVERITY_WARNING, D3D12_ShouldBreakOnD3DDebugWarnings());
 		}
 	}
 #endif
@@ -713,7 +738,7 @@ void FD3D12Adapter::CreateRootDevice(bool bWithDebug)
 			FMemory::Memzero(&NewFilter, sizeof(NewFilter));
 
 			// Turn off info msgs as these get really spewy
-			const bool bLogWarnings = FParse::Param(FCommandLine::Get(), TEXT("d3dbreakonwarning")) || FParse::Param(FCommandLine::Get(), TEXT("d3dlogwarnings"));
+			const bool bLogWarnings = D3D12_ShouldBreakOnD3DDebugWarnings() || D3D12_ShouldLogD3DDebugWarnings();
 			D3D12_MESSAGE_SEVERITY DenySeverity[] = { D3D12_MESSAGE_SEVERITY_INFO, D3D12_MESSAGE_SEVERITY_WARNING };
 			NewFilter.DenyList.NumSeverities = 1 + (bLogWarnings ? 0 : 1);
 			NewFilter.DenyList.pSeverityList = DenySeverity;
@@ -800,7 +825,7 @@ void FD3D12Adapter::CreateRootDevice(bool bWithDebug)
 			pd3dInfoQueue->PushStorageFilter(&NewFilter);
 
 			// Break on D3D debug errors.
-			pd3dInfoQueue->SetBreakOnSeverity(D3D12_MESSAGE_SEVERITY_ERROR, true);
+			pd3dInfoQueue->SetBreakOnSeverity(D3D12_MESSAGE_SEVERITY_ERROR, D3D12_ShouldBreakOnD3DDebugErrors());
 
 			// Enable this to break on a specific id in order to quickly get a callstack
 			//pd3dInfoQueue->SetBreakOnID(D3D12_MESSAGE_ID_DEVICE_DRAW_CONSTANT_BUFFER_TOO_SMALL, true);
@@ -808,7 +833,7 @@ void FD3D12Adapter::CreateRootDevice(bool bWithDebug)
 			// Break on D3D warnings if warning log or warning breakpoint is enabled
 			if (bLogWarnings)
 			{
-				pd3dInfoQueue->SetBreakOnSeverity(D3D12_MESSAGE_SEVERITY_WARNING, true);
+				pd3dInfoQueue->SetBreakOnSeverity(D3D12_MESSAGE_SEVERITY_WARNING, D3D12_ShouldBreakOnD3DDebugWarnings());
 			}
 
 			pd3dInfoQueue->Release();
@@ -844,7 +869,7 @@ void FD3D12Adapter::InitializeDevices()
 	FlushRenderingCommands();
 
 	// Use a debug device if specified on the command line.
-	bool bWithD3DDebug = D3D12RHI_ShouldCreateWithD3DDebug();
+	bool bWithD3DDebug = GRHIGlobals.IsDebugLayerEnabled;
 
 	// If we don't have a device yet, either because this is the first viewport, or the old device was removed, create a device.
 	if (!RootDevice)
@@ -1482,7 +1507,7 @@ FD3D12Adapter::~FD3D12Adapter()
 #endif
 
 #if PLATFORM_WINDOWS
-	if (GetD3DDevice() && D3D12RHI_ShouldCreateWithD3DDebug())
+	if (GetD3DDevice() && GRHIGlobals.IsDebugLayerEnabled)
 	{
 		TRefCountPtr<ID3D12DebugDevice> Debug;
 		if (SUCCEEDED(GetD3DDevice()->QueryInterface(IID_PPV_ARGS(Debug.GetInitReference()))))
diff --git a/Engine/Source/Runtime/D3D12RHI/Private/D3D12DirectCommandListManager.cpp b/Engine/Source/Runtime/D3D12RHI/Private/D3D12DirectCommandListManager.cpp
--- a/Engine/Source/Runtime/D3D12RHI/Private/D3D12DirectCommandListManager.cpp
+++ b/Engine/Source/Runtime/D3D12RHI/Private/D3D12DirectCommandListManager.cpp
@@ -3,8 +3,6 @@
 #include "D3D12RHIPrivate.h"
 #include "Windows.h"
 
-extern bool D3D12RHI_ShouldCreateWithD3DDebug();
-
 TLockFreePointerListUnordered<void, PLATFORM_CACHE_LINE_SIZE> FD3D12SyncPoint::MemoryPool;
 
 FD3D12GPUFence::FD3D12GPUFence(FName InName)
diff --git a/Engine/Source/Runtime/D3D12RHI/Private/D3D12RHI.cpp b/Engine/Source/Runtime/D3D12RHI/Private/D3D12RHI.cpp
--- a/Engine/Source/Runtime/D3D12RHI/Private/D3D12RHI.cpp
+++ b/Engine/Source/Runtime/D3D12RHI/Private/D3D12RHI.cpp
@@ -715,11 +715,6 @@ void FD3D12DynamicRHI::RHIWaitManualFence(FRHICommandList& RHICmdList, ID3D12Fen
 	Context.WaitManualFence(Fence, Value);
 }
 
-bool ID3D12DynamicRHI::IsD3DDebugEnabled()
-{
-	return D3D12RHI_ShouldCreateWithD3DDebug();
-}
-
 void* FD3D12DynamicRHI::RHIGetNativeDevice()
 {
 	return (void*)GetAdapter().GetD3DDevice();
diff --git a/Engine/Source/Runtime/D3D12RHI/Private/D3D12RHIPrivate.h b/Engine/Source/Runtime/D3D12RHI/Private/D3D12RHIPrivate.h
--- a/Engine/Source/Runtime/D3D12RHI/Private/D3D12RHIPrivate.h
+++ b/Engine/Source/Runtime/D3D12RHI/Private/D3D12RHIPrivate.h
@@ -97,15 +97,7 @@ inline t_A RoundUpToNextMultiple(const t_A& a, const t_B& b)
 
 using namespace D3D12RHI;
 
-static bool D3D12RHI_ShouldCreateWithD3DDebug()
-{
-	// Use a debug device if specified on the command line.
-	static bool bCreateWithD3DDebug =
-		FParse::Param(FCommandLine::Get(), TEXT("d3ddebug")) ||
-		FParse::Param(FCommandLine::Get(), TEXT("d3debug")) ||
-		FParse::Param(FCommandLine::Get(), TEXT("dxdebug"));
-	return bCreateWithD3DDebug;
-}
+extern TAutoConsoleVariable<int32> GD3D12DebugCvar;
 
 static bool D3D12RHI_ShouldCreateWithWarp()
 {
diff --git a/Engine/Source/Runtime/D3D12RHI/Private/D3D12StateCache.cpp b/Engine/Source/Runtime/D3D12RHI/Private/D3D12StateCache.cpp
--- a/Engine/Source/Runtime/D3D12RHI/Private/D3D12StateCache.cpp
+++ b/Engine/Source/Runtime/D3D12RHI/Private/D3D12StateCache.cpp
@@ -76,8 +76,6 @@ static FAutoConsoleVariableRef CVarOnlineDescriptorHeapBlockSize(
 	ECVF_ReadOnly
 );
 
-extern bool D3D12RHI_ShouldCreateWithD3DDebug();
-
 inline bool operator!=(D3D12_CPU_DESCRIPTOR_HANDLE lhs, D3D12_CPU_DESCRIPTOR_HANDLE rhs)
 {
 	return lhs.ptr != rhs.ptr;
@@ -946,7 +944,7 @@ bool FD3D12StateCache::AssertResourceStates(ED3D12PipelineType PipelineType)
 	return true;
 #else
 	// Can only verify resource states if the debug layer is used
-	static const bool bWithD3DDebug = D3D12RHI_ShouldCreateWithD3DDebug();
+	static const bool bWithD3DDebug = GRHIGlobals.IsDebugLayerEnabled;
 	if (!bWithD3DDebug)
 	{
 		UE_LOG(LogD3D12RHI, Fatal, TEXT("*** AssertResourceStates requires the debug layer ***"));
diff --git a/Engine/Source/Runtime/D3D12RHI/Private/D3D12Submission.cpp b/Engine/Source/Runtime/D3D12RHI/Private/D3D12Submission.cpp
--- a/Engine/Source/Runtime/D3D12RHI/Private/D3D12Submission.cpp
+++ b/Engine/Source/Runtime/D3D12RHI/Private/D3D12Submission.cpp
@@ -326,7 +326,7 @@ static int32 GetMaxExecuteBatchSize()
 {
 	return
 #if UE_BUILD_DEBUG
-		D3D12RHI_ShouldCreateWithD3DDebug() ? 1 :
+		GRHIGlobals.IsDebugLayerEnabled ? 1 :
 #endif
 		TNumericLimits<int32>::Max();
 }
diff --git a/Engine/Source/Runtime/D3D12RHI/Private/D3D12Util.cpp b/Engine/Source/Runtime/D3D12RHI/Private/D3D12Util.cpp
--- a/Engine/Source/Runtime/D3D12RHI/Private/D3D12Util.cpp
+++ b/Engine/Source/Runtime/D3D12RHI/Private/D3D12Util.cpp
@@ -22,8 +22,6 @@ D3D12Util.h: D3D RHI utility implementation.
   #define D3D12RHI_GPU_CRASH_LOG_VERBOSITY Fatal
 #endif
 
-extern bool D3D12RHI_ShouldCreateWithD3DDebug();
-
 template<typename PerDeviceFunction>
 void FD3D12DynamicRHI::ForEachDevice(ID3D12Device* inDevice, const PerDeviceFunction& pfPerDeviceFunction)
 {
@@ -1677,7 +1675,7 @@ bool AssertResourceState(ID3D12CommandList* pCommandList, FD3D12Resource* pResou
 	}
 
 	// Can only verify resource states if the debug layer is used
-	static const bool bWithD3DDebug = D3D12RHI_ShouldCreateWithD3DDebug();
+	static const bool bWithD3DDebug = GRHIGlobals.IsDebugLayerEnabled;
 	if (!bWithD3DDebug)
 	{
 		UE_LOG(LogD3D12RHI, Fatal, TEXT("*** AssertResourceState requires the debug layer ***"));
diff --git a/Engine/Source/Runtime/D3D12RHI/Private/Windows/WindowsD3D12Device.cpp b/Engine/Source/Runtime/D3D12RHI/Private/Windows/WindowsD3D12Device.cpp
--- a/Engine/Source/Runtime/D3D12RHI/Private/Windows/WindowsD3D12Device.cpp
+++ b/Engine/Source/Runtime/D3D12RHI/Private/Windows/WindowsD3D12Device.cpp
@@ -48,7 +48,6 @@
 
 IMPLEMENT_MODULE(FD3D12DynamicRHIModule, D3D12RHI);
 
-extern bool D3D12RHI_ShouldCreateWithD3DDebug();
 extern bool D3D12RHI_ShouldCreateWithWarp();
 extern bool D3D12RHI_AllowSoftwareFallback();
 extern bool D3D12RHI_ShouldAllowAsyncResourceCreation();
@@ -1351,6 +1350,27 @@ void FD3D12DynamicRHI::Init()
 	}
 #endif
 
+	// Use a debug device if specified on the command line.
+	if (FParse::Param(FCommandLine::Get(), TEXT("d3ddebug")) ||
+		FParse::Param(FCommandLine::Get(), TEXT("d3debug")) ||
+		FParse::Param(FCommandLine::Get(), TEXT("dxdebug")))
+	{
+		GD3D12DebugCvar->Set(1, ECVF_SetByCommandline);
+	}
+	if (FParse::Param(FCommandLine::Get(), TEXT("d3dlogwarnings")))
+	{
+		GD3D12DebugCvar->Set(2, ECVF_SetByCommandline);
+	}
+	if (FParse::Param(FCommandLine::Get(), TEXT("d3dbreakonwarning")))
+	{
+		GD3D12DebugCvar->Set(3, ECVF_SetByCommandline);
+	}
+	if (FParse::Param(FCommandLine::Get(), TEXT("d3dcontinueonerrors")))
+	{
+		GD3D12DebugCvar->Set(4, ECVF_SetByCommandline);
+	}
+	GRHIGlobals.IsDebugLayerEnabled = (GD3D12DebugCvar.GetValueOnAnyThread() > 0);
+
 	check(!GIsRHIInitialized);
 
 	if (FParse::Param(FCommandLine::Get(), TEXT("clearPSODriverCache")))
diff --git a/Engine/Source/Runtime/D3D12RHI/Public/ID3D12DynamicRHI.h b/Engine/Source/Runtime/D3D12RHI/Public/ID3D12DynamicRHI.h
--- a/Engine/Source/Runtime/D3D12RHI/Public/ID3D12DynamicRHI.h
+++ b/Engine/Source/Runtime/D3D12RHI/Public/ID3D12DynamicRHI.h
@@ -45,8 +45,6 @@ struct ID3D12DynamicRHI : public FDynamicRHI
 
 	virtual void						RHISignalManualFence(FRHICommandList& RHICmdList, ID3D12Fence* Fence, uint64 Value) = 0;
 	virtual void						RHIWaitManualFence(FRHICommandList& RHICmdList, ID3D12Fence* Fence, uint64 Value) = 0;
-
-	static D3D12RHI_API bool      IsD3DDebugEnabled();
 };
 
 inline bool IsRHID3D12()
diff --git a/Engine/Source/Runtime/Engine/Private/UnrealEngine.cpp b/Engine/Source/Runtime/Engine/Private/UnrealEngine.cpp
--- a/Engine/Source/Runtime/Engine/Private/UnrealEngine.cpp
+++ b/Engine/Source/Runtime/Engine/Private/UnrealEngine.cpp
@@ -12284,6 +12284,15 @@ void DrawStatsHUD( UWorld* World, FViewport* Viewport, FCanvas* Canvas, UCanvas*
 #endif
 
 #if STATS
+		if (GRHIGlobals.IsDebugLayerEnabled)
+		{
+			SmallTextItem.SetColor(FLinearColor::Red);
+			SmallTextItem.Text = LOCTEXT("GCPROFILINGWARNING", "PROFILING WITH DEBUG DEVICE ON!");
+			Canvas->DrawItem(SmallTextItem, FVector2D(MessageX, MessageY));
+			MessageY += FontSizeY;
+			SmallTextItem.SetColor(FLinearColor::White);
+		}
+
 		if (FThreadStats::IsCollectingData())
 		{
 			SmallTextItem.SetColor(FLinearColor::Red);
diff --git a/Engine/Source/Runtime/RHI/Public/RHIGlobals.h b/Engine/Source/Runtime/RHI/Public/RHIGlobals.h
--- a/Engine/Source/Runtime/RHI/Public/RHIGlobals.h
+++ b/Engine/Source/Runtime/RHI/Public/RHIGlobals.h
@@ -574,6 +574,9 @@ struct FRHIGlobals
 
 	/** Table for finding out which shader platform corresponds to a given feature level for this RHI. */
 	EShaderPlatform ShaderPlatformForFeatureLevel[ERHIFeatureLevel::Num];
+	
+	/** True if the RHI has initialized a device with the debug layer enabled. */
+	bool IsDebugLayerEnabled = false;
 };
 
 extern RHI_API FRHIGlobals GRHIGlobals;
@@ -733,6 +736,7 @@ extern RHI_API FRHIGlobals GRHIGlobals;
 #define GRHISupportsBindless                                   GRHIGlobals.SupportsBindless
 #define GRHISupportsReservedResources                          GRHIGlobals.SupportsReservedResources
 #define GShaderPlatformForFeatureLevel                         GRHIGlobals.ShaderPlatformForFeatureLevel
+#define GRHIIsDebugLayerEnabled                                GRHIGlobals.IsDebugLayerEnabled
 
 // Utility Getters
 
diff --git a/Engine/Source/Runtime/VulkanRHI/Private/VulkanLayers.cpp b/Engine/Source/Runtime/VulkanRHI/Private/VulkanLayers.cpp
--- a/Engine/Source/Runtime/VulkanRHI/Private/VulkanLayers.cpp
+++ b/Engine/Source/Runtime/VulkanRHI/Private/VulkanLayers.cpp
@@ -339,6 +339,7 @@ void FVulkanDynamicRHI::SetupValidationRequests()
 		}
 		GGPUValidationCvar->Set(2, ECVF_SetByCommandline);
 	}
+	GRHIGlobals.IsDebugLayerEnabled = (GValidationCvar.GetValueOnAnyThread() > 0);
 #endif
 }
 
@@ -391,7 +392,7 @@ void FVulkanIntanceSetupHelper::AddDebugLayers(const TArray<FLayerWithExtensions
 	}
 
 	// At this point the CVar holds the final value
-	const bool bUseVulkanValidation = (GValidationCvar.GetValueOnAnyThread() > 0);
+	const bool bUseVulkanValidation = GRHIGlobals.IsDebugLayerEnabled;
 	if (!bGfxReconstructOrVkTrace && bUseVulkanValidation)
 	{
 		if (GStandardValidationCvar.GetValueOnAnyThread() != 0)
@@ -497,7 +498,7 @@ void FVulkanDeviceSetupHelper::AddDebugLayers(const TArray<FLayerWithExtensions>
 #endif
 
 	// Verify that all requested debugging device-layers are available. Skip validation layers under RenderDoc
-	if (!GRenderDocFound && (GValidationCvar.GetValueOnAnyThread() > 0) && (GStandardValidationCvar.GetValueOnAnyThread() != 0))
+	if (!GRenderDocFound && GRHIGlobals.IsDebugLayerEnabled && (GStandardValidationCvar.GetValueOnAnyThread() != 0))
 	{
 		// Path for older drivers
 		AddRequestedLayer(STANDARD_VALIDATION_LAYER_NAME, LayerProperties, UEExtensions, OutLayers);
diff --git a/Engine/Source/Runtime/Windows/D3D11RHI/Private/D3D11Commands.cpp b/Engine/Source/Runtime/Windows/D3D11RHI/Private/D3D11Commands.cpp
--- a/Engine/Source/Runtime/Windows/D3D11RHI/Private/D3D11Commands.cpp
+++ b/Engine/Source/Runtime/Windows/D3D11RHI/Private/D3D11Commands.cpp
@@ -295,9 +295,7 @@ void FD3D11DynamicRHI::RHISetBoundShaderState(FRHIBoundShaderState* BoundShaderS
 		}
 	}
 
-	extern bool D3D11RHI_ShouldCreateWithD3DDebug();
-	static bool bHasD3DDebug = D3D11RHI_ShouldCreateWithD3DDebug();
-	if (GUnbindResourcesBetweenDrawsInDX11 || bHasD3DDebug)
+	if (GUnbindResourcesBetweenDrawsInDX11 || GRHIGlobals.IsDebugLayerEnabled)
 	{
 		ClearAllShaderResources();
 	}
diff --git a/Engine/Source/Runtime/Windows/D3D11RHI/Private/D3D11Device.cpp b/Engine/Source/Runtime/Windows/D3D11RHI/Private/D3D11Device.cpp
--- a/Engine/Source/Runtime/Windows/D3D11RHI/Private/D3D11Device.cpp
+++ b/Engine/Source/Runtime/Windows/D3D11RHI/Private/D3D11Device.cpp
@@ -16,15 +16,6 @@
 	#endif
 #include "Windows/HideWindowsPlatformTypes.h"
 
-bool D3D11RHI_ShouldCreateWithD3DDebug()
-{
-	// Use a debug device if specified on the command line.
-	return
-		FParse::Param(FCommandLine::Get(),TEXT("d3ddebug")) ||
-		FParse::Param(FCommandLine::Get(),TEXT("d3debug")) ||
-		FParse::Param(FCommandLine::Get(),TEXT("dxdebug"));
-}
-
 bool D3D11RHI_ShouldAllowAsyncResourceCreation()
 {
 	static bool bAllowAsyncResourceCreation = !FParse::Param(FCommandLine::Get(),TEXT("nod3dasync"));
@@ -654,7 +645,7 @@ void FD3D11DynamicRHI::CleanupD3DDevice()
 #endif // INTEL_METRICSDISCOVERY
 
 		// When running with D3D debug, clear state and flush the device to get rid of spurious live objects in D3D11's report.
-		if (D3D11RHI_ShouldCreateWithD3DDebug())
+		if (GRHIGlobals.IsDebugLayerEnabled)
 		{
 			Direct3DDeviceIMContext->ClearState();
 			Direct3DDeviceIMContext->Flush();
diff --git a/Engine/Source/Runtime/Windows/D3D11RHI/Private/D3D11RHIPrivate.h b/Engine/Source/Runtime/Windows/D3D11RHI/Private/D3D11RHIPrivate.h
--- a/Engine/Source/Runtime/Windows/D3D11RHI/Private/D3D11RHIPrivate.h
+++ b/Engine/Source/Runtime/Windows/D3D11RHI/Private/D3D11RHIPrivate.h
@@ -139,6 +139,8 @@ DECLARE_CYCLE_STAT_EXTERN(TEXT("RenderTargetCommitUAV"), STAT_D3D11RenderTargetC
 
 extern TAutoConsoleVariable<int32> GCVarUseSharedKeyedMutex;
 
+extern TAutoConsoleVariable<int32> GD3D11DebugCvar;
+
 struct FD3D11GlobalStats
 {
 	// in bytes, never change after RHI, needed to scale game features
diff --git a/Engine/Source/Runtime/Windows/D3D11RHI/Private/Windows/WindowsD3D11Device.cpp b/Engine/Source/Runtime/Windows/D3D11RHI/Private/Windows/WindowsD3D11Device.cpp
--- a/Engine/Source/Runtime/Windows/D3D11RHI/Private/Windows/WindowsD3D11Device.cpp
+++ b/Engine/Source/Runtime/Windows/D3D11RHI/Private/Windows/WindowsD3D11Device.cpp
@@ -41,8 +41,7 @@ bool GDX11IntelMetricsDiscoveryEnabled = false;
 
 FD3D11DynamicRHI*	GD3D11RHI = nullptr;
 
-extern bool D3D11RHI_ShouldCreateWithD3DDebug();
-extern bool D3D11RHI_ShouldAllowAsyncResourceCreation();
+bool D3D11RHI_ShouldAllowAsyncResourceCreation();
 
 static bool D3D11RHI_AllowSoftwareFallback()
 {
@@ -64,6 +63,31 @@ struct AmdAgsInfo
 static AmdAgsInfo AmdInfo;
 #endif
 
+TAutoConsoleVariable<int32> GD3D11DebugCvar(
+	TEXT("r.D3D11.EnableD3DDebug"),
+	0,
+	TEXT("0 to disable d3ddebug layer (default)\n")
+	TEXT("1 to enable error logging (-d3ddebug) \n")
+	TEXT("2 to enable error & warning logging (-d3dlogwarnings)\n")
+	TEXT("3 to enable breaking on errors & warnings (-d3dbreakonwarning)\n")
+	TEXT("4 to enable CONTINUING on errors (-d3dcontinueonerrors)\n"),
+	ECVF_RenderThreadSafe | ECVF_ReadOnly);
+
+bool D3D11_ShouldLogD3DDebugWarnings()
+{
+	return GD3D11DebugCvar.GetValueOnAnyThread() > 1;
+}
+
+bool D3D11_ShouldBreakOnD3DDebugErrors()
+{
+	return GD3D11DebugCvar.GetValueOnAnyThread() > 0 && GD3D11DebugCvar.GetValueOnAnyThread() != 4;
+}
+
+bool D3D11_ShouldBreakOnD3DDebugWarnings()
+{
+	return GD3D11DebugCvar.GetValueOnAnyThread() > 3;
+}
+
 static TAutoConsoleVariable<int32> CVarAMDUseMultiThreadedDevice(
 	TEXT("r.AMDD3D11MultiThreadedDevice"),
 	0,
@@ -99,7 +123,7 @@ namespace RHIConsoleVariables
 
 static void FD3D11DumpLiveObjects()
 {
-	if (D3D11RHI_ShouldCreateWithD3DDebug())
+	if (GRHIGlobals.IsDebugLayerEnabled)
 	{
 		TRefCountPtr<ID3D11Debug> DebugDevice = nullptr;
 		VERIFYD3D11RESULT_EX(GD3D11RHI->GetDevice()->QueryInterface(__uuidof(ID3D11Debug), (void**)DebugDevice.GetInitReference()), GD3D11RHI->GetDevice());
@@ -332,7 +356,7 @@ static bool SafeTestD3D11CreateDevice(IDXGIAdapter* Adapter,D3D_FEATURE_LEVEL Mi
 	ID3D11DeviceContext* D3DDeviceContext = nullptr;
 	uint32 DeviceFlags = D3D11_CREATE_DEVICE_SINGLETHREADED;
 	// Use a debug device if specified on the command line.
-	if(D3D11RHI_ShouldCreateWithD3DDebug())
+	if(GRHIGlobals.IsDebugLayerEnabled)
 	{
 		DeviceFlags |= D3D11_CREATE_DEVICE_DEBUG;
 	}
@@ -960,7 +984,7 @@ void FD3D11DynamicRHIModule::FindAdapter()
 
 	// Try to create the DXGIFactory1.  This will fail if we're not running Vista SP2 or higher.
 	TRefCountPtr<IDXGIFactory1> DXGIFactory1;
-	SafeCreateDXGIFactory(DXGIFactory1.GetInitReference(), D3D11RHI_ShouldCreateWithD3DDebug());
+	SafeCreateDXGIFactory(DXGIFactory1.GetInitReference(), GRHIGlobals.IsDebugLayerEnabled);
 	if(!DXGIFactory1)
 	{
 		return;
@@ -1238,7 +1262,8 @@ void FD3D11DynamicRHI::DisableQuadBufferStereo()
 void FD3D11DynamicRHI::FlushPendingLogs()
 {
 #if !(UE_BUILD_SHIPPING && WITH_EDITOR)
-	if (D3D11RHI_ShouldCreateWithD3DDebug())
+
+	if (GRHIGlobals.IsDebugLayerEnabled)
 	{
 		TRefCountPtr<ID3D11InfoQueue> InfoQueue = nullptr;
 		VERIFYD3D11RESULT_EX(Direct3DDevice->QueryInterface(IID_ID3D11InfoQueue, (void**)InfoQueue.GetInitReference()), Direct3DDevice);
@@ -1777,7 +1802,27 @@ void FD3D11DynamicRHI::InitD3DDevice()
 		uint32 DeviceFlags = D3D11RHI_ShouldAllowAsyncResourceCreation() ? 0 : D3D11_CREATE_DEVICE_SINGLETHREADED;
 
 		// Use a debug device if specified on the command line.
-		const bool bWithD3DDebug = D3D11RHI_ShouldCreateWithD3DDebug();
+		if (FParse::Param(FCommandLine::Get(), TEXT("d3ddebug")) ||
+			FParse::Param(FCommandLine::Get(), TEXT("d3debug")) ||
+			FParse::Param(FCommandLine::Get(), TEXT("dxdebug")))
+		{
+			GD3D11DebugCvar->Set(1, ECVF_SetByCommandline);
+		}
+		if (FParse::Param(FCommandLine::Get(), TEXT("d3dlogwarnings")))
+		{
+			GD3D11DebugCvar->Set(2, ECVF_SetByCommandline);
+		}
+		if (FParse::Param(FCommandLine::Get(), TEXT("d3dbreakonwarning")))
+		{
+			GD3D11DebugCvar->Set(3, ECVF_SetByCommandline);
+		}
+		if (FParse::Param(FCommandLine::Get(), TEXT("d3dcontinueonerrors")))
+		{
+			GD3D11DebugCvar->Set(4, ECVF_SetByCommandline);
+		}
+		GRHIGlobals.IsDebugLayerEnabled = (GD3D11DebugCvar.GetValueOnAnyThread() > 0);
+
+		const bool bWithD3DDebug = GRHIGlobals.IsDebugLayerEnabled;
 		FGenericCrashContext::SetEngineData(TEXT("RHI.D3DDebug"), bWithD3DDebug ? TEXT("true") : TEXT("false"));
 
 		if (bWithD3DDebug)
@@ -2076,7 +2121,7 @@ void FD3D11DynamicRHI::InitD3DDevice()
 #endif
 
 				/* filter messages */
-				const bool bLogWarnings = FParse::Param(FCommandLine::Get(), TEXT("d3dbreakonwarning")) || FParse::Param(FCommandLine::Get(), TEXT("d3dlogwarnings"));
+				const bool bLogWarnings = D3D11_ShouldBreakOnD3DDebugWarnings() || D3D11_ShouldLogD3DDebugWarnings();
 				D3D11_INFO_QUEUE_FILTER NewFilter;
 				FMemory::Memzero(&NewFilter, sizeof(NewFilter));
 
@@ -2122,11 +2167,11 @@ void FD3D11DynamicRHI::InitD3DDevice()
 				d3dInfoQueue->PushStorageFilter(&NewFilter);
 
 				/* ensure callback is called */
-				d3dInfoQueue->SetBreakOnSeverity(D3D11_MESSAGE_SEVERITY_CORRUPTION, true);
-				d3dInfoQueue->SetBreakOnSeverity(D3D11_MESSAGE_SEVERITY_ERROR, true);
+				d3dInfoQueue->SetBreakOnSeverity(D3D11_MESSAGE_SEVERITY_CORRUPTION, D3D11_ShouldBreakOnD3DDebugErrors());
+				d3dInfoQueue->SetBreakOnSeverity(D3D11_MESSAGE_SEVERITY_ERROR, D3D11_ShouldBreakOnD3DDebugErrors());
 				if (bLogWarnings)
 				{
-					d3dInfoQueue->SetBreakOnSeverity(D3D11_MESSAGE_SEVERITY_WARNING, true);
+					d3dInfoQueue->SetBreakOnSeverity(D3D11_MESSAGE_SEVERITY_WARNING, D3D11_ShouldBreakOnD3DDebugWarnings());
 				}
 			}
 		}
