No PCH build fix for Stats includes when STATS is 0

#rnx

diff --git a/Engine/Source/Runtime/Windows/D3D11RHI/Public/D3D11ConstantBuffer.h b/Engine/Source/Runtime/Windows/D3D11RHI/Public/D3D11ConstantBuffer.h
--- a/Engine/Source/Runtime/Windows/D3D11RHI/Public/D3D11ConstantBuffer.h
+++ b/Engine/Source/Runtime/Windows/D3D11RHI/Public/D3D11ConstantBuffer.h
@@ -7,7 +7,7 @@
 #pragma once
 
 #include "RenderResource.h"
-#include "Stats/Stats2.h"
+#include "Stats/Stats.h"
 
 class FD3D11DynamicRHI;
 
