Minor platform cleanup.

diff --git a/Engine/Source/Runtime/Windows/D3D11RHI/Private/D3D11Device.cpp b/Engine/Source/Runtime/Windows/D3D11RHI/Private/D3D11Device.cpp
--- a/Engine/Source/Runtime/Windows/D3D11RHI/Private/D3D11Device.cpp
+++ b/Engine/Source/Runtime/Windows/D3D11RHI/Private/D3D11Device.cpp
@@ -645,12 +645,10 @@ void FD3D11DynamicRHI::CleanupD3DDevice()
 			}
 		}
 
-#if PLATFORM_WINDOWS
 		if (ExceptionHandlerHandle != INVALID_HANDLE_VALUE)
 		{
 			RemoveVectoredExceptionHandler(ExceptionHandlerHandle);
 		}
-#endif
 
 		// ORION - avoid shutdown crash that is currently present in UE
 #if !PLATFORM_SEH_EXCEPTIONS_DISABLED
diff --git a/Engine/Source/Runtime/Windows/D3D11RHI/Private/D3D11RHIPrivate.h b/Engine/Source/Runtime/Windows/D3D11RHI/Private/D3D11RHIPrivate.h
--- a/Engine/Source/Runtime/Windows/D3D11RHI/Private/D3D11RHIPrivate.h
+++ b/Engine/Source/Runtime/Windows/D3D11RHI/Private/D3D11RHIPrivate.h
@@ -917,9 +917,7 @@ protected:
 
 	FDisplayInformationArray DisplayList;
 
-#if PLATFORM_WINDOWS
 	HANDLE ExceptionHandlerHandle = INVALID_HANDLE_VALUE;
-#endif
 
 	bool bRenderDoc = false;
 
diff --git a/Engine/Source/Runtime/Windows/D3D11RHI/Private/D3D11Texture.cpp b/Engine/Source/Runtime/Windows/D3D11RHI/Private/D3D11Texture.cpp
--- a/Engine/Source/Runtime/Windows/D3D11RHI/Private/D3D11Texture.cpp
+++ b/Engine/Source/Runtime/Windows/D3D11RHI/Private/D3D11Texture.cpp
@@ -125,7 +125,6 @@ void UpdateD3D11TextureStats(FD3D11Texture& Texture, bool bAllocating)
 		INC_DWORD_STAT(STAT_D3D11TexturesReleased);
 	}
 
-#if PLATFORM_WINDOWS
 	// On Windows there is no way to hook into the low level d3d allocations and frees.
 	// This means that we must manually add the tracking here.
 	if (bAllocating)
@@ -143,7 +142,6 @@ void UpdateD3D11TextureStats(FD3D11Texture& Texture, bool bAllocating)
 		LLM(FLowLevelMemTracker::Get().OnLowLevelFree(ELLMTracker::Default , Texture.GetResource()));
 		MemoryTrace_Free((uint64)Texture.GetResource(), EMemoryTraceRootHeap::VideoMemory);
 	}
-#endif
 }
 
 FDynamicRHI::FRHICalcTextureSizeResult FD3D11DynamicRHI::RHICalcTexturePlatformSize(FRHITextureDesc const& Desc, uint32 FirstMipIndex)
@@ -210,7 +208,7 @@ void ApplyBC7SoftwareAdapterWorkaround(bool bSoftwareAdapter, D3D11_TEXTURE2D_DE
 }
 
 /** If true, guard texture creates with SEH to log more information about a driver crash we are seeing during texture streaming. */
-#define GUARDED_TEXTURE_CREATES (PLATFORM_WINDOWS && !(UE_BUILD_SHIPPING || UE_BUILD_TEST || PLATFORM_COMPILER_CLANG))
+#define GUARDED_TEXTURE_CREATES (!(UE_BUILD_SHIPPING || UE_BUILD_TEST || PLATFORM_COMPILER_CLANG))
 
 /**
  * Creates a 2D texture optionally guarded by a structured exception handler.
diff --git a/Engine/Source/Runtime/Windows/D3D11RHI/Private/D3D11Viewport.cpp b/Engine/Source/Runtime/Windows/D3D11RHI/Private/D3D11Viewport.cpp
--- a/Engine/Source/Runtime/Windows/D3D11RHI/Private/D3D11Viewport.cpp
+++ b/Engine/Source/Runtime/Windows/D3D11RHI/Private/D3D11Viewport.cpp
@@ -313,9 +313,8 @@ void FD3D11Viewport::Resize(uint32 InSizeX, uint32 InSizeY, bool bInIsFullscreen
 	}
 
 	RECT WindowRect = {};
-#if PLATFORM_WINDOWS
 	GetWindowRect(WindowHandle, &WindowRect);
-#endif
+
 	FVector2D WindowTopLeft((float)WindowRect.left, (float)WindowRect.top);
 	FVector2D WindowBottomRight((float)WindowRect.right, (float)WindowRect.bottom);
 	bool bHDREnabled;
diff --git a/Engine/Source/Runtime/Windows/D3D11RHI/Private/Windows/WindowsD3D11Device.cpp b/Engine/Source/Runtime/Windows/D3D11RHI/Private/Windows/WindowsD3D11Device.cpp
--- a/Engine/Source/Runtime/Windows/D3D11RHI/Private/Windows/WindowsD3D11Device.cpp
+++ b/Engine/Source/Runtime/Windows/D3D11RHI/Private/Windows/WindowsD3D11Device.cpp
@@ -2098,10 +2098,8 @@ void FD3D11DynamicRHI::InitD3DDevice()
 			TRefCountPtr<ID3D11InfoQueue> d3dInfoQueue;
 			if (SUCCEEDED(GD3D11RHI->GetDevice()->QueryInterface(__uuidof(ID3D11InfoQueue), (void**)d3dInfoQueue.GetInitReference())))
 			{
-#if PLATFORM_WINDOWS
 				/* install callback */
 				ExceptionHandlerHandle = AddVectoredExceptionHandler(1, D3D11VectoredExceptionHandler);
-#endif
 
 				/* filter messages */
 				const bool bLogWarnings = D3D11_ShouldBreakOnD3DDebugWarnings() || D3D11_ShouldLogD3DDebugWarnings();
@@ -2163,7 +2161,7 @@ void FD3D11DynamicRHI::InitD3DDevice()
 
 		StateCache.Init(Direct3DDeviceIMContext);
 
-#if (UE_BUILD_SHIPPING && WITH_EDITOR) && PLATFORM_WINDOWS && !PLATFORM_64BITS
+#if (UE_BUILD_SHIPPING && WITH_EDITOR) && !PLATFORM_64BITS
 		// Disable PIX for windows in the shipping editor builds
 		D3DPERF_SetOptions(1);
 #endif
@@ -2216,7 +2214,6 @@ void FD3D11DynamicRHI::InitD3DDevice()
 
 		CACHE_NV_AFTERMATH_ENABLED();
 
-#if PLATFORM_WINDOWS
 		IUnknown* RenderDoc;
 		IID RenderDocID;
 		if (SUCCEEDED(IIDFromString(L"{A7AA6116-9C8D-4BBA-9083-B4D816B71B78}", &RenderDocID)))
@@ -2238,7 +2235,6 @@ void FD3D11DynamicRHI::InitD3DDevice()
 			// Running under Intel GPA, so enable capturing mode
 			GDynamicRHI->EnableIdealGPUCaptureOptions(true);
 		}
-#endif
 
 		if (IsRHIDeviceNVIDIA())
 		{
diff --git a/Engine/Source/Runtime/Windows/D3D11RHI/Private/Windows/WindowsD3D11Viewport.cpp b/Engine/Source/Runtime/Windows/D3D11RHI/Private/Windows/WindowsD3D11Viewport.cpp
--- a/Engine/Source/Runtime/Windows/D3D11RHI/Private/Windows/WindowsD3D11Viewport.cpp
+++ b/Engine/Source/Runtime/Windows/D3D11RHI/Private/Windows/WindowsD3D11Viewport.cpp
@@ -216,9 +216,8 @@ FD3D11Viewport::FD3D11Viewport(FD3D11DynamicRHI* InD3DRHI,HWND InWindowHandle,ui
 				SwapChain1->QueryInterface(IID_PPV_ARGS(SwapChain.GetInitReference()));
 
 				RECT WindowRect = {};
-#if PLATFORM_WINDOWS
 				GetWindowRect(WindowHandle, &WindowRect);
-#endif
+
 				FVector2D WindowTopLeft((float)WindowRect.left, (float)WindowRect.top);
 				FVector2D WindowBottomRight((float)WindowRect.right, (float)WindowRect.bottom);
 				bool bHDREnabled = false;
