Changing D3D11 global constant buffer updating to be more flexible with its efficient uploading.
When global constant buffers are updated with pending data, only PS and VS were trying to upload to smaller buffers to avoid the overhead of a full copy. This set of smaller buffers was hardcoded which was causing performance regressions when the max constant buffer size was increased.
This change extends the sub buffer allocations to all shader stages and makes it more flexible on the subbuffer count when the max constant buffer size changes.
Increased the Max constant buffer size in D3D11 to 16*4096 from 2*4096.
Changed the D3D11 VSConstantBuffers/PSConstantBuffers/etc from an array of 1 to single pointers.
Removed a number of unnecessary variables from the constant buffer code.

diff --git a/Engine/Source/Runtime/Windows/D3D11RHI/Private/D3D11Commands.cpp b/Engine/Source/Runtime/Windows/D3D11RHI/Private/D3D11Commands.cpp
--- a/Engine/Source/Runtime/Windows/D3D11RHI/Private/D3D11Commands.cpp
+++ b/Engine/Source/Runtime/Windows/D3D11RHI/Private/D3D11Commands.cpp
@@ -565,30 +565,32 @@ void FD3D11DynamicRHI::RHISetShaderUniformBuffer(FRHIComputeShader* ComputeShade
 
 void FD3D11DynamicRHI::RHISetShaderParameter(FRHIGraphicsShader* ShaderRHI,uint32 BufferIndex,uint32 BaseIndex,uint32 NumBytes,const void* NewValue)
 {
+	check(BufferIndex == 0);
+
 	switch (ShaderRHI->GetFrequency())
 	{
 	case SF_Vertex:
 	{
 		FD3D11VertexShader* VertexShader = static_cast<FD3D11VertexShader*>(ShaderRHI);
 		VALIDATE_BOUND_SHADER(VertexShader);
-		checkSlow(VSConstantBuffers[BufferIndex]);
-		VSConstantBuffers[BufferIndex]->UpdateConstant((const uint8*)NewValue, BaseIndex, NumBytes);
+		checkSlow(VSConstantBuffer);
+		VSConstantBuffer->UpdateConstant((const uint8*)NewValue, BaseIndex, NumBytes);
 	}
 	break;
 	case SF_Geometry:
 	{
 		FD3D11GeometryShader* GeometryShader = static_cast<FD3D11GeometryShader*>(ShaderRHI);
 		VALIDATE_BOUND_SHADER(GeometryShader);
-		checkSlow(GSConstantBuffers[BufferIndex]);
-		GSConstantBuffers[BufferIndex]->UpdateConstant((const uint8*)NewValue, BaseIndex, NumBytes);
+		checkSlow(GSConstantBuffer);
+		GSConstantBuffer->UpdateConstant((const uint8*)NewValue, BaseIndex, NumBytes);
 	}
 	break;
 	case SF_Pixel:
 	{
 		FD3D11PixelShader* PixelShader = static_cast<FD3D11PixelShader*>(ShaderRHI);
 		VALIDATE_BOUND_SHADER(PixelShader);
-		checkSlow(PSConstantBuffers[BufferIndex]);
-		PSConstantBuffers[BufferIndex]->UpdateConstant((const uint8*)NewValue, BaseIndex, NumBytes);
+		checkSlow(PSConstantBuffer);
+		PSConstantBuffer->UpdateConstant((const uint8*)NewValue, BaseIndex, NumBytes);
 	}
 	break;
 	default:
@@ -599,8 +601,8 @@ void FD3D11DynamicRHI::RHISetShaderParameter(FRHIGraphicsShader* ShaderRHI,uint3
 void FD3D11DynamicRHI::RHISetShaderParameter(FRHIComputeShader* ComputeShaderRHI,uint32 BufferIndex,uint32 BaseIndex,uint32 NumBytes,const void* NewValue)
 {
 	//VALIDATE_BOUND_SHADER(ComputeShaderRHI);
-	checkSlow(CSConstantBuffers[BufferIndex]);
-	CSConstantBuffers[BufferIndex]->UpdateConstant((const uint8*)NewValue,BaseIndex,NumBytes);
+	checkSlow(CSConstantBuffer);
+	CSConstantBuffer->UpdateConstant((const uint8*)NewValue,BaseIndex,NumBytes);
 }
 
 void FD3D11DynamicRHI::RHISetShaderParameters(FRHIComputeShader* Shader, TConstArrayView<uint8> InParametersData, TConstArrayView<FRHIShaderParameter> InParameters, TConstArrayView<FRHIShaderParameterResource> InResourceParameters, TConstArrayView<FRHIShaderParameterResource> InBindlessParameters)
@@ -1044,6 +1046,20 @@ static D3D11_PRIMITIVE_TOPOLOGY GetD3D11PrimitiveType(EPrimitiveType PrimitiveTy
 	return D3D11_PRIMITIVE_TOPOLOGY_TRIANGLELIST;
 }
 
+namespace FD3DRHIUtil
+{
+	template <EShaderFrequency ShaderFrequencyT>
+	inline void CommitConstants(FD3D11ConstantBuffer* InConstantBuffer, FD3D11StateCache& StateCache, bool bDiscardSharedConstants)
+	{
+		FWinD3D11ConstantBuffer* ConstantBuffer = static_cast<FWinD3D11ConstantBuffer*>(InConstantBuffer);
+		// Array may contain NULL entries to pad out to proper 
+		if (ConstantBuffer && ConstantBuffer->CommitConstantsToDevice(bDiscardSharedConstants))
+		{
+			ID3D11Buffer* DeviceBuffer = ConstantBuffer->GetConstantBuffer();
+			StateCache.SetConstantBuffer<ShaderFrequencyT>(DeviceBuffer, GLOBAL_CONSTANT_BUFFER_INDEX);
+		}
+	}
+};
 
 void FD3D11DynamicRHI::CommitNonComputeShaderConstants()
 {
@@ -1055,31 +1071,19 @@ void FD3D11DynamicRHI::CommitNonComputeShaderConstants()
 	if (CurrentBoundShaderState->bShaderNeedsGlobalConstantBuffer[SF_Vertex])
 	{
 		// Commit and bind vertex shader constants
-		for(uint32 i=0;i<MAX_CONSTANT_BUFFER_SLOTS; i++)
-		{
-			FD3D11ConstantBuffer* ConstantBuffer = VSConstantBuffers[i];
-			FD3DRHIUtil::CommitConstants<SF_Vertex>(ConstantBuffer, StateCache, i, bDiscardSharedConstants);
-		}
+		FD3DRHIUtil::CommitConstants<SF_Vertex>(VSConstantBuffer, StateCache, bDiscardSharedConstants);
 	}
 
 	if (CurrentBoundShaderState->bShaderNeedsGlobalConstantBuffer[SF_Geometry])
 	{
 		// Commit and bind geometry shader constants
-		for(uint32 i=0;i<MAX_CONSTANT_BUFFER_SLOTS; i++)
-		{
-			FD3D11ConstantBuffer* ConstantBuffer = GSConstantBuffers[i];
-			FD3DRHIUtil::CommitConstants<SF_Geometry>(ConstantBuffer, StateCache, i, bDiscardSharedConstants);
-		}
+		FD3DRHIUtil::CommitConstants<SF_Geometry>(GSConstantBuffer, StateCache, bDiscardSharedConstants);
 	}
 
 	if (CurrentBoundShaderState->bShaderNeedsGlobalConstantBuffer[SF_Pixel])
 	{
 		// Commit and bind pixel shader constants
-		for(uint32 i=0;i<MAX_CONSTANT_BUFFER_SLOTS; i++)
-		{
-			FD3D11ConstantBuffer* ConstantBuffer = PSConstantBuffers[i];
-			FD3DRHIUtil::CommitConstants<SF_Pixel>(ConstantBuffer, StateCache, i, bDiscardSharedConstants);
-		}
+		FD3DRHIUtil::CommitConstants<SF_Pixel>(PSConstantBuffer, StateCache, bDiscardSharedConstants);
 	}
 
 	bDiscardSharedConstants = false;
@@ -1087,14 +1091,8 @@ void FD3D11DynamicRHI::CommitNonComputeShaderConstants()
 
 void FD3D11DynamicRHI::CommitComputeShaderConstants()
 {
-	bool bLocalDiscardSharedConstants = true;
-
 	// Commit and bind compute shader constants
-	for(uint32 i=0;i<MAX_CONSTANT_BUFFER_SLOTS; i++)
-	{
-		FD3D11ConstantBuffer* ConstantBuffer = CSConstantBuffers[i];
-		FD3DRHIUtil::CommitConstants<SF_Compute>(ConstantBuffer, StateCache, i, bDiscardSharedConstants);
-	}
+	FD3DRHIUtil::CommitConstants<SF_Compute>(CSConstantBuffer, StateCache, bDiscardSharedConstants);
 }
 
 template <class ShaderType>
diff --git a/Engine/Source/Runtime/Windows/D3D11RHI/Private/D3D11ConstantBuffer.cpp b/Engine/Source/Runtime/Windows/D3D11RHI/Private/D3D11ConstantBuffer.cpp
--- a/Engine/Source/Runtime/Windows/D3D11RHI/Private/D3D11ConstantBuffer.cpp
+++ b/Engine/Source/Runtime/Windows/D3D11RHI/Private/D3D11ConstantBuffer.cpp
@@ -6,26 +6,9 @@
 
 #include "D3D11RHIPrivate.h"
 
-// SHANOND: Looks like UpdateSubresource is going to be the way to update these CBs.
-//			The driver writers are trying to optimize for UpdateSubresource and this
-//			will also avoid any driver renaming issues we may hit with map_discard.
-
-#define MAX_POOL_BUFFERS 1 // use update subresource and plop it into the command stream
-
-/** Sizes of constant buffers defined in ED3D11ShaderOffsetBuffer. */
-const uint32 GConstantBufferSizes[MAX_CONSTANT_BUFFER_SLOTS] = 
-{
-	// CBs must be a multiple of 16
-	(uint32)Align(MAX_GLOBAL_CONSTANT_BUFFER_BYTE_SIZE, 16),
-};
-
 // New circular buffer system for faster constant uploads.  Avoids CopyResource and speeds things up considerably
-FD3D11ConstantBuffer::FD3D11ConstantBuffer(FD3D11DynamicRHI* InD3DRHI, uint32 InSize, uint32 SubBuffers) : 
-	D3DRHI(InD3DRHI),
-	MaxSize(InSize),
-	ShadowData(NULL),
-	CurrentUpdateSize(0),
-	TotalUpdateSize(0)
+FD3D11ConstantBuffer::FD3D11ConstantBuffer(FD3D11DynamicRHI* InD3DRHI)
+	: D3DRHI(InD3DRHI)
 {
 	InitResource();
 }
@@ -42,8 +25,8 @@ void FD3D11ConstantBuffer::InitDynamicRHI()
 {
 // New circular buffer system for faster constant uploads.  Avoids CopyResource and speeds things up considerably
 	// aligned for best performance
-	ShadowData = (uint8*)FMemory::Malloc(MaxSize, 16);
-	FMemory::Memzero(ShadowData,MaxSize);
+	ShadowData = (uint8*)FMemory::Malloc(GetMaxSize(), 16);
+	FMemory::Memzero(ShadowData, GetMaxSize());
 }
 
 void FD3D11ConstantBuffer::ReleaseDynamicRHI()
diff --git a/Engine/Source/Runtime/Windows/D3D11RHI/Private/D3D11RHIPrivate.h b/Engine/Source/Runtime/Windows/D3D11RHI/Private/D3D11RHIPrivate.h
--- a/Engine/Source/Runtime/Windows/D3D11RHI/Private/D3D11RHIPrivate.h
+++ b/Engine/Source/Runtime/Windows/D3D11RHI/Private/D3D11RHIPrivate.h
@@ -902,10 +902,10 @@ protected:
 	bool bDiscardSharedConstants;
 
 	/** A list of all D3D constant buffers RHIs that have been created. */
-	TArray<TRefCountPtr<FD3D11ConstantBuffer> > VSConstantBuffers;
-	TArray<TRefCountPtr<FD3D11ConstantBuffer> > PSConstantBuffers;
-	TArray<TRefCountPtr<FD3D11ConstantBuffer> > GSConstantBuffers;
-	TArray<TRefCountPtr<FD3D11ConstantBuffer> > CSConstantBuffers;
+	TRefCountPtr<FD3D11ConstantBuffer> VSConstantBuffer;
+	TRefCountPtr<FD3D11ConstantBuffer> PSConstantBuffer;
+	TRefCountPtr<FD3D11ConstantBuffer> GSConstantBuffer;
+	TRefCountPtr<FD3D11ConstantBuffer> CSConstantBuffer;
 
 	/** A history of the most recently used bound shader states, used to keep transient bound shader states from being recreated for each use. */
 	TGlobalResource< TBoundShaderStateHistory<10000> > BoundShaderStateHistory;
diff --git a/Engine/Source/Runtime/Windows/D3D11RHI/Private/Windows/D3D11RHIPrivateUtil.h b/Engine/Source/Runtime/Windows/D3D11RHI/Private/Windows/D3D11RHIPrivateUtil.h
--- a/Engine/Source/Runtime/Windows/D3D11RHI/Private/Windows/D3D11RHIPrivateUtil.h
+++ b/Engine/Source/Runtime/Windows/D3D11RHI/Private/Windows/D3D11RHIPrivateUtil.h
@@ -7,18 +7,3 @@
 #pragma once
 
 #include "WindowsD3D11ConstantBuffer.h"
-
-struct FD3DRHIUtil
-{
-	template <EShaderFrequency ShaderFrequencyT>
-	static FORCEINLINE void CommitConstants(FD3D11ConstantBuffer* InConstantBuffer, FD3D11StateCache& StateCache, uint32 Index, bool bDiscardSharedConstants)
-	{
-		auto* ConstantBuffer = ((FWinD3D11ConstantBuffer*)InConstantBuffer);
-		// Array may contain NULL entries to pad out to proper 
-		if (ConstantBuffer && ConstantBuffer->CommitConstantsToDevice(bDiscardSharedConstants))
-		{
-			ID3D11Buffer* DeviceBuffer = ConstantBuffer->GetConstantBuffer();
-			StateCache.SetConstantBuffer<ShaderFrequencyT>(DeviceBuffer, Index);
-		}
-	}
-};
diff --git a/Engine/Source/Runtime/Windows/D3D11RHI/Private/Windows/WindowsD3D11ConstantBuffer.cpp b/Engine/Source/Runtime/Windows/D3D11RHI/Private/Windows/WindowsD3D11ConstantBuffer.cpp
--- a/Engine/Source/Runtime/Windows/D3D11RHI/Private/Windows/WindowsD3D11ConstantBuffer.cpp
+++ b/Engine/Source/Runtime/Windows/D3D11RHI/Private/Windows/WindowsD3D11ConstantBuffer.cpp
@@ -9,10 +9,12 @@
 
 void FWinD3D11ConstantBuffer::InitDynamicRHI()
 {
-	TRefCountPtr<ID3D11Buffer> CBuffer = NULL;
-	D3D11_BUFFER_DESC BufferDesc;
-	ZeroMemory( &BufferDesc, sizeof( D3D11_BUFFER_DESC ) );
+	TRefCountPtr<ID3D11Buffer> CBuffer;
+
+	D3D11_BUFFER_DESC BufferDesc{};
+
 	// Verify constant buffer ByteWidth requirements
+	constexpr uint32 MaxSize = GetMaxSize();
 	check(MaxSize <= (D3D11_REQ_CONSTANT_BUFFER_ELEMENT_COUNT * 16) && (MaxSize % 16) == 0);
 	BufferDesc.ByteWidth = MaxSize;
 
@@ -21,13 +23,16 @@ void FWinD3D11ConstantBuffer::InitDynamicRHI()
 	BufferDesc.BindFlags = D3D11_BIND_CONSTANT_BUFFER;
 	BufferDesc.MiscFlags = 0;
 
-	Buffers = new TRefCountPtr<ID3D11Buffer>[NumSubBuffers];
 	CurrentSubBuffer = 0;
-	for(uint32 s = 0;s < NumSubBuffers;s++)
+
+	while (BufferDesc.ByteWidth >= MIN_GLOBAL_CONSTANT_BUFFER_BYTE_SIZE)
 	{
+		FSubBuffer& SubBuffer = SubBuffers.Emplace_GetRef(BufferDesc.ByteWidth);
+
 		CA_SUPPRESS(6385);	// Doesn't like COM
-		VERIFYD3D11RESULT_EX(D3DRHI->GetDevice()->CreateBuffer(&BufferDesc, NULL, Buffers[s].GetInitReference()), D3DRHI->GetDevice());
-		UpdateBufferStats(Buffers[s],true);
+		VERIFYD3D11RESULT_EX(D3DRHI->GetDevice()->CreateBuffer(&BufferDesc, nullptr, SubBuffer.Buffer.GetInitReference()), D3DRHI->GetDevice());
+		UpdateBufferStats(SubBuffer.Buffer, true);
+
 		BufferDesc.ByteWidth = Align(BufferDesc.ByteWidth / 2, 16);
 	}
 
@@ -38,16 +43,24 @@ void FWinD3D11ConstantBuffer::ReleaseDynamicRHI()
 {
 	FD3D11ConstantBuffer::ReleaseDynamicRHI();
 
-	if(Buffers)
+	for (FSubBuffer& SubBuffer : SubBuffers)
+	{
+		UpdateBufferStats(SubBuffer.Buffer, false);
+	}
+
+	SubBuffers.Empty();
+}
+
+uint32 FWinD3D11ConstantBuffer::FindSubBufferForAllocationSize(uint64 InSize) const
+{
+	for (uint32 Index = SubBuffers.Num() - 1; Index > 0; Index--)
 	{
-		for(uint32 s = 0;s < NumSubBuffers;s++)
+		if (InSize <= SubBuffers[Index].Size)
 		{
-			UpdateBufferStats(Buffers[s],false);
+			return Index;
 		}
-
-		delete [] Buffers;
-		Buffers = NULL;
 	}
+	return 0;
 }
 
 bool FWinD3D11ConstantBuffer::CommitConstantsToDevice( bool bDiscardSharedConstants )
@@ -60,6 +73,8 @@ bool FWinD3D11ConstantBuffer::CommitConstantsToDevice( bool bDiscardSharedConsta
 
 	SCOPE_CYCLE_COUNTER(STAT_D3D11GlobalConstantBufferUpdateTime);
 
+	uint32	TotalUpdateSize = 0;
+
 	if ( bDiscardSharedConstants )
 	{
 		// If we're discarding shared constants, just use constants that have been updated since the last Commit.
@@ -74,21 +89,14 @@ bool FWinD3D11ConstantBuffer::CommitConstantsToDevice( bool bDiscardSharedConsta
 	// This basically keeps track dynamically how much data has been updated every frame
 	// and then divides up a "max" constant buffer size by halves down until it finds a large sections that more tightly covers
 	// the amount updated, assuming that all data in a constant buffer is updated each draw call and contiguous.
-	// This only used for the IndexSlot==0 constant buffer on the vertex shader.
-	// Which will have an indeterminate number of constant values that are generated by material shaders.
-	CurrentSubBuffer = 1;
-	uint32 BufferSize = MaxSize / 2;
-	while ( BufferSize >= TotalUpdateSize && CurrentSubBuffer < NumSubBuffers )
-	{
-		CurrentSubBuffer++;
-		BufferSize /= 2;
-	}
-	CurrentSubBuffer--;
-	BufferSize *= 2;
+	CurrentSubBuffer = FindSubBufferForAllocationSize(TotalUpdateSize);
 
-	ID3D11Buffer* Buffer = Buffers[CurrentSubBuffer];
-	check(((uint64)(ShadowData) & 0xf) == 0);
-	D3DRHI->GetDeviceContext()->UpdateSubresource(Buffer,0,NULL,(void*)ShadowData,BufferSize,BufferSize);
+	check(IsAligned(ShadowData, 16));
+
+	const uint32 BufferSize = SubBuffers[CurrentSubBuffer].Size;
+	ID3D11Buffer* Buffer = SubBuffers[CurrentSubBuffer].Buffer.GetReference();
+
+	D3DRHI->GetDeviceContext()->UpdateSubresource(Buffer, 0, nullptr, ShadowData, BufferSize, BufferSize);
 
 	CurrentUpdateSize = 0;
 
@@ -97,31 +105,10 @@ bool FWinD3D11ConstantBuffer::CommitConstantsToDevice( bool bDiscardSharedConsta
 
 void FD3D11DynamicRHI::InitConstantBuffers()
 {
-	// Allocate shader constant buffers.  All shader types can have access to all buffers.
-	//  Index==0 is reserved for "custom" params, and the rest are reserved by the system for Common
-	//	constants
-	VSConstantBuffers.Empty(MAX_CONSTANT_BUFFER_SLOTS);
-	PSConstantBuffers.Empty(MAX_CONSTANT_BUFFER_SLOTS);
-	GSConstantBuffers.Empty(MAX_CONSTANT_BUFFER_SLOTS);
-	CSConstantBuffers.Empty(MAX_CONSTANT_BUFFER_SLOTS);
-	for(int32 BufferIndex = 0;BufferIndex < MAX_CONSTANT_BUFFER_SLOTS;BufferIndex++)
-	{
-		uint32 Size = GConstantBufferSizes[BufferIndex];
-		uint32 SubBuffers = 1;
-		if(BufferIndex == GLOBAL_CONSTANT_BUFFER_INDEX) //-V547
-		{
-			SubBuffers = 5;
-		}
-
-		// Vertex shader can have subbuffers for index==0.  This is from Epic's original design for the auto-fit of size to
-		//	reduce the update costs of the buffer.
-
-		// New circular buffer system for faster constant uploads.  Avoids CopyResource and speeds things up considerably
-		VSConstantBuffers.Add(new FWinD3D11ConstantBuffer(this, Size, SubBuffers));
-		PSConstantBuffers.Add(new FWinD3D11ConstantBuffer(this, Size, SubBuffers));
-		GSConstantBuffers.Add(new FWinD3D11ConstantBuffer(this, Size));
-		CSConstantBuffers.Add(new FWinD3D11ConstantBuffer(this, Size));
-	}
+	VSConstantBuffer = new FWinD3D11ConstantBuffer(this);
+	PSConstantBuffer = new FWinD3D11ConstantBuffer(this);
+	GSConstantBuffer = new FWinD3D11ConstantBuffer(this);
+	CSConstantBuffer = new FWinD3D11ConstantBuffer(this);
 }
 
 DEFINE_STAT(STAT_D3D11GlobalConstantBufferUpdateTime);
diff --git a/Engine/Source/Runtime/Windows/D3D11RHI/Private/Windows/WindowsD3D11ConstantBuffer.h b/Engine/Source/Runtime/Windows/D3D11RHI/Private/Windows/WindowsD3D11ConstantBuffer.h
--- a/Engine/Source/Runtime/Windows/D3D11RHI/Private/Windows/WindowsD3D11ConstantBuffer.h
+++ b/Engine/Source/Runtime/Windows/D3D11RHI/Private/Windows/WindowsD3D11ConstantBuffer.h
@@ -13,11 +13,8 @@ struct ID3D11Buffer;
 class FWinD3D11ConstantBuffer : public FD3D11ConstantBuffer
 {
 public:
-	FWinD3D11ConstantBuffer(FD3D11DynamicRHI* InD3DRHI, uint32 InSize = 0, uint32 SubBuffers = 1) :
-		FD3D11ConstantBuffer(InD3DRHI, InSize, SubBuffers),
-		Buffers(nullptr),
-		CurrentSubBuffer(0),
-		NumSubBuffers(SubBuffers)
+	FWinD3D11ConstantBuffer(FD3D11DynamicRHI* InD3DRHI) :
+		FD3D11ConstantBuffer(InD3DRHI)
 	{
 	}
 
@@ -30,16 +27,23 @@ public:
 	*/
 	ID3D11Buffer* GetConstantBuffer() const
 	{
-		return Buffers[CurrentSubBuffer];
+		return SubBuffers[CurrentSubBuffer].Buffer.GetReference();
 	}
 
+	uint32 FindSubBufferForAllocationSize(uint64 InSize) const;
+
 	/**
 	* Unlocks the constant buffer so the data can be transmitted to the device
 	*/
 	bool CommitConstantsToDevice(bool bDiscardSharedConstants);
 
 private:
-	TRefCountPtr<ID3D11Buffer>* Buffers;
-	uint32	CurrentSubBuffer;
-	uint32	NumSubBuffers;
+	struct FSubBuffer
+	{
+		FSubBuffer(uint64 InSize) : Size(InSize) { }
+		TRefCountPtr<ID3D11Buffer> Buffer;
+		uint64 Size;
+	};
+	TArray<FSubBuffer> SubBuffers;
+	uint32	CurrentSubBuffer = 0;
 };
diff --git a/Engine/Source/Runtime/Windows/D3D11RHI/Public/D3D11ConstantBuffer.h b/Engine/Source/Runtime/Windows/D3D11RHI/Public/D3D11ConstantBuffer.h
--- a/Engine/Source/Runtime/Windows/D3D11RHI/Public/D3D11ConstantBuffer.h
+++ b/Engine/Source/Runtime/Windows/D3D11RHI/Public/D3D11ConstantBuffer.h
@@ -15,7 +15,8 @@ class FD3D11DynamicRHI;
 Size of the default constant buffer.
 Note: D3D11 allows for 64k of storage but increasing this has a negative impact on performance.
 */
-#define MAX_GLOBAL_CONSTANT_BUFFER_BYTE_SIZE		(2*4096)
+#define MAX_GLOBAL_CONSTANT_BUFFER_BYTE_SIZE		(16*4096)
+#define MIN_GLOBAL_CONSTANT_BUFFER_BYTE_SIZE		128
 
 // !!! These offsets must match the cbuffer register definitions in Common.usf !!!
 enum ED3D11ShaderOffsetBuffer
@@ -25,24 +26,22 @@ enum ED3D11ShaderOffsetBuffer
 	MAX_CONSTANT_BUFFER_SLOTS
 };
 
-/** Sizes of constant buffers defined in ED3D11ShaderOffsetBuffer. */
-extern const uint32 GConstantBufferSizes[MAX_CONSTANT_BUFFER_SLOTS];
-
 /**
  * A D3D constant buffer
  */
 class FD3D11ConstantBuffer : public FRenderResource, public FRefCountedObject
 {
 public:
-
 	// New circular buffer system for faster constant uploads.  Avoids CopyResource and speeds things up considerably
-	FD3D11ConstantBuffer(FD3D11DynamicRHI* InD3DRHI, uint32 InSize = 0, uint32 SubBuffers = 1);
+	FD3D11ConstantBuffer(FD3D11DynamicRHI* InD3DRHI);
 	virtual ~FD3D11ConstantBuffer();
 
 	// FRenderResource interface.
 	virtual void	InitDynamicRHI() override;
 	virtual void	ReleaseDynamicRHI() override;
 
+	static inline constexpr uint32 GetMaxSize() { return (uint32)Align(MAX_GLOBAL_CONSTANT_BUFFER_BYTE_SIZE, 16); }
+
 	/**
 	* Updates a variable in the constant buffer.
 	* @param Data - The data to copy into the constant buffer
@@ -52,24 +51,17 @@ public:
 	void UpdateConstant(const uint8* Data, uint16 Offset, uint16 InSize)
 	{
 		// Check that the data we are shadowing fits in the allocated shadow data
-		check((uint32)Offset + (uint32)InSize <= MaxSize);
+		check((uint32)Offset + (uint32)InSize <= GetMaxSize());
 		FMemory::Memcpy(ShadowData+Offset, Data, InSize);
 		CurrentUpdateSize = FMath::Max((uint32)(Offset + InSize), CurrentUpdateSize);
 	}
 
 protected:
 	FD3D11DynamicRHI* D3DRHI;
-	uint32	MaxSize;
-	uint8*	ShadowData;
+	uint8*	ShadowData = nullptr;
 
 	/** Size of all constants that has been updated since the last call to Commit. */
-	uint32	CurrentUpdateSize;
-
-	/**
-	 * Size of all constants that has been updated since the last Discard.
-	 * Includes "shared" constants that don't necessarily gets updated between every Commit.
-	 */
-	uint32	TotalUpdateSize;
+	uint32	CurrentUpdateSize = 0;
 };
 
 DECLARE_CYCLE_STAT_EXTERN(TEXT("Global Constant buffer update time"),STAT_D3D11GlobalConstantBufferUpdateTime,STATGROUP_D3D11RHI, );
