Fix issue with XAudio2 dll expecting ref count to increment when GetDllHandle not doing so if the dll is already loaded

diff --git a/Engine/Source/Runtime/Windows/AudioMixerXAudio2/Private/AudioMixerPlatformXAudio2.cpp b/Engine/Source/Runtime/Windows/AudioMixerXAudio2/Private/AudioMixerPlatformXAudio2.cpp
--- a/Engine/Source/Runtime/Windows/AudioMixerXAudio2/Private/AudioMixerPlatformXAudio2.cpp
+++ b/Engine/Source/Runtime/Windows/AudioMixerXAudio2/Private/AudioMixerPlatformXAudio2.cpp
@@ -623,8 +623,10 @@ namespace Audio
 		// Load the xaudio2 library and keep a handle so we can free it on teardown
 		// Note: windows internally ref-counts the library per call to load library so 
 		// when we call FreeLibrary, it will only free it once the refcount is zero
-		XAudio2Dll = (HMODULE)FPlatformProcess::GetDllHandle(*DllName.GetPlainNameString());
-			
+		// Also, FPlatformProcess::GetDllHandle should not be used, as it will not increase ref count further if the library is already loaded.
+		// FPaths::ConvertRelativePathToFull is used for parity with how GetDllHandle calls LoadLibrary.
+		XAudio2Dll = LoadLibrary(*FPaths::ConvertRelativePathToFull(DllName.GetPlainNameString()));
+
 		// returning null means we failed to load XAudio2, which means everything will fail
 		if (XAudio2Dll == nullptr)
 		{
