Fixing the -Prefer<IHV> functionality for DirectX when selecting which graphics adapter to use and cleaning up DX11 adapter selection code to more closely match our DX12 implementation.

diff --git a/Engine/Source/Runtime/D3D12RHI/Private/Windows/WindowsD3D12Device.cpp b/Engine/Source/Runtime/D3D12RHI/Private/Windows/WindowsD3D12Device.cpp
--- a/Engine/Source/Runtime/D3D12RHI/Private/Windows/WindowsD3D12Device.cpp
+++ b/Engine/Source/Runtime/D3D12RHI/Private/Windows/WindowsD3D12Device.cpp
@@ -849,6 +849,7 @@ void FD3D12DynamicRHIModule::FindAdapter()
 	TRefCountPtr<IDXGIAdapter> TempAdapter;
 	const D3D_FEATURE_LEVEL MinRequiredFeatureLevel = GetRequiredD3DFeatureLevel();
 
+	FD3D12AdapterDesc PreferredAdapter;
 	FD3D12AdapterDesc FirstDiscreteAdapter;
 	FD3D12AdapterDesc BestMemoryAdapter;
 	FD3D12AdapterDesc FirstAdapter;
@@ -940,23 +941,28 @@ void FD3D12DynamicRHIModule::FindAdapter()
 
 				const bool bSkipAdapter = bSkipWARP || bSkipPerfHUDAdapter || bSkipHmdGraphicsAdapter || bSkipExplicitAdapter || FParse::Param(FCommandLine::Get(), TEXT("ForceZeroAdapters"));
 
-				if (!bSkipAdapter)
+				if (!bSkipAdapter && CurrentAdapter.IsValid())
 				{
-					if ((PreferredVendor != EGpuVendorId::Unknown) && (PreferredVendor == RHIConvertToGpuVendorId(AdapterDesc.VendorId)) && FirstDiscreteAdapter.IsValid())
+					if (PreferredVendor != EGpuVendorId::Unknown && PreferredVendor == RHIConvertToGpuVendorId(AdapterDesc.VendorId) && !PreferredAdapter.IsValid())
 					{
-						FirstDiscreteAdapter = CurrentAdapter;
+						PreferredAdapter = CurrentAdapter;
 					}
-					else if (!bIsWARP && !CurrentAdapter.bIsIntegrated)
+					
+					if (!bIsWARP && !CurrentAdapter.bIsIntegrated)
 					{
 						if (!FirstDiscreteAdapter.IsValid())
 						{
 							FirstDiscreteAdapter = CurrentAdapter;
 						}
 
-						if ((PreferredVendor == EGpuVendorId::Unknown || PreferredVendor == RHIConvertToGpuVendorId(AdapterDesc.VendorId)) &&
-						    CurrentAdapter.Desc.DedicatedVideoMemory > BestMemoryAdapter.Desc.DedicatedVideoMemory)
+						if (CurrentAdapter.Desc.DedicatedVideoMemory > BestMemoryAdapter.Desc.DedicatedVideoMemory)
 						{
 							BestMemoryAdapter = CurrentAdapter;
+							if (PreferredVendor != EGpuVendorId::Unknown && PreferredVendor == RHIConvertToGpuVendorId(AdapterDesc.VendorId))
+							{
+								// Choose the best option of the preferred IHV devices
+								PreferredAdapter = BestMemoryAdapter;
+							}
 						}
 					}
 
@@ -964,10 +970,6 @@ void FD3D12DynamicRHIModule::FindAdapter()
 					{
 						FirstAdapter = CurrentAdapter;
 					}
-					else if ((PreferredVendor != EGpuVendorId::Unknown) && (PreferredVendor == RHIConvertToGpuVendorId(AdapterDesc.VendorId)) && FirstAdapter.IsValid())
-					{
-						FirstAdapter = CurrentAdapter;
-					}
 				}
 			}
 		}
@@ -977,7 +979,12 @@ void FD3D12DynamicRHIModule::FindAdapter()
 	if (bFavorDiscreteAdapter)
 	{
 		// We assume Intel is integrated graphics (slower than discrete) than NVIDIA or AMD cards and rather take a different one
-		if (BestMemoryAdapter.IsValid())
+		if (PreferredAdapter.IsValid())
+		{
+			NewAdapter = TSharedPtr<FD3D12Adapter>(new FWindowsD3D12Adapter(PreferredAdapter));
+			ChosenAdapters.Add(NewAdapter);
+		}
+		else if (BestMemoryAdapter.IsValid())
 		{
 			NewAdapter = TSharedPtr<FD3D12Adapter>(new FWindowsD3D12Adapter(BestMemoryAdapter));
 			ChosenAdapters.Add(NewAdapter);
diff --git a/Engine/Source/Runtime/Windows/D3D11RHI/Private/Windows/WindowsD3D11Device.cpp b/Engine/Source/Runtime/Windows/D3D11RHI/Private/Windows/WindowsD3D11Device.cpp
--- a/Engine/Source/Runtime/Windows/D3D11RHI/Private/Windows/WindowsD3D11Device.cpp
+++ b/Engine/Source/Runtime/Windows/D3D11RHI/Private/Windows/WindowsD3D11Device.cpp
@@ -1006,7 +1006,7 @@ void FD3D11DynamicRHIModule::FindAdapter()
 	int32 CVarExplicitAdapterValue = HmdGraphicsAdapterLuid == 0 ? (CVarGraphicsAdapter ? CVarGraphicsAdapter->GetValueOnGameThread() : -1) : -2;
 	FParse::Value(FCommandLine::Get(), TEXT("graphicsadapter="), CVarExplicitAdapterValue);
 
-	const bool bFavorNonIntegrated = CVarExplicitAdapterValue == -1;
+	const bool bFavorDiscreteAdapter = CVarExplicitAdapterValue == -1;
 
 	TRefCountPtr<IDXGIAdapter> TempAdapter;
 	D3D_FEATURE_LEVEL MinAllowedFeatureLevel = GetMinAllowedD3DFeatureLevel();
@@ -1015,9 +1015,13 @@ void FD3D11DynamicRHIModule::FindAdapter()
 	UE_LOG(LogD3D11RHI, Log, TEXT("D3D11 min allowed feature level: %s"), GetFeatureLevelString(MinAllowedFeatureLevel));
 	UE_LOG(LogD3D11RHI, Log, TEXT("D3D11 max allowed feature level: %s"), GetFeatureLevelString(MaxAllowedFeatureLevel));
 
-	FD3D11Adapter FirstWithoutIntegratedAdapter;
+	FD3D11Adapter PreferredAdapter;
+	FD3D11Adapter BestMemoryAdapter;
+	FD3D11Adapter FirstDiscreteAdapter;
 	FD3D11Adapter FirstAdapter;
 
+	SIZE_T BestDedicatedMemory = 0;
+
 	UE_LOG(LogD3D11RHI, Log, TEXT("D3D11 adapters:"));
 
 	const EGpuVendorId PreferredVendor = RHIGetPreferredAdapterVendor();
@@ -1085,10 +1089,7 @@ void FD3D11DynamicRHIModule::FindAdapter()
 					AdapterDesc.VendorId
 					);
 
-				bool bIsAMD = AdapterDesc.VendorId == 0x1002;
-				bool bIsIntel = AdapterDesc.VendorId == 0x8086;
-				bool bIsNVIDIA = AdapterDesc.VendorId == 0x10DE;
-				bool bIsMicrosoft = AdapterDesc.VendorId == 0x1414;
+				const bool bIsWARP = RHIConvertToGpuVendorId(AdapterDesc.VendorId) == EGpuVendorId::Microsoft;
 
 				// Simple heuristic but without profiling it's hard to do better
 				bool bIsNonLocalMemoryPresent = false;
@@ -1102,18 +1103,17 @@ void FD3D11DynamicRHIModule::FindAdapter()
 
 				// TODO: Using GPUDetect for Intel GPUs to check for integrated vs discrete status, pending GPUDetect update
 
-				const bool bIsSoftware = bIsMicrosoft;
 				const bool bIsIntegrated = !bIsNonLocalMemoryPresent;
 				// PerfHUD is for performance profiling
 				const bool bIsPerfHUD = !FCString::Stricmp(AdapterDesc.Description,TEXT("NVIDIA PerfHUD"));
 
-				FD3D11Adapter CurrentAdapter(TempAdapter, ActualFeatureLevel, bIsSoftware, bIsIntegrated);
+				FD3D11Adapter CurrentAdapter(TempAdapter, ActualFeatureLevel, bIsWARP, bIsIntegrated);
 
 				// Add special check to support HMDs, which do not have associated outputs.
 				// To reject the software emulation, unless the cvar wants it.
 				// https://msdn.microsoft.com/en-us/library/windows/desktop/bb205075(v=vs.85).aspx#WARP_new_for_Win8
 				// Before we tested for no output devices but that failed where a laptop had a Intel (with output) and NVidia (with no output)
-				const bool bSkipSoftwareAdapter = bIsSoftware && !bAllowSoftwareFallback && CVarExplicitAdapterValue < 0 && HmdGraphicsAdapterLuid == 0;
+				const bool bSkipSoftwareAdapter = bIsWARP && !bAllowSoftwareFallback && CVarExplicitAdapterValue < 0 && HmdGraphicsAdapterLuid == 0;
 				
 				// we don't allow the PerfHUD adapter
 				const bool bSkipPerfHUDAdapter = bIsPerfHUD && !bAllowPerfHUD;
@@ -1126,25 +1126,36 @@ void FD3D11DynamicRHIModule::FindAdapter()
 				
 				const bool bSkipAdapter = bSkipSoftwareAdapter || bSkipPerfHUDAdapter || bSkipHmdGraphicsAdapter || bSkipExplicitAdapter;
 
-				if (!bSkipAdapter)
+				if (!bSkipAdapter && CurrentAdapter.IsValid())
 				{
-					if (!bIsIntegrated && !FirstWithoutIntegratedAdapter.IsValid())
+					if (PreferredVendor != EGpuVendorId::Unknown && PreferredVendor == RHIConvertToGpuVendorId(AdapterDesc.VendorId) && !PreferredAdapter.IsValid())
 					{
-						FirstWithoutIntegratedAdapter = CurrentAdapter;
+						PreferredAdapter = CurrentAdapter;
 					}
-					else if ((PreferredVendor != EGpuVendorId::Unknown) && (PreferredVendor == RHIConvertToGpuVendorId(AdapterDesc.VendorId)) && FirstWithoutIntegratedAdapter.IsValid())
+					
+					if (!bIsWARP && !CurrentAdapter.bIsIntegrated)
 					{
-						FirstWithoutIntegratedAdapter = CurrentAdapter;
+						if (!FirstDiscreteAdapter.IsValid())
+						{
+							FirstDiscreteAdapter = CurrentAdapter;
+						}
+
+						if (AdapterDesc.DedicatedVideoMemory > BestDedicatedMemory)
+						{
+							BestMemoryAdapter = CurrentAdapter;
+							BestDedicatedMemory = AdapterDesc.DedicatedVideoMemory;
+							if (PreferredVendor != EGpuVendorId::Unknown && PreferredVendor == RHIConvertToGpuVendorId(AdapterDesc.VendorId))
+							{
+								// Choose the best option of the preferred IHV devices
+								PreferredAdapter = BestMemoryAdapter;
+							}
+						}
 					}
 
 					if (!FirstAdapter.IsValid())
 					{
 						FirstAdapter = CurrentAdapter;
 					}
-					else if ((PreferredVendor != EGpuVendorId::Unknown) && (PreferredVendor == RHIConvertToGpuVendorId(AdapterDesc.VendorId)) && FirstAdapter.IsValid())
-					{
-						FirstAdapter = CurrentAdapter;
-					}
 				}
 			}
 			else
@@ -1158,12 +1169,21 @@ void FD3D11DynamicRHIModule::FindAdapter()
 		}
 	}
 
-	if(bFavorNonIntegrated)
+	if (bFavorDiscreteAdapter)
 	{
-		ChosenAdapter = FirstWithoutIntegratedAdapter;
-
-		// We assume Intel is integrated graphics (slower than discrete) than NVIDIA or AMD cards and rather take a different one
-		if(!ChosenAdapter.IsValid())
+		if (PreferredAdapter.IsValid())
+		{
+			ChosenAdapter = PreferredAdapter;
+		}
+		else if (BestMemoryAdapter.IsValid())
+		{
+			ChosenAdapter = BestMemoryAdapter;
+		}
+		else if (FirstDiscreteAdapter.IsValid())
+		{
+			ChosenAdapter = FirstDiscreteAdapter;
+		}
+		else
 		{
 			ChosenAdapter = FirstAdapter;
 		}
