Fix C4702
#rnx

diff --git a/Engine/Source/Runtime/Windows/AudioMixerXAudio2/Private/AudioMixerPlatformXAudio2.cpp b/Engine/Source/Runtime/Windows/AudioMixerXAudio2/Private/AudioMixerPlatformXAudio2.cpp
--- a/Engine/Source/Runtime/Windows/AudioMixerXAudio2/Private/AudioMixerPlatformXAudio2.cpp
+++ b/Engine/Source/Runtime/Windows/AudioMixerXAudio2/Private/AudioMixerPlatformXAudio2.cpp
@@ -1485,16 +1485,14 @@ namespace Audio
 		{
 			AudioStreamInfo.StreamState = EAudioOutputStreamState::Running;
 			OutputAudioStreamSourceVoice->Start();
-			return true;
 		}
 		else
 		{
 			check(!bIsUsingNullDevice);
 			StartRunningNullDevice();
-			return true;
 		}
 
-		return false;
+		return true;
 	}
 
 	bool FMixerPlatformXAudio2::StopAudioStream()
