Ongoing C4702 warning fixes

diff --git a/Engine/Plugins/Compositing/Composure/Source/Composure/Public/EditorSupport/CompFreezeFrameController.h b/Engine/Plugins/Compositing/Composure/Source/Composure/Public/EditorSupport/CompFreezeFrameController.h
--- a/Engine/Plugins/Compositing/Composure/Source/Composure/Public/EditorSupport/CompFreezeFrameController.h
+++ b/Engine/Plugins/Compositing/Composure/Source/Composure/Public/EditorSupport/CompFreezeFrameController.h
@@ -21,7 +21,7 @@ public:
 		return LockKey.IsValid(); 
 	}
 	
-	FORCEINLINE FFreezeFrameControlHandle Lock()
+	FFreezeFrameControlHandle Lock()
 	{ 
 		if (ensure(!IsLocked()))
 		{
diff --git a/Engine/Source/Runtime/CoreUObject/Public/UObject/UObjectBaseUtility.h b/Engine/Source/Runtime/CoreUObject/Public/UObject/UObjectBaseUtility.h
--- a/Engine/Source/Runtime/CoreUObject/Public/UObject/UObjectBaseUtility.h
+++ b/Engine/Source/Runtime/CoreUObject/Public/UObject/UObjectBaseUtility.h
@@ -873,6 +873,7 @@ public:
 			}
 			return StatID;
 		}
+		return TStatId(); // not doing stats at the moment, or ever
 #elif ENABLE_STATNAMEDEVENTS_UOBJECT
 		const TStatId& StatID = GUObjectArray.IndexToObject(InternalIndex)->StatID;
 		if (!StatID.IsValidStat() && (bForDeferredUse || GCycleStatsShouldEmitNamedEvents))
@@ -880,8 +881,9 @@ public:
 			CreateStatID();
 		}
 		return StatID;
-#endif // STATS
+#else
 		return TStatId(); // not doing stats at the moment, or ever
+#endif // STATS
 	}
 
 private:
diff --git a/Engine/Source/Runtime/Engine/Private/PhysicsEngine/PhysDrawing.cpp b/Engine/Source/Runtime/Engine/Private/PhysicsEngine/PhysDrawing.cpp
--- a/Engine/Source/Runtime/Engine/Private/PhysicsEngine/PhysDrawing.cpp
+++ b/Engine/Source/Runtime/Engine/Private/PhysicsEngine/PhysDrawing.cpp
@@ -1436,10 +1436,7 @@ void FConstraintInstance::GetUsedMaterials(TArray<UMaterialInterface*>& Material
 void FConstraintInstance::DrawConstraintImp(const FPDIOrCollector& PDIOrCollector, float Scale, float LimitDrawScale, bool bDrawLimits, bool bDrawSelected, const FTransform& Con1Frame, const FTransform& Con2Frame, bool bDrawAsPoint) const
 {
 	// Do nothing if we're shipping
-#if UE_BUILD_SHIPPING
-	return;
-#endif
-
+#if !UE_BUILD_SHIPPING
 	const ESceneDepthPriorityGroup Layer = ESceneDepthPriorityGroup::SDPG_World;
 	FPrimitiveDrawInterface* PDI = PDIOrCollector.GetPDI();
 
@@ -1568,4 +1565,5 @@ void FConstraintInstance::DrawConstraintImp(const FPDIOrCollector& PDIOrCollecto
 		bool bLinearZLimited = ( GetLinearZMotion() == LCM_Limited && GetLinearLimit() >= RB_MinSizeToLockDOF );
 		DrawLinearLimit(PDI, Con2Frame.GetTranslation(), Con2Frame.GetScaledAxis( EAxis::Z ), Con2Frame.GetScaledAxis( EAxis::X ), GetLinearLimit(), bLinearZLimited, LimitDrawScale);
 	}
+#endif
 }
diff --git a/Engine/Source/Runtime/Sockets/Private/SocketSubsystem.cpp b/Engine/Source/Runtime/Sockets/Private/SocketSubsystem.cpp
--- a/Engine/Source/Runtime/Sockets/Private/SocketSubsystem.cpp
+++ b/Engine/Source/Runtime/Sockets/Private/SocketSubsystem.cpp
@@ -72,7 +72,9 @@ static IModuleInterface* LoadSubsystemModule(const FString& SubsystemName)
 		return ModuleManager.GetModule(ModuleName);
 	}
 
+#if !UE_BUILD_SHIPPING && !UE_BUILD_SHIPPING_WITH_EDITOR
 	return nullptr;
+#endif
 }
 
 FUniqueSocket ISocketSubsystem::CreateUniqueSocket(const FName& SocketType, const FString& SocketDescription, bool bForceUDP)
diff --git a/Engine/Source/Runtime/UMG/Private/Components/Widget.cpp b/Engine/Source/Runtime/UMG/Private/Components/Widget.cpp
--- a/Engine/Source/Runtime/UMG/Private/Components/Widget.cpp
+++ b/Engine/Source/Runtime/UMG/Private/Components/Widget.cpp
@@ -1663,8 +1663,6 @@ FSizeParam UWidget::ConvertSerializedSizeParamToRuntime(const FSlateChildSize& I
 	case ESlateSizeRule::Fill:
 		return FStretch(Input.Value);
 	}
-
-	return FAuto();
 }
 
 UWidget* UWidget::FindChildContainingDescendant(UWidget* Root, UWidget* Descendant)
diff --git a/Engine/Source/Runtime/Windows/AudioMixerXAudio2/Private/AudioMixerPlatformXAudio2.cpp b/Engine/Source/Runtime/Windows/AudioMixerXAudio2/Private/AudioMixerPlatformXAudio2.cpp
--- a/Engine/Source/Runtime/Windows/AudioMixerXAudio2/Private/AudioMixerPlatformXAudio2.cpp
+++ b/Engine/Source/Runtime/Windows/AudioMixerXAudio2/Private/AudioMixerPlatformXAudio2.cpp
@@ -771,6 +771,7 @@ namespace Audio
 			return false;
 		}
 		OutNumOutputDevices = AllAudioDevices->Size;
+		return true;
 #elif PLATFORM_WINDOWS && XAUDIO_SUPPORTS_DEVICE_DETAILS
 
 		IMMDeviceEnumerator* DeviceEnumerator = nullptr;
@@ -795,8 +796,8 @@ namespace Audio
 		return SUCCEEDED(Result);
 #else
 		OutNumOutputDevices = 1;
-#endif 
 		return true;
+#endif 
 	}
 
 #if PLATFORM_WINDOWS
@@ -1223,6 +1224,7 @@ namespace Audio
 				UE_LOG(LogAudioMixer, Display, TEXT("%d: %s"), i, EAudioMixerChannel::ToString(OutInfo.OutputChannelArray[i]));
 			}
 		}
+		return true;
 #else // #elif PLATFORM_HOLOLENS // #if PLATFORM_WINDOWS
 		OutInfo.bIsSystemDefault = true;
 		OutInfo.SampleRate = 44100;
@@ -1239,8 +1241,8 @@ namespace Audio
 		OutInfo.OutputChannelArray.Add(EAudioMixerChannel::BackRight);
 		OutInfo.OutputChannelArray.Add(EAudioMixerChannel::SideLeft);
 		OutInfo.OutputChannelArray.Add(EAudioMixerChannel::SideRight);
-#endif // #else // #elif PLATFORM_HOLOLENS // #if PLATFORM_WINDOWS
 		return true;
+#endif // #else // #elif PLATFORM_HOLOLENS // #if PLATFORM_WINDOWS
 	}
 
 	bool FMixerPlatformXAudio2::GetDefaultOutputDeviceIndex(uint32& OutDefaultDeviceIndex) const
