Refactoring D3D11  binding to not go through Context methods.
- Writing custom version of UE::RHICore::ApplyStaticUniformBuffers, UE::RHICore::RHISetShaderParametersShared, and UE::RHICore::RHISetShaderUnbindsShared for D3D11RHI to avoid redundant overhead.
- This removed the last uses of RHISetShaderTexture, RHISetShaderSampler, RHISetUAVParameter, RHISetShaderResourceViewParameter, and RHISetShaderParameter from D3D11RHI. Those functions are now gone.

diff --git a/Engine/Source/Runtime/Windows/D3D11RHI/Private/D3D11Commands.cpp b/Engine/Source/Runtime/Windows/D3D11RHI/Private/D3D11Commands.cpp
--- a/Engine/Source/Runtime/Windows/D3D11RHI/Private/D3D11Commands.cpp
+++ b/Engine/Source/Runtime/Windows/D3D11RHI/Private/D3D11Commands.cpp
@@ -84,7 +84,6 @@ static TAutoConsoleVariable<int32> CVarAllowUAVFlushExt(
 	TEXT(" 0: off"),
 	ECVF_RenderThreadSafe);
 
-
 // Vertex state.
 void FD3D11DynamicRHI::RHISetStreamSource(uint32 StreamIndex, FRHIBuffer* VertexBufferRHI, uint32 Offset)
 {
@@ -102,6 +101,33 @@ void FD3D11DynamicRHI::RHISetRasterizerState(FRHIRasterizerState* NewStateRHI)
 	StateCache.SetRasterizerState(NewState->Resource);
 }
 
+template<EShaderFrequency ShaderFrequency>
+void FD3D11DynamicRHI::BindUniformBuffer(uint32 BufferIndex, FRHIUniformBuffer* BufferRHI)
+{
+	check(BufferRHI && BufferRHI->GetLayout().GetHash());
+
+	FD3D11UniformBuffer* Buffer = ResourceCast(BufferRHI);
+
+	ID3D11Buffer* ConstantBuffer = Buffer ? Buffer->Resource.GetReference() : nullptr;
+	StateCache.SetConstantBuffer<ShaderFrequency>(ConstantBuffer, BufferIndex);
+
+	BoundUniformBuffers[ShaderFrequency][BufferIndex] = BufferRHI;
+	DirtyUniformBuffers[ShaderFrequency] |= (1 << BufferIndex);
+}
+
+template <typename TRHIShader>
+void FD3D11DynamicRHI::ApplyStaticUniformBuffers(TRHIShader* Shader)
+{
+	if (Shader)
+	{
+		UE::RHICore::ApplyStaticUniformBuffers(Shader, Shader->StaticSlots, Shader->ShaderResourceTable.ResourceTableLayoutHashes, StaticUniformBuffers,
+			[this](int32 BufferIndex, FRHIUniformBuffer* Buffer)
+			{
+				BindUniformBuffer<static_cast<EShaderFrequency>(TRHIShader::StaticFrequency)>(BufferIndex, Buffer);
+			});
+	}
+}
+
 void FD3D11DynamicRHI::RHISetGraphicsPipelineState(FRHIGraphicsPipelineState* GraphicsState, uint32 StencilRef, bool bApplyAdditionalState)
 {
 	FRHIGraphicsPipelineStateFallBack* FallbackGraphicsState = static_cast<FRHIGraphicsPipelineStateFallBack*>(GraphicsState);
@@ -301,346 +327,195 @@ void FD3D11DynamicRHI::RHISetBoundShaderState(FRHIBoundShaderState* BoundShaderS
 	}
 }
 
-void FD3D11DynamicRHI::RHISetShaderTexture(FRHIGraphicsShader* ShaderRHI,uint32 TextureIndex, FRHITexture* NewTextureRHI)
+void FD3D11DynamicRHI::RHISetStaticUniformBuffers(const FUniformBufferStaticBindings& InUniformBuffers)
 {
-	FD3D11Texture* NewTexture = ResourceCast(NewTextureRHI);
-	ID3D11ShaderResourceView* ShaderResourceView = NewTexture ? NewTexture->GetShaderResourceView() : nullptr;
+	FMemory::Memzero(StaticUniformBuffers.GetData(), StaticUniformBuffers.Num() * sizeof(FRHIUniformBuffer*));
 
-	switch (ShaderRHI->GetFrequency())
-	{
-	case SF_Vertex:
-	{
-		FD3D11VertexShader* VertexShader = static_cast<FD3D11VertexShader*>(ShaderRHI);
-		VALIDATE_BOUND_SHADER(VertexShader);
-		SetShaderResourceView<SF_Vertex>(NewTexture, ShaderResourceView, TextureIndex);
-	}
-	break;
-	case SF_Geometry:
-	{
-		FD3D11GeometryShader* GeometryShader = static_cast<FD3D11GeometryShader*>(ShaderRHI);
-		VALIDATE_BOUND_SHADER(GeometryShader);
-		SetShaderResourceView<SF_Geometry>(NewTexture, ShaderResourceView, TextureIndex);
-	}
-	break;
-	case SF_Pixel:
+	for (int32 Index = 0; Index < InUniformBuffers.GetUniformBufferCount(); ++Index)
 	{
-		FD3D11PixelShader* PixelShader = static_cast<FD3D11PixelShader*>(ShaderRHI);
-		VALIDATE_BOUND_SHADER(PixelShader);
-		SetShaderResourceView<SF_Pixel>(NewTexture, ShaderResourceView, TextureIndex);
-	}
-	break;
-	default:
-		checkf(0, TEXT("Undefined FRHIShader Type %d!"), (int32)ShaderRHI->GetFrequency());
+		StaticUniformBuffers[InUniformBuffers.GetSlot(Index)] = InUniformBuffers.GetUniformBuffer(Index);
 	}
 }
 
-void FD3D11DynamicRHI::RHISetShaderTexture(FRHIComputeShader* ComputeShaderRHI,uint32 TextureIndex, FRHITexture* NewTextureRHI)
+template<EShaderFrequency ShaderFrequency>
+struct FD3D11ResourceBinder
 {
-	//VALIDATE_BOUND_SHADER(ComputeShaderRHI);
-
-	FD3D11Texture* NewTexture = ResourceCast(NewTextureRHI);
-	ID3D11ShaderResourceView* ShaderResourceView = NewTexture ? NewTexture->GetShaderResourceView() : nullptr;
-	SetShaderResourceView<SF_Compute>(NewTexture, ShaderResourceView, TextureIndex);
-}
+	FD3D11DynamicRHI& RHI;
 
+	FD3D11ResourceBinder(FD3D11DynamicRHI& InRHI)
+		: RHI(InRHI)
+	{
+	}
 
-void FD3D11DynamicRHI::RHISetUAVParameter(FRHIPixelShader* ComputeShaderRHI, uint32 UAVIndex, FRHIUnorderedAccessView* UAVRHI)
-{
-	FD3D11UnorderedAccessView* UAV = ResourceCast(UAVRHI);
-
-	if (UAV)
+	void SetUAV(FRHIUnorderedAccessView* InUnorderedAccessView, uint8 Index)
 	{
-		ConditionalClearShaderResource(UAV->GetBaseResource(), true);
-		for (uint32 i = 0; i < D3D11_PS_CS_UAV_REGISTER_COUNT; i++)
+		if (ShaderFrequency == SF_Compute)
 		{
-			if (i != UAVIndex && CurrentUAVs[i] == UAV)
-			{
-				CurrentUAVs[i] = nullptr;
-			}
+			RHI.InternalSetUAVCS(Index, FD3D11DynamicRHI::ResourceCast(InUnorderedAccessView));
+		}
+		else if (ShaderFrequency == SF_Pixel)
+		{
+			RHI.InternalSetUAVPS(Index, FD3D11DynamicRHI::ResourceCast(InUnorderedAccessView));
+		}
+		else
+		{
+			checkf(false, TEXT("UAVs are not supported on vertex and geometry shaders."));
 		}
 	}
 
-	if (CurrentUAVs[UAVIndex] != UAV)
+	void SetSRV(FRHIShaderResourceView* InShaderResourceView, uint8 Index)
 	{
-		CurrentUAVs[UAVIndex] = UAV;
-		UAVSChanged = 1;
-	}
-}
-
-void FD3D11DynamicRHI::RHISetUAVParameter(FRHIComputeShader* ComputeShaderRHI,uint32 UAVIndex, FRHIUnorderedAccessView* UAVRHI)
-{
-	//VALIDATE_BOUND_SHADER(ComputeShaderRHI);
+		FD3D11ShaderResourceView* D3D11ShaderResourceView = FD3D11DynamicRHI::ResourceCast(InShaderResourceView);
+		FD3D11ViewableResource* D3D11ViewableResource = D3D11ShaderResourceView ? D3D11ShaderResourceView->GetBaseResource() : nullptr;
+		ID3D11ShaderResourceView* D3D11SRV = D3D11ShaderResourceView ? D3D11ShaderResourceView->View : nullptr;
 
-	FD3D11UnorderedAccessView* UAV = ResourceCast(UAVRHI);
-
-	if(UAV)
-	{
-		ConditionalClearShaderResource(UAV->GetBaseResource(), true);
+		RHI.SetShaderResourceView<ShaderFrequency>(
+			D3D11ViewableResource,
+			D3D11SRV,
+			Index
+		);
 	}
 
-	ID3D11UnorderedAccessView* D3D11UAV = UAV ? UAV->View : NULL;
-
-	uint32 InitialCount = -1;
-	Direct3DDeviceIMContext->CSSetUnorderedAccessViews(UAVIndex,1,&D3D11UAV, &InitialCount );
-}
-
-void FD3D11DynamicRHI::RHISetUAVParameter(FRHIComputeShader* ComputeShaderRHI,uint32 UAVIndex, FRHIUnorderedAccessView* UAVRHI, uint32 InitialCount )
-{
-	//VALIDATE_BOUND_SHADER(ComputeShaderRHI);
-
-	FD3D11UnorderedAccessView* UAV = ResourceCast(UAVRHI);
-	
-	if(UAV)
+	void SetTexture(FRHITexture* InTexture, uint8 Index)
 	{
-		ConditionalClearShaderResource(UAV->GetBaseResource(), true);
-	}
+		FD3D11Texture* D3D11Texture = FD3D11DynamicRHI::ResourceCast(InTexture);
+		ID3D11ShaderResourceView* ShaderResourceView = D3D11Texture ? D3D11Texture->GetShaderResourceView() : nullptr;
 
-	ID3D11UnorderedAccessView* D3D11UAV = UAV ? UAV->View : NULL;
-	Direct3DDeviceIMContext->CSSetUnorderedAccessViews(UAVIndex,1,&D3D11UAV, &InitialCount );
-}
-
-void FD3D11DynamicRHI::RHISetShaderResourceViewParameter(FRHIGraphicsShader* ShaderRHI,uint32 TextureIndex, FRHIShaderResourceView* SRVRHI)
-{
-	FD3D11ShaderResourceView* SRV = ResourceCast(SRVRHI);
-	FD3D11ViewableResource* Resource = nullptr;
-	ID3D11ShaderResourceView* D3D11SRV = nullptr;
-	if (SRV)
-	{
-		Resource = SRV->GetBaseResource();
-		D3D11SRV = SRV->View;
+		RHI.SetShaderResourceView<ShaderFrequency>(
+			D3D11Texture,
+			ShaderResourceView,
+			Index
+		);
 	}
 
-	switch (ShaderRHI->GetFrequency())
-	{
-	case SF_Vertex:
-	{
-		FD3D11VertexShader* VertexShader = static_cast<FD3D11VertexShader*>(ShaderRHI);
-		VALIDATE_BOUND_SHADER(VertexShader);
-		SetShaderResourceView<SF_Vertex>(Resource, D3D11SRV, TextureIndex);
-	}
-	break;
-	case SF_Geometry:
+	void SetSampler(FRHISamplerState* Sampler, uint8 Index)
 	{
-		FD3D11GeometryShader* GeometryShader = static_cast<FD3D11GeometryShader*>(ShaderRHI);
-		VALIDATE_BOUND_SHADER(GeometryShader);
-		SetShaderResourceView<SF_Geometry>(Resource, D3D11SRV, TextureIndex);
+		RHI.GetStateCache().SetSamplerState<ShaderFrequency>(FD3D11DynamicRHI::ResourceCast(Sampler)->Resource, Index);
 	}
-	break;
-	case SF_Pixel:
-	{
-		FD3D11PixelShader* PixelShader = static_cast<FD3D11PixelShader*>(ShaderRHI);
-		VALIDATE_BOUND_SHADER(PixelShader);
-		SetShaderResourceView<SF_Pixel>(Resource, D3D11SRV, TextureIndex);
-	}
-	break;
-	default:
-		checkf(0, TEXT("Undefined FRHIShader Type %d!"), (int32)ShaderRHI->GetFrequency());
-	}
-}
+};
 
-void FD3D11DynamicRHI::RHISetShaderResourceViewParameter(FRHIComputeShader* ComputeShaderRHI,uint32 TextureIndex, FRHIShaderResourceView* SRVRHI)
+template<EShaderFrequency ShaderFrequency>
+void FD3D11DynamicRHI::SetShaderParametersCommon(FD3D11ConstantBuffer* StageConstantBuffer, TConstArrayView<uint8> InParametersData, TConstArrayView<FRHIShaderParameter> InParameters, TConstArrayView<FRHIShaderParameterResource> InResourceParameters)
 {
-	//VALIDATE_BOUND_SHADER(ComputeShaderRHI);
-
-	FD3D11ShaderResourceView* SRV = ResourceCast(SRVRHI);
-
-	FD3D11ViewableResource* Resource = nullptr;
-	ID3D11ShaderResourceView* D3D11SRV = nullptr;
-	
-	if (SRV)
+	if (InParameters.Num())
 	{
-		Resource = SRV->GetBaseResource();
-		D3D11SRV = SRV->View;
+		for (const FRHIShaderParameter& Parameter : InParameters)
+		{
+			check(Parameter.BufferIndex == 0);
+			StageConstantBuffer->UpdateConstant(&InParametersData[Parameter.ByteOffset], Parameter.BaseIndex, Parameter.ByteSize);
+		}
 	}
 
-	SetShaderResourceView<SF_Compute>(Resource, D3D11SRV, TextureIndex);
-}
+	FD3D11ResourceBinder<ShaderFrequency> Binder(*this);
 
-void FD3D11DynamicRHI::RHISetShaderSampler(FRHIGraphicsShader* ShaderRHI,uint32 SamplerIndex, FRHISamplerState* NewStateRHI)
-{
-	FD3D11SamplerState* NewState = ResourceCast(NewStateRHI);
-	ID3D11SamplerState* StateResource = NewState->Resource;
-	switch (ShaderRHI->GetFrequency())
-	{
-	case SF_Vertex:
-	{
-		FD3D11VertexShader* VertexShader = static_cast<FD3D11VertexShader*>(ShaderRHI);
-		VALIDATE_BOUND_SHADER(VertexShader);
-		StateCache.SetSamplerState<SF_Vertex>(StateResource, SamplerIndex);
-	}
-	break;
-	case SF_Geometry:
+	for (const FRHIShaderParameterResource& Parameter : InResourceParameters)
 	{
-		FD3D11GeometryShader* GeometryShader = static_cast<FD3D11GeometryShader*>(ShaderRHI);
-		VALIDATE_BOUND_SHADER(GeometryShader);
-		StateCache.SetSamplerState<SF_Geometry>(StateResource, SamplerIndex);
+		if (Parameter.Type == FRHIShaderParameterResource::EType::UnorderedAccessView)
+		{
+			Binder.SetUAV(static_cast<FRHIUnorderedAccessView*>(Parameter.Resource), Parameter.Index);
+		}
 	}
-	break;
-	case SF_Pixel:
+
+	for (const FRHIShaderParameterResource& Parameter : InResourceParameters)
 	{
-		FD3D11PixelShader* PixelShader = static_cast<FD3D11PixelShader*>(ShaderRHI);
-		VALIDATE_BOUND_SHADER(PixelShader);
-		StateCache.SetSamplerState<SF_Pixel>(StateResource, SamplerIndex);
-	}
-	break;
-	default:
-		checkf(0, TEXT("Undefined FRHIShader Type %d!"), (int32)ShaderRHI->GetFrequency());
+		switch (Parameter.Type)
+		{
+		case FRHIShaderParameterResource::EType::Texture:
+			Binder.SetTexture(static_cast<FRHITexture*>(Parameter.Resource), Parameter.Index);
+			break;
+		case FRHIShaderParameterResource::EType::ResourceView:
+			Binder.SetSRV(static_cast<FRHIShaderResourceView*>(Parameter.Resource), Parameter.Index);
+			break;
+		case FRHIShaderParameterResource::EType::UnorderedAccessView:
+			break;
+		case FRHIShaderParameterResource::EType::Sampler:
+			Binder.SetSampler(static_cast<FRHISamplerState*>(Parameter.Resource), Parameter.Index);
+			break;
+		case FRHIShaderParameterResource::EType::UniformBuffer:
+			BindUniformBuffer<ShaderFrequency>(Parameter.Index, static_cast<FRHIUniformBuffer*>(Parameter.Resource));
+			break;
+		default:
+			checkf(false, TEXT("Unhandled resource type?"));
+			break;
+		}
 	}
 }
 
-void FD3D11DynamicRHI::RHISetShaderSampler(FRHIComputeShader* ComputeShaderRHI,uint32 SamplerIndex, FRHISamplerState* NewStateRHI)
-{
-	//VALIDATE_BOUND_SHADER(ComputeShaderRHI);
-	FD3D11ComputeShader* ComputeShader = ResourceCast(ComputeShaderRHI);
-	FD3D11SamplerState* NewState = ResourceCast(NewStateRHI);
 
-	ID3D11SamplerState* StateResource = NewState->Resource;
-	StateCache.SetSamplerState<SF_Compute>(StateResource, SamplerIndex);
-}
-
-void FD3D11DynamicRHI::RHISetStaticUniformBuffers(const FUniformBufferStaticBindings& InUniformBuffers)
+void FD3D11DynamicRHI::RHISetShaderParameters(FRHIComputeShader* Shader, TConstArrayView<uint8> InParametersData, TConstArrayView<FRHIShaderParameter> InParameters, TConstArrayView<FRHIShaderParameterResource> InResourceParameters, TConstArrayView<FRHIShaderParameterResource> InBindlessParameters)
 {
-	FMemory::Memzero(StaticUniformBuffers.GetData(), StaticUniformBuffers.Num() * sizeof(FRHIUniformBuffer*));
-
-	for (int32 Index = 0; Index < InUniformBuffers.GetUniformBufferCount(); ++Index)
-	{
-		StaticUniformBuffers[InUniformBuffers.GetSlot(Index)] = InUniformBuffers.GetUniformBuffer(Index);
-	}
+	SetShaderParametersCommon<SF_Compute>(CSConstantBuffer, InParametersData, InParameters, InResourceParameters);
 }
 
-void FD3D11DynamicRHI::RHISetShaderUniformBuffer(FRHIGraphicsShader* ShaderRHI,uint32 BufferIndex, FRHIUniformBuffer* BufferRHI)
+void FD3D11DynamicRHI::RHISetShaderParameters(FRHIGraphicsShader* Shader, TConstArrayView<uint8> InParametersData, TConstArrayView<FRHIShaderParameter> InParameters, TConstArrayView<FRHIShaderParameterResource> InResourceParameters, TConstArrayView<FRHIShaderParameterResource> InBindlessParameters)
 {
-	check(BufferRHI->GetLayout().GetHash());
-	FD3D11UniformBuffer* Buffer = ResourceCast(BufferRHI);
-	ID3D11Buffer* ConstantBuffer = Buffer ? Buffer->Resource.GetReference() : nullptr;
-	EShaderFrequency Stage = SF_NumFrequencies;
-	switch (ShaderRHI->GetFrequency())
+	switch (Shader->GetFrequency())
 	{
 	case SF_Vertex:
-	{
-		FD3D11VertexShader* VertexShader = static_cast<FD3D11VertexShader*>(ShaderRHI);
-		VALIDATE_BOUND_SHADER(VertexShader);
-		Stage = SF_Vertex;
-		StateCache.SetConstantBuffer<SF_Vertex>(ConstantBuffer, BufferIndex);
-	}
-	break;
+		VALIDATE_BOUND_SHADER(static_cast<FRHIVertexShader*>(Shader));
+		SetShaderParametersCommon<SF_Vertex>(VSConstantBuffer, InParametersData, InParameters, InResourceParameters);
+		break;
 	case SF_Geometry:
-	{
-		FD3D11GeometryShader* GeometryShader = static_cast<FD3D11GeometryShader*>(ShaderRHI);
-		VALIDATE_BOUND_SHADER(GeometryShader);
-		Stage = SF_Geometry;
-		StateCache.SetConstantBuffer<SF_Geometry>(ConstantBuffer, BufferIndex);
-	}
-	break;
+		VALIDATE_BOUND_SHADER(static_cast<FRHIGeometryShader*>(Shader));
+		SetShaderParametersCommon<SF_Geometry>(GSConstantBuffer, InParametersData, InParameters, InResourceParameters);
+		break;
 	case SF_Pixel:
-	{
-		FD3D11PixelShader* PixelShader = static_cast<FD3D11PixelShader*>(ShaderRHI);
-		VALIDATE_BOUND_SHADER(PixelShader);
-		Stage = SF_Pixel;
-		StateCache.SetConstantBuffer<SF_Pixel>(ConstantBuffer, BufferIndex);
-	}
-	break;
+		VALIDATE_BOUND_SHADER(static_cast<FRHIPixelShader*>(Shader));
+		SetShaderParametersCommon<SF_Pixel>(PSConstantBuffer, InParametersData, InParameters, InResourceParameters);
+		break;
 	default:
-		checkf(0, TEXT("Undefined FRHIShader Type %d!"), (int32)ShaderRHI->GetFrequency());
-		return;
+		checkf(0, TEXT("Undefined FRHIGraphicsShader Type %d!"), (int32)Shader->GetFrequency());
 	}
-
-	BoundUniformBuffers[Stage][BufferIndex] = BufferRHI;
-	DirtyUniformBuffers[Stage] |= (1 << BufferIndex);
 }
 
-void FD3D11DynamicRHI::RHISetShaderUniformBuffer(FRHIComputeShader* ComputeShader,uint32 BufferIndex, FRHIUniformBuffer* BufferRHI)
+template<EShaderFrequency ShaderFrequency>
+void FD3D11DynamicRHI::SetShaderUnbindsCommon(TConstArrayView<FRHIShaderParameterUnbind> InUnbinds)
 {
-	check(BufferRHI->GetLayout().GetHash());
-	//VALIDATE_BOUND_SHADER(ComputeShader);
-	FD3D11UniformBuffer* Buffer = ResourceCast(BufferRHI);
+	FD3D11ResourceBinder<ShaderFrequency> Binder(*this);
+
+	for (const FRHIShaderParameterUnbind& Unbind : InUnbinds)
 	{
-		ID3D11Buffer* ConstantBuffer = Buffer ? Buffer->Resource.GetReference() : nullptr;
-		StateCache.SetConstantBuffer<SF_Compute>(ConstantBuffer, BufferIndex);
+		switch (Unbind.Type)
+		{
+		case FRHIShaderParameterUnbind::EType::ResourceView:
+			Binder.SetSRV(nullptr, Unbind.Index);
+			break;
+		case FRHIShaderParameterUnbind::EType::UnorderedAccessView:
+			Binder.SetUAV(nullptr, Unbind.Index);
+			break;
+		default:
+			checkf(false, TEXT("Unhandled unbind resource type?"));
+			break;
+		}
 	}
-
-	BoundUniformBuffers[SF_Compute][BufferIndex] = BufferRHI;
-	DirtyUniformBuffers[SF_Compute] |= (1 << BufferIndex);
 }
 
-void FD3D11DynamicRHI::RHISetShaderParameter(FRHIGraphicsShader* ShaderRHI,uint32 BufferIndex,uint32 BaseIndex,uint32 NumBytes,const void* NewValue)
+void FD3D11DynamicRHI::RHISetShaderUnbinds(FRHIComputeShader* Shader, TConstArrayView<FRHIShaderParameterUnbind> InUnbinds)
 {
-	check(BufferIndex == 0);
+	SetShaderUnbindsCommon<SF_Compute>(InUnbinds);
+}
 
-	switch (ShaderRHI->GetFrequency())
+void FD3D11DynamicRHI::RHISetShaderUnbinds(FRHIGraphicsShader* Shader, TConstArrayView<FRHIShaderParameterUnbind> InUnbinds)
+{
+	switch (Shader->GetFrequency())
 	{
 	case SF_Vertex:
-	{
-		FD3D11VertexShader* VertexShader = static_cast<FD3D11VertexShader*>(ShaderRHI);
-		VALIDATE_BOUND_SHADER(VertexShader);
-		checkSlow(VSConstantBuffer);
-		VSConstantBuffer->UpdateConstant((const uint8*)NewValue, BaseIndex, NumBytes);
-	}
-	break;
+		VALIDATE_BOUND_SHADER(static_cast<FRHIVertexShader*>(Shader));
+		SetShaderUnbindsCommon<SF_Vertex>(InUnbinds);
+		break;
 	case SF_Geometry:
-	{
-		FD3D11GeometryShader* GeometryShader = static_cast<FD3D11GeometryShader*>(ShaderRHI);
-		VALIDATE_BOUND_SHADER(GeometryShader);
-		checkSlow(GSConstantBuffer);
-		GSConstantBuffer->UpdateConstant((const uint8*)NewValue, BaseIndex, NumBytes);
-	}
-	break;
+		VALIDATE_BOUND_SHADER(static_cast<FRHIGeometryShader*>(Shader));
+		SetShaderUnbindsCommon<SF_Geometry>(InUnbinds);
+		break;
 	case SF_Pixel:
-	{
-		FD3D11PixelShader* PixelShader = static_cast<FD3D11PixelShader*>(ShaderRHI);
-		VALIDATE_BOUND_SHADER(PixelShader);
-		checkSlow(PSConstantBuffer);
-		PSConstantBuffer->UpdateConstant((const uint8*)NewValue, BaseIndex, NumBytes);
-	}
-	break;
+		VALIDATE_BOUND_SHADER(static_cast<FRHIPixelShader*>(Shader));
+		SetShaderUnbindsCommon<SF_Pixel>(InUnbinds);
+		break;
 	default:
-		checkf(0, TEXT("Undefined FRHIShader Type %d!"), (int32)ShaderRHI->GetFrequency());
+		checkf(0, TEXT("Undefined FRHIGraphicsShader Type %d!"), (int32)Shader->GetFrequency());
 	}
 }
 
-void FD3D11DynamicRHI::RHISetShaderParameter(FRHIComputeShader* ComputeShaderRHI,uint32 BufferIndex,uint32 BaseIndex,uint32 NumBytes,const void* NewValue)
-{
-	//VALIDATE_BOUND_SHADER(ComputeShaderRHI);
-	checkSlow(CSConstantBuffer);
-	CSConstantBuffer->UpdateConstant((const uint8*)NewValue,BaseIndex,NumBytes);
-}
-
-void FD3D11DynamicRHI::RHISetShaderParameters(FRHIComputeShader* Shader, TConstArrayView<uint8> InParametersData, TConstArrayView<FRHIShaderParameter> InParameters, TConstArrayView<FRHIShaderParameterResource> InResourceParameters, TConstArrayView<FRHIShaderParameterResource> InBindlessParameters)
-{
-	UE::RHICore::RHISetShaderParametersShared(
-		*this
-		, Shader
-		, InParametersData
-		, InParameters
-		, InResourceParameters
-		, InBindlessParameters
-		, true /*bBindUAVsFirst*/
-	);
-}
-
-void FD3D11DynamicRHI::RHISetShaderParameters(FRHIGraphicsShader* Shader, TConstArrayView<uint8> InParametersData, TConstArrayView<FRHIShaderParameter> InParameters, TConstArrayView<FRHIShaderParameterResource> InResourceParameters, TConstArrayView<FRHIShaderParameterResource> InBindlessParameters)
-{
-	UE::RHICore::RHISetShaderParametersShared(
-		*this
-		, Shader
-		, InParametersData
-		, InParameters
-		, InResourceParameters
-		, InBindlessParameters
-		, true /*bBindUAVsFirst*/
-	);
-}
-
-void FD3D11DynamicRHI::RHISetShaderUnbinds(FRHIComputeShader* Shader, TConstArrayView<FRHIShaderParameterUnbind> InUnbinds)
-{
-	UE::RHICore::RHISetShaderUnbindsShared(*this, Shader, InUnbinds);
-}
-
-void FD3D11DynamicRHI::RHISetShaderUnbinds(FRHIGraphicsShader* Shader, TConstArrayView<FRHIShaderParameterUnbind> InUnbinds)
-{
-	UE::RHICore::RHISetShaderUnbindsShared(*this, Shader, InUnbinds);
-}
-
 void FD3D11DynamicRHI::ValidateExclusiveDepthStencilAccess(FExclusiveDepthStencil RequestedAccess) const
 {
 	const bool bSrcDepthWrite = RequestedAccess.IsDepthWrite();
@@ -735,20 +610,36 @@ void FD3D11DynamicRHI::CommitRenderTargets(bool bClearUAVs)
 	}
 }
 
+void FD3D11DynamicRHI::InternalSetUAVCS(uint32 BindIndex, FD3D11UnorderedAccessView* UnorderedAccessViewRHI)
+{
+	if (UnorderedAccessViewRHI)
+	{
+		ConditionalClearShaderResource(UnorderedAccessViewRHI->GetBaseResource(), true);
+	}
+
+	ID3D11UnorderedAccessView* D3D11UAV = UnorderedAccessViewRHI ? UnorderedAccessViewRHI->View : nullptr;
+
+	uint32 InitialCount = -1;
+	Direct3DDeviceIMContext->CSSetUnorderedAccessViews(BindIndex, 1, &D3D11UAV, &InitialCount);
+}
+
 void FD3D11DynamicRHI::InternalSetUAVPS(uint32 BindIndex, FD3D11UnorderedAccessView* UnorderedAccessViewRHI)
 {
 	check(BindIndex < D3D11_PS_CS_UAV_REGISTER_COUNT);
-	if(CurrentUAVs[BindIndex] != UnorderedAccessViewRHI)
+	if (CurrentUAVs[BindIndex] != UnorderedAccessViewRHI)
 	{
 		CurrentUAVs[BindIndex] = UnorderedAccessViewRHI;
 		UAVSChanged = 1;
 	}
-	ConditionalClearShaderResource(UnorderedAccessViewRHI->GetBaseResource(), true);
-	for (uint32 i = 0; i < D3D11_PS_CS_UAV_REGISTER_COUNT; i++)
+	if (UnorderedAccessViewRHI)
 	{
-		if (i != BindIndex && CurrentUAVs[i] == UnorderedAccessViewRHI)
+		ConditionalClearShaderResource(UnorderedAccessViewRHI->GetBaseResource(), true);
+		for (uint32 i = 0; i < D3D11_PS_CS_UAV_REGISTER_COUNT; i++)
 		{
-			CurrentUAVs[i] = nullptr;
+			if (i != BindIndex && CurrentUAVs[i] == UnorderedAccessViewRHI)
+			{
+				CurrentUAVs[i] = nullptr;
+			}
 		}
 	}
 }
@@ -1114,43 +1005,8 @@ void FD3D11DynamicRHI::SetResourcesFromTables(const ShaderType* RESTRICT Shader)
 	checkSlow(Shader);
 	static constexpr EShaderFrequency Frequency = static_cast<EShaderFrequency>(ShaderType::StaticFrequency);
 
-	struct FUniformResourceBinder
-	{
-		FD3D11DynamicRHI& RHI;
-
-		void SetUAV(FRHIUnorderedAccessView* UAV, uint8 Index)
-		{
-			RHI.InternalSetUAVPS(Index, ResourceCast(UAV));
-		}
-
-		void SetSRV(FRHIShaderResourceView* SRV, uint8 Index)
-		{
-			FD3D11ShaderResourceView* SRVRHI = ResourceCast(SRV);
-			RHI.SetShaderResourceView<Frequency>(
-				SRVRHI->GetBaseResource(),
-				SRVRHI->View,
-				Index
-			);
-		}
-
-		void SetTexture(FRHITexture* Texture, uint8 Index)
-		{
-			FD3D11Texture* TextureD3D11 = ResourceCast(Texture);
-			RHI.SetShaderResourceView<Frequency>(
-				TextureD3D11,
-				TextureD3D11->GetShaderResourceView(),
-				Index
-			);
-		}
-
-		void SetSampler(FRHISamplerState* Sampler, uint8 Index)
-		{
-			RHI.StateCache.SetSamplerState<Frequency>(ResourceCast(Sampler)->Resource, Index);
-		}
-	};
-
 	UE::RHICore::SetResourcesFromTables(
-		  FUniformResourceBinder { *this }
+		  FD3D11ResourceBinder<Frequency> { *this }
 		, *Shader
 		, Shader->ShaderResourceTable
 		, DirtyUniformBuffers[Frequency]
diff --git a/Engine/Source/Runtime/Windows/D3D11RHI/Private/D3D11RHIPrivate.h b/Engine/Source/Runtime/Windows/D3D11RHI/Private/D3D11RHIPrivate.h
--- a/Engine/Source/Runtime/Windows/D3D11RHI/Private/D3D11RHIPrivate.h
+++ b/Engine/Source/Runtime/Windows/D3D11RHI/Private/D3D11RHIPrivate.h
@@ -461,6 +461,18 @@ public:
 		return Result;
 	}
 
+	template<EShaderFrequency ShaderFrequency>
+	void BindUniformBuffer(uint32 BufferIndex, FRHIUniformBuffer* BufferRHI);
+
+	template <typename TRHIShader>
+	void ApplyStaticUniformBuffers(TRHIShader* Shader);
+
+	template<EShaderFrequency ShaderFrequency>
+	void SetShaderParametersCommon(FD3D11ConstantBuffer* StageConstantBuffer, TConstArrayView<uint8> InParametersData, TConstArrayView<FRHIShaderParameter> InParameters, TConstArrayView<FRHIShaderParameterResource> InResourceParameters);
+
+	template<EShaderFrequency ShaderFrequency>
+	void SetShaderUnbindsCommon(TConstArrayView<FRHIShaderParameterUnbind> InUnbinds);
+
 	/**
 	 * Reads a D3D query's data into the provided buffer.
 	 * @param Query - The D3D query to read data from.
@@ -587,20 +599,7 @@ public:
 	virtual void RHISetScissorRect(bool bEnable, uint32 MinX, uint32 MinY, uint32 MaxX, uint32 MaxY) final override;
 	virtual void RHISetBoundShaderState(FRHIBoundShaderState* BoundShaderState) final override;
 	virtual void RHISetGraphicsPipelineState(FRHIGraphicsPipelineState* GraphicsState, uint32 StencilRef, bool bApplyAdditionalState) final override;
-	void RHISetShaderTexture(FRHIGraphicsShader* Shader, uint32 TextureIndex, FRHITexture* NewTexture);
-	void RHISetShaderTexture(FRHIComputeShader* PixelShader, uint32 TextureIndex, FRHITexture* NewTexture);
-	void RHISetShaderSampler(FRHIComputeShader* ComputeShader, uint32 SamplerIndex, FRHISamplerState* NewState);
-	void RHISetShaderSampler(FRHIGraphicsShader* Shader, uint32 SamplerIndex, FRHISamplerState* NewState);
-	void RHISetUAVParameter(FRHIPixelShader* PixelShader, uint32 UAVIndex, FRHIUnorderedAccessView* UAVRHI);
-	void RHISetUAVParameter(FRHIComputeShader* ComputeShader, uint32 UAVIndex, FRHIUnorderedAccessView* UAV);
-	void RHISetUAVParameter(FRHIComputeShader* ComputeShader, uint32 UAVIndex, FRHIUnorderedAccessView* UAV, uint32 InitialCount);
-	void RHISetShaderResourceViewParameter(FRHIGraphicsShader* Shader, uint32 SamplerIndex, FRHIShaderResourceView* SRV);
-	void RHISetShaderResourceViewParameter(FRHIComputeShader* ComputeShader, uint32 SamplerIndex, FRHIShaderResourceView* SRV);
 	virtual void RHISetStaticUniformBuffers(const FUniformBufferStaticBindings& InUniformBuffers) final override;
-	void RHISetShaderUniformBuffer(FRHIGraphicsShader* Shader, uint32 BufferIndex, FRHIUniformBuffer* Buffer);
-	void RHISetShaderUniformBuffer(FRHIComputeShader* ComputeShader, uint32 BufferIndex, FRHIUniformBuffer* Buffer);
-	void RHISetShaderParameter(FRHIGraphicsShader* Shader, uint32 BufferIndex, uint32 BaseIndex, uint32 NumBytes, const void* NewValue);
-	void RHISetShaderParameter(FRHIComputeShader* ComputeShader, uint32 BufferIndex, uint32 BaseIndex, uint32 NumBytes, const void* NewValue);
 	virtual void RHISetShaderParameters(FRHIComputeShader* Shader, TConstArrayView<uint8> InParametersData, TConstArrayView<FRHIShaderParameter> InParameters, TConstArrayView<FRHIShaderParameterResource> InResourceParameters, TConstArrayView<FRHIShaderParameterResource> InBindlessParameters) final override;
 	virtual void RHISetShaderParameters(FRHIGraphicsShader* Shader, TConstArrayView<uint8> InParametersData, TConstArrayView<FRHIShaderParameter> InParameters, TConstArrayView<FRHIShaderParameterResource> InResourceParameters, TConstArrayView<FRHIShaderParameterResource> InBindlessParameters) final override;
 	virtual void RHISetShaderUnbinds(FRHIComputeShader* Shader, TConstArrayView<FRHIShaderParameterUnbind> InUnbinds) final override;
@@ -610,6 +609,7 @@ public:
 	virtual void RHISetBlendState(FRHIBlendState* NewState, const FLinearColor& BlendFactor) final override;
 	virtual void RHISetBlendFactor(const FLinearColor& BlendFactor) final override;
 	void SetRenderTargets(uint32 NumSimultaneousRenderTargets, const FRHIRenderTargetView* NewRenderTargets, const FRHIDepthRenderTargetView* NewDepthStencilTarget);
+	void InternalSetUAVCS(uint32 BindIndex, FD3D11UnorderedAccessView* UnorderedAccessViewRHI);
 	void InternalSetUAVPS(uint32 BindIndex, FD3D11UnorderedAccessView* UnorderedAccessViewRHI);
 	void SetRenderTargetsAndClear(const FRHISetRenderTargetsInfo& RenderTargetsInfo);
 	virtual void RHIDrawPrimitive(uint32 BaseVertexIndex, uint32 NumPrimitives, uint32 NumInstances) final override;
@@ -782,6 +782,8 @@ public:
 
 	EPixelFormat GetDisplayFormat(EPixelFormat InPixelFormat) const;
 
+	FD3D11StateCache& GetStateCache() { return StateCache; }
+
 protected:
 	/** The global D3D interface. */
 	TRefCountPtr<IDXGIFactory1> DXGIFactory1;
@@ -888,15 +890,6 @@ protected:
 	/** Bit array to track which uniform buffers have changed since the last draw call. */
 	uint16 DirtyUniformBuffers[SF_NumStandardFrequencies];
 
-	template <typename TRHIShader>
-	void ApplyStaticUniformBuffers(TRHIShader* Shader)
-	{
-		if (Shader)
-		{
-			UE::RHICore::ApplyStaticUniformBuffers(this, Shader, Shader->StaticSlots, Shader->ShaderResourceTable.ResourceTableLayoutHashes, StaticUniformBuffers);
-		}
-	}
-
 	TArray<FRHIUniformBuffer*> StaticUniformBuffers;
 
 	/** Tracks the current depth stencil access type. */
