Fixing IDXGISwapchainProvider.h single file compile errors by adding the necessary DXGI includes.

diff --git a/Engine/Source/Runtime/D3D12RHI/Private/Windows/WindowsD3D12Viewport.cpp b/Engine/Source/Runtime/D3D12RHI/Private/Windows/WindowsD3D12Viewport.cpp
--- a/Engine/Source/Runtime/D3D12RHI/Private/Windows/WindowsD3D12Viewport.cpp
+++ b/Engine/Source/Runtime/D3D12RHI/Private/Windows/WindowsD3D12Viewport.cpp
@@ -5,16 +5,12 @@
 =============================================================================*/
 
 #include "D3D12RHIPrivate.h"
-#include "RenderCore.h"
+#include "Features/IModularFeatures.h"
 #include "HDRHelper.h"
-
-#include "Windows/AllowWindowsPlatformTypes.h"
-#include "Windows.h"
-
 #include "HAL/ThreadHeartBeat.h"
-
 #include "Windows/IDXGISwapchainProvider.h"
-#include "Features/IModularFeatures.h"
+
+#include "Windows/AllowWindowsPlatformTypes.h"
 
 static const uint32 WindowsDefaultNumBackBuffers = 3;
 
diff --git a/Engine/Source/Runtime/RHI/Public/Windows/IDXGISwapchainProvider.h b/Engine/Source/Runtime/RHI/Public/Windows/IDXGISwapchainProvider.h
--- a/Engine/Source/Runtime/RHI/Public/Windows/IDXGISwapchainProvider.h
+++ b/Engine/Source/Runtime/RHI/Public/Windows/IDXGISwapchainProvider.h
@@ -2,10 +2,17 @@
 
 #pragma once
 
+#include "HAL/Platform.h"
 #include "Features/IModularFeature.h"
 #include "RHIDefinitions.h"
 #include "UObject/NameTypes.h"
 
+#include "Microsoft/AllowMicrosoftPlatformTypes.h"
+THIRD_PARTY_INCLUDES_START
+	#include <dxgi1_6.h>
+THIRD_PARTY_INCLUDES_END
+#include "Microsoft/HideMicrosoftPlatformTypes.h"
+
 class IDXGISwapchainProvider : public IModularFeature
 {
 public:
diff --git a/Engine/Source/Runtime/Windows/D3D11RHI/Private/Windows/WindowsD3D11Viewport.cpp b/Engine/Source/Runtime/Windows/D3D11RHI/Private/Windows/WindowsD3D11Viewport.cpp
--- a/Engine/Source/Runtime/Windows/D3D11RHI/Private/Windows/WindowsD3D11Viewport.cpp
+++ b/Engine/Source/Runtime/Windows/D3D11RHI/Private/Windows/WindowsD3D11Viewport.cpp
@@ -5,15 +5,13 @@
 =============================================================================*/
 
 #include "D3D11RHIPrivate.h"
-#include "RenderCore.h"
 #include "Misc/CommandLine.h"
+#include "Windows/IDXGISwapchainProvider.h"
+#include "Features/IModularFeatures.h"
 
 #include "Windows/AllowWindowsPlatformTypes.h"
 #include <dwmapi.h>
 
-#include "Windows/IDXGISwapchainProvider.h"
-#include "Features/IModularFeatures.h"
-
 #ifndef DXGI_SWAP_CHAIN_FLAG_ALLOW_TEARING
 #define DXGI_SWAP_CHAIN_FLAG_ALLOW_TEARING  2048
 #endif
@@ -455,4 +453,4 @@ void FD3D11DynamicRHI::RHIGetDisplaysInformation(FDisplayInformationArray& OutDi
 	OutDisplayInformation.Append(DisplayList);
 }
 
-#include "Windows/HideWindowsPlatformTypes.h"
\ No newline at end of file
+#include "Windows/HideWindowsPlatformTypes.h"
