Make D3D11 query timeout thresholds configurable, increase the default timeout, and clean up timeout logic.

The VerifyD3D11Result logic already takes care of logging and checking the device removal reason.

diff --git a/Engine/Source/Runtime/Windows/D3D11RHI/Private/D3D11Query.cpp b/Engine/Source/Runtime/Windows/D3D11RHI/Private/D3D11Query.cpp
--- a/Engine/Source/Runtime/Windows/D3D11RHI/Private/D3D11Query.cpp
+++ b/Engine/Source/Runtime/Windows/D3D11RHI/Private/D3D11Query.cpp
@@ -7,6 +7,22 @@
 #include "D3D11RHIPrivate.h"
 #include "RenderCore.h"
 
+float GD3D11AbsoluteTimeQueryTimeoutValue = 30.0f;
+static FAutoConsoleVariableRef CVarD3D11AbsoluteTimeQueryTimeoutValue(
+	TEXT("r.D3D11.AbsoluteTimeQueryTimeoutValue"),
+	GD3D11AbsoluteTimeQueryTimeoutValue,
+	TEXT("Set the timeout value, in seconds, to wait for a D3D11 absolute time query."),
+	ECVF_Default
+);
+
+float GD3D11QueryTimeoutValue = 5.0f;
+static FAutoConsoleVariableRef CVarD3D11QueryTimeoutValue(
+	TEXT("r.D3D11.QueryTimeoutValue"),
+	GD3D11QueryTimeoutValue,
+	TEXT("Set the timeout value, in seconds, to wait for a D3D11 query. This value does not apply to absolute time queries (which are controlled by r.D3D11.AbsoluteTimeQueryTimeoutValue)."),
+	ECVF_Default
+);
+
 class FD3D11RenderQueryBatcher final
 {
 public:
@@ -291,7 +307,6 @@ bool FD3D11DynamicRHI::GetQueryData(ID3D11Query* Query, void* Data, SIZE_T DataS
 	HRESULT Result;
 	SAFE_GET_QUERY_DATA
 
-
 	// Isn't the query finished yet, and can we wait for it?
 	if ( Result == S_FALSE && bWait )
 	{
@@ -301,7 +316,7 @@ bool FD3D11DynamicRHI::GetQueryData(ID3D11Query* Query, void* Data, SIZE_T DataS
 		double StartTime = FPlatformTime::Seconds();
 		double TimeoutWarningLimit = 5.0;
 		// timer queries are used for Benchmarks which can stall a bit more
-		double TimeoutValue = (QueryType == RQT_AbsoluteTime) ? 30.0 : 0.5;
+		double TimeoutValue = (QueryType == RQT_AbsoluteTime) ? GD3D11AbsoluteTimeQueryTimeoutValue : GD3D11QueryTimeoutValue;
 
 		do
 		{
@@ -312,8 +327,6 @@ bool FD3D11DynamicRHI::GetQueryData(ID3D11Query* Query, void* Data, SIZE_T DataS
 				return true;
 			}
 
-
-
 			float DeltaTime = FPlatformTime::Seconds() - StartTime;
 			if(DeltaTime > TimeoutWarningLimit)
 			{
@@ -324,17 +337,11 @@ bool FD3D11DynamicRHI::GetQueryData(ID3D11Query* Query, void* Data, SIZE_T DataS
 
 			if(DeltaTime > TimeoutValue)
 			{
-				HRESULT DeviceRemovedReason = Direct3DDevice->GetDeviceRemovedReason();
-				UE_LOG(LogD3D11RHI, Log, TEXT("Timed out while waiting for GPU to catch up. (%.1f s) (ErrorCode %08x) (%08x)"), TimeoutValue, (uint32)Result, (uint32)DeviceRemovedReason);
-				if(FAILED(Result))
-				{
-					VERIFYD3D11RESULT_EX(Result, Direct3DDevice);
-				}
-				else if (QueryType == RQT_AbsoluteTime)
-				{
-					UE_LOG(LogD3D11RHI, Log, TEXT("GPU has hung or crashed, checking status."));
-					GPUProfilingData.CheckGpuHeartbeat(true);
-				}
+				UE_LOG(LogD3D11RHI, Log, TEXT("Timed out while waiting for GPU query. (Timeout %.1f s) (ErrorCode %08x)"), TimeoutValue, (uint32)Result);
+
+				// Check for GPU status and crash.
+				GPUProfilingData.CheckGpuHeartbeat(true);
+				VERIFYD3D11RESULT_EX(Result, Direct3DDevice);
 				return false;
 			}
 		} while ( Result == S_FALSE );
