Remove PLATFORM_HOLOLENS from Audio

diff --git a/Engine/Source/Runtime/AudioMixerCore/Private/AudioMixer.cpp b/Engine/Source/Runtime/AudioMixerCore/Private/AudioMixer.cpp
--- a/Engine/Source/Runtime/AudioMixerCore/Private/AudioMixer.cpp
+++ b/Engine/Source/Runtime/AudioMixerCore/Private/AudioMixer.cpp
@@ -906,7 +906,7 @@ namespace Audio
 
 	bool IAudioMixer::ShouldUseDeviceInfoCache()
 	{		
-#if PLATFORM_WINDOWS // PLATFORM_HOLOLENS uses old path.
+#if PLATFORM_WINDOWS 
 		return bUseAudioDeviceInfoCacheCVar != 0;
 #else //PLATFORM_WINDOWS
 		return false;
diff --git a/Engine/Source/Runtime/VorbisAudioDecoder/Module/Private/VorbisAudioInfo.cpp b/Engine/Source/Runtime/VorbisAudioDecoder/Module/Private/VorbisAudioInfo.cpp
--- a/Engine/Source/Runtime/VorbisAudioDecoder/Module/Private/VorbisAudioInfo.cpp
+++ b/Engine/Source/Runtime/VorbisAudioDecoder/Module/Private/VorbisAudioInfo.cpp
@@ -795,9 +795,6 @@ void LoadVorbisLibraries()
 		PlatformString = TEXT("Win64");
 		DLLNameStub = TEXT("_64.dll");
 #endif
-#if PLATFORM_HOLOLENS
-		PlatformString = TEXT("HoloLens");
-#endif
 
 #if PLATFORM_CPU_ARM_FAMILY && !defined(_M_ARM64EC)
 #if PLATFORM_64BITS
diff --git a/Engine/Source/Runtime/Windows/AudioMixerXAudio2/Private/AudioMixerPlatformXAudio2.cpp b/Engine/Source/Runtime/Windows/AudioMixerXAudio2/Private/AudioMixerPlatformXAudio2.cpp
--- a/Engine/Source/Runtime/Windows/AudioMixerXAudio2/Private/AudioMixerPlatformXAudio2.cpp
+++ b/Engine/Source/Runtime/Windows/AudioMixerXAudio2/Private/AudioMixerPlatformXAudio2.cpp
@@ -86,7 +86,7 @@ static XAUDIO2_PROCESSOR GetXAudio2ProcessorsToUse()
 	return ProcessorsToUse;
 }
 
-#if PLATFORM_WINDOWS || PLATFORM_HOLOLENS
+#if PLATFORM_WINDOWS 
 FName GetDllName(FName Current = NAME_None) 
 {
 #if PLATFORM_64BITS
@@ -96,8 +96,7 @@ FName GetDllName(FName Current = NAME_None)
 #endif
 	return *XAudio2_9Redist;
 }
-#endif //#if PLATFORM_WINDOWS || PLATFORM_HOLOLENS
-
+#endif //#if PLATFORM_WINDOWS 
 
 /*
 	Whether or not to enable xaudio2 debugging mode
@@ -111,10 +110,6 @@ FName GetDllName(FName Current = NAME_None)
 
 namespace Audio
 {
-#if PLATFORM_HOLOLENS
-	static Windows::Devices::Enumeration::DeviceInformationCollection^ AllAudioDevices = nullptr;
-#endif
-
 	void FXAudio2VoiceCallback::OnBufferEnd(void* BufferContext)
 	{
 		SCOPED_NAMED_EVENT(FXAudio2VoiceCallback_OnBufferEnd, FColor::Blue);
@@ -159,17 +154,17 @@ namespace Audio
 		, bIsDeviceOpen(false)
 		, bIsSuspended(false)
 	{
-#if PLATFORM_WINDOWS || PLATFORM_HOLOLENS
+#if PLATFORM_WINDOWS 
 		FPlatformMisc::CoInitialize();
 		DllName = GetDllName();
-#endif // #if PLATFORM_WINDOWS || PLATFORM_HOLOLENS
+#endif // #if PLATFORM_WINDOWS 
 	}
 
 	FMixerPlatformXAudio2::~FMixerPlatformXAudio2()
 	{
-#if PLATFORM_WINDOWS || PLATFORM_HOLOLENS
+#if PLATFORM_WINDOWS 
 		FPlatformMisc::CoUninitialize();
-#endif // #if PLATFORM_WINDOWS || PLATFORM_HOLOLENS
+#endif // #if PLATFORM_WINDOWS 
 	}
 
 #if PLATFORM_WINDOWS
@@ -640,22 +635,6 @@ namespace Audio
 			FMessageDialog::Open(EAppMsgType::Ok, NSLOCTEXT("Audio", "XAudio2Error", "Failed to initialize audio. This may be an issue with your installation of XAudio 2.7. XAudio2 is available in the DirectX End-User Runtime (June 2010)."));
 			return false;
 		}
-
-#if PLATFORM_HOLOLENS
-		using namespace Windows::Foundation;
-		using namespace Windows::Devices::Enumeration;
-		IAsyncOperation<DeviceInformationCollection^>^ EnumerationOp = DeviceInformation::FindAllAsync(DeviceClass::AudioRender);
-		while (EnumerationOp->Status == AsyncStatus::Started)
-		{
-			// Spin
-		}
-
-		if (EnumerationOp->Status == AsyncStatus::Completed)
-		{
-			AllAudioDevices = EnumerationOp->GetResults();
-		}
-#endif
-
 #if XAUDIO2_DEBUG_ENABLED
 		XAUDIO2_DEBUG_CONFIGURATION DebugConfiguration = { 0 };
 		DebugConfiguration.TraceMask = XAUDIO2_LOG_ERRORS | XAUDIO2_LOG_WARNINGS;
@@ -753,14 +732,7 @@ namespace Audio
 
 		// XAudio2 for HoloLens doesn't have GetDeviceCount, use Windows::Devices::Enumeration instead
 		// See https://blogs.msdn.microsoft.com/chuckw/2012/04/02/xaudio2-and-windows-8/
-#if PLATFORM_HOLOLENS
-		if (!AllAudioDevices)
-		{
-			return false;
-		}
-		OutNumOutputDevices = AllAudioDevices->Size;
-		return true;
-#elif PLATFORM_WINDOWS && XAUDIO_SUPPORTS_DEVICE_DETAILS
+#if PLATFORM_WINDOWS && XAUDIO_SUPPORTS_DEVICE_DETAILS
 
 		IMMDeviceEnumerator* DeviceEnumerator = nullptr;
 		IMMDeviceCollection* DeviceCollection = nullptr;
@@ -1052,168 +1024,7 @@ namespace Audio
 
 		return bSucceeded && SUCCEEDED(Result);
 
-#elif PLATFORM_HOLOLENS // #if PLATFORM_WINDOWS
-
-		// XAudio2 for HoloLens doesn't have GetDeviceDetails, use Windows::Devices::Enumeration instead
-		// See https://blogs.msdn.microsoft.com/chuckw/2012/04/02/xaudio2-and-windows-8/
-		if (!AllAudioDevices)
-		{
-			return false;
-		}
-
-		Windows::Devices::Enumeration::DeviceInformation^ WindowsDeviceInfo = AllAudioDevices->GetAt(InDeviceIndex);
-		OutInfo.Name = WindowsDeviceInfo->Name->Data();
-		OutInfo.bIsSystemDefault = WindowsDeviceInfo->IsDefault;
-
-		// No direct equivalent of OutputFormat.  If we have a voice already we can assemble what
-		// we need from methods on that.  But what to do if we don't?
-		WAVEFORMATEXTENSIBLE FakeWaveFormatExtensible;
-		XAUDIO2_VOICE_DETAILS VoiceDetails;
-		if (OutputAudioStreamMasteringVoice && InDeviceIndex == AudioStreamInfo.OutputDeviceIndex)
-		{
-			OutputAudioStreamMasteringVoice->GetVoiceDetails(&VoiceDetails);
-			OutputAudioStreamMasteringVoice->GetChannelMask(&FakeWaveFormatExtensible.dwChannelMask);
-		}
-		else if (!OutputAudioStreamMasteringVoice)
-		{
-			// If we don't yet have a mastering voice we can create a temporary one with asking for the default channels
-			// and sample rate, and then query that to discover the device's preferred format. 
-			IXAudio2MasteringVoice* TempMasteringVoice;
-			if (SUCCEEDED(XAudio2System->CreateMasteringVoice(&TempMasteringVoice, XAUDIO2_DEFAULT_CHANNELS, XAUDIO2_DEFAULT_SAMPLERATE, 0, AllAudioDevices->GetAt(InDeviceIndex)->Id->Data(), nullptr)))
-			{
-				TempMasteringVoice->GetVoiceDetails(&VoiceDetails);
-				TempMasteringVoice->GetChannelMask(&FakeWaveFormatExtensible.dwChannelMask);
-				TempMasteringVoice->DestroyVoice();
-			}
-			else
-			{
-				return false;
-			}
-		}
-		else
-		{
-			// Can't create multiple mastering voices, so currently we can't report
-			// device preferred format in this scenario.
-			return false;
-		}
-
-		WAVEFORMATEX& WaveFormatEx = *(WAVEFORMATEX*)&FakeWaveFormatExtensible;
-		WaveFormatEx.wFormatTag = WAVE_FORMAT_EXTENSIBLE;
-		WaveFormatEx.nSamplesPerSec = VoiceDetails.InputSampleRate;
-		WaveFormatEx.nChannels = VoiceDetails.InputChannels;
-
-		OutInfo.SampleRate = WaveFormatEx.nSamplesPerSec;
-
-		OutInfo.NumChannels = FMath::Clamp((int32)WaveFormatEx.nChannels, 2, 8);
-
-		// XAudio2 automatically converts the audio format to output device us so we don't need to do any format conversions
-		OutInfo.Format = EAudioMixerStreamDataFormat::Float;
-
-		OutInfo.OutputChannelArray.Reset();
-
-		// Extensible format supports surround sound so we need to parse the channel configuration to build our channel output array
-		if (WaveFormatEx.wFormatTag == WAVE_FORMAT_EXTENSIBLE)
-		{
-			// Cast to the extensible format to get access to extensible data
-			const WAVEFORMATEXTENSIBLE* WaveFormatExtensible = (WAVEFORMATEXTENSIBLE*)&WaveFormatEx;
-
-			// Loop through the extensible format channel flags in the standard order and build our output channel array
-			// From https://msdn.microsoft.com/en-us/library/windows/hardware/dn653308(v=vs.85).aspx
-			// The channels in the interleaved stream corresponding to these spatial positions must appear in the order specified above. This holds true even in the 
-			// case of a non-contiguous subset of channels. For example, if a stream contains left, bass enhance and right, then channel 1 is left, channel 2 is right, 
-			// and channel 3 is bass enhance. This enables the linkage of multi-channel streams to well-defined multi-speaker configurations.
-
-			check(EAudioMixerChannel::ChannelTypeCount == UE_ARRAY_COUNT(ChannelTypeMap));
-			uint32 ChanCount = 0;
-			for (uint32 ChannelTypeIndex = 0; ChannelTypeIndex < EAudioMixerChannel::ChannelTypeCount && ChanCount < (uint32)OutInfo.NumChannels; ++ChannelTypeIndex)
-			{
-				if (WaveFormatExtensible->dwChannelMask & ChannelTypeMap[ChannelTypeIndex])
-				{
-					OutInfo.OutputChannelArray.Add((EAudioMixerChannel::Type)ChannelTypeIndex);
-					++ChanCount;
-				}
-			}
-
-			// We didnt match channel masks for all channels, revert to a default ordering
-			if (ChanCount < (uint32)OutInfo.NumChannels)
-			{
-				UE_LOG(LogAudioMixer, Warning, TEXT("Did not find the channel type flags for audio device '%s'. Reverting to a default channel ordering."), *OutInfo.Name);
-
-				OutInfo.OutputChannelArray.Reset();
-
-				static EAudioMixerChannel::Type DefaultChannelOrdering[] = {
-					EAudioMixerChannel::FrontLeft,
-					EAudioMixerChannel::FrontRight,
-					EAudioMixerChannel::FrontCenter,
-					EAudioMixerChannel::LowFrequency,
-					EAudioMixerChannel::SideLeft,
-					EAudioMixerChannel::SideRight,
-					EAudioMixerChannel::BackLeft,
-					EAudioMixerChannel::BackRight,
-				};
-
-				EAudioMixerChannel::Type* ChannelOrdering = DefaultChannelOrdering;
-
-				// Override channel ordering for some special cases
-				if (OutInfo.NumChannels == 4)
-				{
-					static EAudioMixerChannel::Type DefaultChannelOrderingQuad[] = {
-						EAudioMixerChannel::FrontLeft,
-						EAudioMixerChannel::FrontRight,
-						EAudioMixerChannel::BackLeft,
-						EAudioMixerChannel::BackRight,
-					};
-
-					ChannelOrdering = DefaultChannelOrderingQuad;
-				}
-				else if (OutInfo.NumChannels == 6)
-				{
-					static EAudioMixerChannel::Type DefaultChannelOrdering51[] = {
-						EAudioMixerChannel::FrontLeft,
-						EAudioMixerChannel::FrontRight,
-						EAudioMixerChannel::FrontCenter,
-						EAudioMixerChannel::LowFrequency,
-						EAudioMixerChannel::BackLeft,
-						EAudioMixerChannel::BackRight,
-					};
-
-					ChannelOrdering = DefaultChannelOrdering51;
-				}
-
-				check(OutInfo.NumChannels <= 8);
-				for (int32 Index = 0; Index < OutInfo.NumChannels; ++Index)
-				{
-					OutInfo.OutputChannelArray.Add(ChannelOrdering[Index]);
-				}
-
-			}
-		}
-		else
-		{
-			// Non-extensible formats only support mono or stereo channel output
-			OutInfo.OutputChannelArray.Add(EAudioMixerChannel::FrontLeft);
-			if (OutInfo.NumChannels == 2)
-			{
-				OutInfo.OutputChannelArray.Add(EAudioMixerChannel::FrontRight);
-			}
-		}
-
-		UE_LOG(LogAudioMixer, Display, TEXT("Audio Device Output Speaker Info:"));
-		UE_LOG(LogAudioMixer, Display, TEXT("Name: %s"), *OutInfo.Name);
-		UE_LOG(LogAudioMixer, Display, TEXT("Is Default: %s"), OutInfo.bIsSystemDefault ? TEXT("Yes") : TEXT("No"));
-		UE_LOG(LogAudioMixer, Display, TEXT("Sample Rate: %d"), OutInfo.SampleRate);
-		UE_LOG(LogAudioMixer, Display, TEXT("Channel Count Used: %d"), OutInfo.NumChannels);
-		UE_LOG(LogAudioMixer, Display, TEXT("Device Channel Count: %d"), WaveFormatEx.nChannels);
-		UE_LOG(LogAudioMixer, Display, TEXT("Channel Order:"));
-		for (int32 i = 0; i < OutInfo.NumChannels; ++i)
-		{
-			if (i < OutInfo.OutputChannelArray.Num())
-			{
-				UE_LOG(LogAudioMixer, Display, TEXT("%d: %s"), i, EAudioMixerChannel::ToString(OutInfo.OutputChannelArray[i]));
-			}
-		}
-		return true;
-#else // #elif PLATFORM_HOLOLENS // #if PLATFORM_WINDOWS
+#else // #elif PLATFORM_WINDOWS
 		OutInfo.bIsSystemDefault = true;
 		OutInfo.SampleRate = 44100;
 		OutInfo.DeviceId = 0;
@@ -1230,7 +1041,7 @@ namespace Audio
 		OutInfo.OutputChannelArray.Add(EAudioMixerChannel::SideLeft);
 		OutInfo.OutputChannelArray.Add(EAudioMixerChannel::SideRight);
 		return true;
-#endif // #else // #elif PLATFORM_HOLOLENS // #if PLATFORM_WINDOWS
+#endif 
 	}
 
 	bool FMixerPlatformXAudio2::GetDefaultOutputDeviceIndex(uint32& OutDefaultDeviceIndex) const
@@ -1272,41 +1083,10 @@ namespace Audio
 
 		if (GetNumOutputDevices(NumOutputDevices) && NumOutputDevices > 0)
 		{
-#if PLATFORM_HOLOLENS
-			// On windows, default device index is 0
-			// But if that device cannot be configured try to find one that can be.  This happens in the hololens emulator.
-			if (AudioStreamInfo.OutputDeviceIndex == AUDIO_MIXER_DEFAULT_DEVICE_INDEX)
-			{
-				bool bFoundUsefulDevice = false;
-				for (uint32 i = 0; i < NumOutputDevices; ++i)
-				{
-					OpenStreamParams.OutputDeviceIndex = i;
-					AudioStreamInfo.OutputDeviceIndex = OpenStreamParams.OutputDeviceIndex;
-					if (GetOutputDeviceInfo(AudioStreamInfo.OutputDeviceIndex, AudioStreamInfo.DeviceInfo))
-					{
-						bFoundUsefulDevice = true;
-						break;
-					}
-				}
-				if (!bFoundUsefulDevice)
-				{
-					return false;
-				}
-			}
-			else
-			{
-				if (!GetOutputDeviceInfo(AudioStreamInfo.OutputDeviceIndex, AudioStreamInfo.DeviceInfo))
-				{
-					return false;
-				}
-			}
-#else
 			if (!GetOutputDeviceInfo(AudioStreamInfo.OutputDeviceIndex, AudioStreamInfo.DeviceInfo))
 			{
 				return false;
 			}
-#endif
-
 
 			// Store the device ID here in case it is removed. We can switch back if the device comes back.
 			if (Params.bRestoreIfRemoved)
@@ -1323,17 +1103,7 @@ namespace Audio
 				*AudioStreamInfo.DeviceInfo.DeviceId,
 				nullptr,
 				AudioCategory_GameEffects);
-#elif PLATFORM_HOLOLENS
-		// XAudio2 for HoloLens has different parameters to CreateMasteringVoice
-		// See https://blogs.msdn.microsoft.com/chuckw/2012/04/02/xaudio2-and-windows-8/
-		Result = XAudio2System->CreateMasteringVoice(
-			&OutputAudioStreamMasteringVoice, 
-			AudioStreamInfo.DeviceInfo.NumChannels, 
-			AudioStreamInfo.DeviceInfo.SampleRate, 
-			0, 
-			AllAudioDevices->GetAt(AudioStreamInfo.OutputDeviceIndex)->Id->Data(), 
-			nullptr);
-#else
+#else //PLATFORM_WINDOWS			
 			Result = XAudio2System->CreateMasteringVoice(
 				&OutputAudioStreamMasteringVoice,
 				AudioStreamInfo.DeviceInfo.NumChannels,
@@ -1341,7 +1111,7 @@ namespace Audio
 				0,
 				nullptr,
 				nullptr);
-#endif // #if PLATFORM_WINDOWS
+#endif //#else PLATFORM_WINDOWS
 
 			XAUDIO2_GOTO_CLEANUP_ON_FAIL(Result);
 
@@ -1694,24 +1464,6 @@ namespace Audio
 
 			HRESULT Result;
 			// Create a new master voice
-			// XAudio2 for HoloLens has different parameters to CreateMasteringVoice
-			// See https://blogs.msdn.microsoft.com/chuckw/2012/04/02/xaudio2-and-windows-8/
-#if PLATFORM_HOLOLENS
-			Result = XAudio2System->CreateMasteringVoice(
-				&OutputAudioStreamMasteringVoice, 
-				AudioStreamInfo.DeviceInfo.NumChannels, 
-				AudioStreamInfo.DeviceInfo.SampleRate, 
-				0, 
-				AllAudioDevices->GetAt(AudioStreamInfo.OutputDeviceIndex)->Id->Data(), 
-				nullptr);
-			if (FAILED(Result))
-			{
-				XAUDIO2_LOG_RESULT("XAudio2System->CreateMasteringVoice", Result);
-				// Switch to running null device by setting OutputAudioStreamMasteringVoice to null.
-				// Will default to null device when calling `ResumePlaybackOnNewDevice()`
-				OutputAudioStreamMasteringVoice = nullptr;
-			}
-#else
 			{
 				SCOPED_NAMED_EVENT(FMixerPlatformXAudio2_MoveAudioStreamToNewAudioDevice_CreateMaster, FColor::Blue);
 
@@ -1735,7 +1487,6 @@ namespace Audio
 					AudioStreamInfo.DeviceInfo.NumChannels, AudioStreamInfo.DeviceInfo.SampleRate, *AudioStreamInfo.DeviceInfo.DeviceId, *AudioStreamInfo.DeviceInfo.Name
 				);
 			}
-#endif
 
 			// Setup the format of the output source voice
 			WAVEFORMATEX Format = { 0 };
