Fix incorrect rendering in D3D11 and 12 when losing focus in exclusive fullscreen mode.

Slate switches to Windowed Fullscreen when we lose focus, but doesn't explicitly resize the window, so we never call SetWindowPos or ResizeTarget. If the fullscreen swapchain had a lower resolution than the monitor's native resolution, or if the swapchain had been windowed before, we were stretching the backbuffer to that smaller resolution inside the borderless window, instead of covering all of it. ResizeBuffers is not sufficient because it only affects the buffers, not the window. We need to also call ResizeTarget explicitly when dropping out of exclusive fullscreen, to make sure the window is correctly set up.

Vulkan doesn't have this problem because it doesn't support exclusive fullscreen.

#rnx

diff --git a/Engine/Source/Runtime/D3D12RHI/Private/D3D12Viewport.cpp b/Engine/Source/Runtime/D3D12RHI/Private/D3D12Viewport.cpp
--- a/Engine/Source/Runtime/D3D12RHI/Private/D3D12Viewport.cpp
+++ b/Engine/Source/Runtime/D3D12RHI/Private/D3D12Viewport.cpp
@@ -527,6 +527,20 @@ void FD3D12Viewport::Resize(uint32 InSizeX, uint32 InSizeY, bool bInIsFullscreen
 			// Use ConditionalResetSwapChain to call SetFullscreenState, to handle the failure case.
 			// Ignore the viewport's focus state; since Resize is called as the result of a user action we assume authority without waiting for Focus.
 			ConditionalResetSwapChain(true);
+
+#if D3D12_VIEWPORT_EXPOSES_SWAP_CHAIN
+			if (!bIsFullscreen)
+			{
+				// When exiting fullscreen, make sure that the window has the correct size. This is necessary in the following scenario:
+				//	* we enter exclusive fullscreen with a resolution lower than the monitor's native resolution, or from windowed with a window size smaller than the screen
+				//	* the application loses focus, so Slate asks us to switch to Windowed Fullscreen (see FSlateRenderer::IsViewportFullscreen)
+				//	* InSizeX and InSizeY are given to us as the monitor resolution, so we resize the buffers to the correct resolution below (in ResizeInternal)
+				//	* however, the target still has the smaller size, because Slate doesn't know it has to resize the window too (as far as it's concerned, it's already the right size)
+				//	* therefore, we need to call ResizeTarget, which in windowed mode behaves like SetWindowPos.
+				const DXGI_MODE_DESC BufferDesc = SetupDXGI_MODE_DESC();
+				SwapChain1->ResizeTarget(&BufferDesc);
+			}
+#endif // D3D12_VIEWPORT_EXPOSES_SWAP_CHAIN
 		}
 	}
 
diff --git a/Engine/Source/Runtime/Windows/D3D11RHI/Private/D3D11Viewport.cpp b/Engine/Source/Runtime/Windows/D3D11RHI/Private/D3D11Viewport.cpp
--- a/Engine/Source/Runtime/Windows/D3D11RHI/Private/D3D11Viewport.cpp
+++ b/Engine/Source/Runtime/Windows/D3D11RHI/Private/D3D11Viewport.cpp
@@ -306,6 +306,19 @@ void FD3D11Viewport::Resize(uint32 InSizeX, uint32 InSizeY, bool bInIsFullscreen
 			// Use ConditionalResetSwapChain to call SetFullscreenState, to handle the failure case.
 			// Ignore the viewport's focus state; since Resize is called as the result of a user action we assume authority without waiting for Focus.
 			ResetSwapChainInternal(true);
+
+			if (!bIsFullscreen)
+			{
+				// When exiting fullscreen, make sure that the window has the correct size. This is necessary in the following scenario:
+				//	* we enter exclusive fullscreen with a resolution lower than the monitor's native resolution, or from windowed with a window size smaller than the screen
+				//	* the application loses focus, so Slate asks us to switch to Windowed Fullscreen (see FSlateRenderer::IsViewportFullscreen)
+				//	* InSizeX and InSizeY are given to us as the monitor resolution, so we resize the buffers to the correct resolution below
+				//	* however, the target still has the smaller size, because Slate doesn't know it has to resize the window too (as far as it's concerned, it's already the right size)
+				//	* therefore, we need to call ResizeTarget, which in windowed mode behaves like SetWindowPos.
+				const DXGI_MODE_DESC BufferDesc = SetupDXGI_MODE_DESC();
+				SwapChain->ResizeTarget(&BufferDesc);
+			}
+
 			DXGI_FORMAT RenderTargetFormat = GetRenderTargetFormat(PixelFormat);
 			VERIFYD3D11RESIZEVIEWPORTRESULT(SwapChain->ResizeBuffers(0, SizeX, SizeY, RenderTargetFormat, GetSwapChainFlags()), OldState, NewState, D3DRHI->GetDevice());
 		}
