Fixing includes in D3D11RHI so that all headers can be included independently.

diff --git a/Engine/Source/Runtime/Windows/D3D11RHI/Private/D3D11ConstantBuffer.cpp b/Engine/Source/Runtime/Windows/D3D11RHI/Private/D3D11ConstantBuffer.cpp
--- a/Engine/Source/Runtime/Windows/D3D11RHI/Private/D3D11ConstantBuffer.cpp
+++ b/Engine/Source/Runtime/Windows/D3D11RHI/Private/D3D11ConstantBuffer.cpp
@@ -4,6 +4,7 @@
 	D3D11ConstantBuffer.cpp: D3D Constant buffer RHI implementation.
 =============================================================================*/
 
+#include "D3D11ConstantBuffer.h"
 #include "D3D11RHIPrivate.h"
 
 // New circular buffer system for faster constant uploads.  Avoids CopyResource and speeds things up considerably
diff --git a/Engine/Source/Runtime/Windows/D3D11RHI/Private/D3D11RHI.cpp b/Engine/Source/Runtime/Windows/D3D11RHI/Private/D3D11RHI.cpp
--- a/Engine/Source/Runtime/Windows/D3D11RHI/Private/D3D11RHI.cpp
+++ b/Engine/Source/Runtime/Windows/D3D11RHI/Private/D3D11RHI.cpp
@@ -4,6 +4,7 @@
 	D3D11RHI.cpp: Unreal D3D RHI library implementation.
 =============================================================================*/
 
+#include "D3D11RHI.h"
 #include "D3D11RHIPrivate.h"
 #include "RHIStaticStates.h"
 #include "StaticBoundShaderState.h"
diff --git a/Engine/Source/Runtime/Windows/D3D11RHI/Private/D3D11State.cpp b/Engine/Source/Runtime/Windows/D3D11RHI/Private/D3D11State.cpp
--- a/Engine/Source/Runtime/Windows/D3D11RHI/Private/D3D11State.cpp
+++ b/Engine/Source/Runtime/Windows/D3D11RHI/Private/D3D11State.cpp
@@ -4,6 +4,7 @@
 	D3D11State.cpp: D3D state implementation.
 =============================================================================*/
 
+#include "D3D11State.h"
 #include "D3D11RHIPrivate.h"
 #include "RHIUtilities.h"
 
diff --git a/Engine/Source/Runtime/Windows/D3D11RHI/Private/D3D11StateCache.cpp b/Engine/Source/Runtime/Windows/D3D11RHI/Private/D3D11StateCache.cpp
--- a/Engine/Source/Runtime/Windows/D3D11RHI/Private/D3D11StateCache.cpp
+++ b/Engine/Source/Runtime/Windows/D3D11RHI/Private/D3D11StateCache.cpp
@@ -6,6 +6,7 @@
 //-----------------------------------------------------------------------------
 //	Include Files
 //-----------------------------------------------------------------------------
+#include "D3D11StateCache.h"
 #include "D3D11RHIPrivate.h"
 #include "D3D11StateCache.h"
 
diff --git a/Engine/Source/Runtime/Windows/D3D11RHI/Private/D3D11StateCachePrivate.h b/Engine/Source/Runtime/Windows/D3D11RHI/Private/D3D11StateCachePrivate.h
--- a/Engine/Source/Runtime/Windows/D3D11RHI/Private/D3D11StateCachePrivate.h
+++ b/Engine/Source/Runtime/Windows/D3D11RHI/Private/D3D11StateCachePrivate.h
@@ -5,6 +5,9 @@
 
 #pragma once
 
+#include "RHIDefinitions.h"
+#include "Windows/D3D11ThirdParty.h"
+
 //-----------------------------------------------------------------------------
 //	Configuration
 //-----------------------------------------------------------------------------
diff --git a/Engine/Source/Runtime/Windows/D3D11RHI/Private/D3D11Util.cpp b/Engine/Source/Runtime/Windows/D3D11RHI/Private/D3D11Util.cpp
--- a/Engine/Source/Runtime/Windows/D3D11RHI/Private/D3D11Util.cpp
+++ b/Engine/Source/Runtime/Windows/D3D11RHI/Private/D3D11Util.cpp
@@ -4,6 +4,7 @@
 	D3D11Util.h: D3D RHI utility implementation.
 =============================================================================*/
 
+#include "D3D11Util.h"
 #include "D3D11RHIPrivate.h"
 #include "EngineModule.h"
 #include "RendererInterface.h"
diff --git a/Engine/Source/Runtime/Windows/D3D11RHI/Private/D3D11Viewport.cpp b/Engine/Source/Runtime/Windows/D3D11RHI/Private/D3D11Viewport.cpp
--- a/Engine/Source/Runtime/Windows/D3D11RHI/Private/D3D11Viewport.cpp
+++ b/Engine/Source/Runtime/Windows/D3D11RHI/Private/D3D11Viewport.cpp
@@ -4,6 +4,7 @@
 	D3D11Viewport.cpp: D3D viewport RHI implementation.
 =============================================================================*/
 
+#include "D3D11Viewport.h"
 #include "D3D11RHIPrivate.h"
 #include "RenderCore.h"
 #include "HDRHelper.h"
@@ -648,6 +649,21 @@ bool FD3D11Viewport::Present(bool bLockToVsync)
 	return bNativelyPresented;
 }
 
+void* FD3D11Viewport::GetNativeSwapChain() const
+{
+	return GetSwapChain();
+}
+
+void* FD3D11Viewport::GetNativeBackBufferTexture() const
+{
+	return GetBackBuffer()->GetD3D11Texture2D();
+}
+
+void* FD3D11Viewport::GetNativeBackBufferRT() const
+{
+	return GetBackBuffer()->GetRenderTargetView(0, 0);
+}
+
 EColorSpaceAndEOTF FD3D11DynamicRHI::RHIGetColorSpace(FRHIViewport* ViewportRHI)
 {
 	FD3D11Viewport* Viewport = ResourceCast(ViewportRHI);
diff --git a/Engine/Source/Runtime/Windows/D3D11RHI/Private/Windows/WindowsD3D11ConstantBuffer.h b/Engine/Source/Runtime/Windows/D3D11RHI/Private/Windows/WindowsD3D11ConstantBuffer.h
--- a/Engine/Source/Runtime/Windows/D3D11RHI/Private/Windows/WindowsD3D11ConstantBuffer.h
+++ b/Engine/Source/Runtime/Windows/D3D11RHI/Private/Windows/WindowsD3D11ConstantBuffer.h
@@ -7,8 +7,7 @@
 #pragma once
 
 #include "D3D11ConstantBuffer.h"
-
-struct ID3D11Buffer;
+#include "Windows/D3D11ThirdParty.h"
 
 class FWinD3D11ConstantBuffer : public FD3D11ConstantBuffer
 {
diff --git a/Engine/Source/Runtime/Windows/D3D11RHI/Public/D3D11Resources.h b/Engine/Source/Runtime/Windows/D3D11RHI/Public/D3D11Resources.h
--- a/Engine/Source/Runtime/Windows/D3D11RHI/Public/D3D11Resources.h
+++ b/Engine/Source/Runtime/Windows/D3D11RHI/Public/D3D11Resources.h
@@ -7,8 +7,9 @@
 #pragma once
 
 #include "BoundShaderStateCache.h"
+#include "ShaderCore.h"
+#include "Windows/D3D11ThirdParty.h"
 
-interface ID3D11DeviceContext;
 typedef ID3D11DeviceContext FD3D11DeviceContext;
 
 template <>
@@ -572,24 +573,3 @@ struct TD3D11ResourceTraits<FRHIUnorderedAccessView>
 	typedef FD3D11UnorderedAccessView TConcreteType;
 };
 
-template<>
-struct TD3D11ResourceTraits<FRHISamplerState>
-{
-	typedef FD3D11SamplerState TConcreteType;
-};
-template<>
-struct TD3D11ResourceTraits<FRHIRasterizerState>
-{
-	typedef FD3D11RasterizerState TConcreteType;
-};
-template<>
-struct TD3D11ResourceTraits<FRHIDepthStencilState>
-{
-	typedef FD3D11DepthStencilState TConcreteType;
-};
-template<>
-struct TD3D11ResourceTraits<FRHIBlendState>
-{
-	typedef FD3D11BlendState TConcreteType;
-};
-
diff --git a/Engine/Source/Runtime/Windows/D3D11RHI/Public/D3D11State.h b/Engine/Source/Runtime/Windows/D3D11RHI/Public/D3D11State.h
--- a/Engine/Source/Runtime/Windows/D3D11RHI/Public/D3D11State.h
+++ b/Engine/Source/Runtime/Windows/D3D11RHI/Public/D3D11State.h
@@ -6,6 +6,8 @@
 
 #pragma once
 
+#include "D3D11Resources.h"
+
 class FD3D11SamplerState : public FRHISamplerState
 {
 public:
@@ -37,3 +39,23 @@ public:
 	TRefCountPtr<ID3D11BlendState> Resource;
 };
 
+template<>
+struct TD3D11ResourceTraits<FRHISamplerState>
+{
+	typedef FD3D11SamplerState TConcreteType;
+};
+template<>
+struct TD3D11ResourceTraits<FRHIRasterizerState>
+{
+	typedef FD3D11RasterizerState TConcreteType;
+};
+template<>
+struct TD3D11ResourceTraits<FRHIDepthStencilState>
+{
+	typedef FD3D11DepthStencilState TConcreteType;
+};
+template<>
+struct TD3D11ResourceTraits<FRHIBlendState>
+{
+	typedef FD3D11BlendState TConcreteType;
+};
diff --git a/Engine/Source/Runtime/Windows/D3D11RHI/Public/D3D11Util.h b/Engine/Source/Runtime/Windows/D3D11RHI/Public/D3D11Util.h
--- a/Engine/Source/Runtime/Windows/D3D11RHI/Public/D3D11Util.h
+++ b/Engine/Source/Runtime/Windows/D3D11RHI/Public/D3D11Util.h
@@ -6,6 +6,14 @@
 
 #pragma once
 
+#include "RHICommandList.h"
+#include "RHIDefinitions.h"
+#include "Templates/RefCounting.h"
+#include "Windows/D3D11ThirdParty.h"
+
+class FRHIBuffer;
+class FRHITexture;
+
 #define D3D11RHI_IMMEDIATE_CONTEXT	(GD3D11RHI->GetDeviceContext())
 #define D3D11RHI_DEVICE				(GD3D11RHI->GetDevice())
 
@@ -229,6 +237,7 @@ struct FD3D11RHIGenericCommandString
 {
 	static const TCHAR* TStr() { return TEXT("FD3D11RHIGenericCommand"); }
 };
+
 template <
 	typename JobType,
 	typename = TEnableIf<
diff --git a/Engine/Source/Runtime/Windows/D3D11RHI/Public/D3D11Viewport.h b/Engine/Source/Runtime/Windows/D3D11RHI/Public/D3D11Viewport.h
--- a/Engine/Source/Runtime/Windows/D3D11RHI/Public/D3D11Viewport.h
+++ b/Engine/Source/Runtime/Windows/D3D11RHI/Public/D3D11Viewport.h
@@ -6,9 +6,11 @@
 
 #pragma once
 
-#include "CoreMinimal.h"
+#include "D3D11Resources.h"
 #include "RenderResource.h"
 #include "RenderUtils.h"
+#include "RHIResources.h"
+#include "Windows/D3D11ThirdParty.h"
 
 /** A D3D event query resource. */
 class FD3D11EventQuery
@@ -80,9 +82,9 @@ public:
 
 	IDXGISwapChain* GetSwapChain() const { return SwapChain; }
 
-	virtual void* GetNativeSwapChain() const override { return GetSwapChain(); }
-	virtual void* GetNativeBackBufferTexture() const override { return GetBackBuffer()->GetD3D11Texture2D(); }
-	virtual void* GetNativeBackBufferRT() const override { return GetBackBuffer()->GetRenderTargetView(0, 0); }
+	virtual void* GetNativeSwapChain() const override;
+	virtual void* GetNativeBackBufferTexture() const override;
+	virtual void* GetNativeBackBufferRT() const override;
 
 	virtual void SetCustomPresent(FRHICustomPresent* InCustomPresent) override
 	{
diff --git a/Engine/Source/Runtime/Windows/D3D11RHI/Public/Windows/D3D11ThirdParty.h b/Engine/Source/Runtime/Windows/D3D11RHI/Public/Windows/D3D11ThirdParty.h
--- a/Engine/Source/Runtime/Windows/D3D11RHI/Public/Windows/D3D11ThirdParty.h
+++ b/Engine/Source/Runtime/Windows/D3D11RHI/Public/Windows/D3D11ThirdParty.h
@@ -2,9 +2,8 @@
 
 #pragma once
 
-// Disable macro redefinition warning for compatibility with Windows SDK 8+
-#pragma warning(push)
-#pragma warning(disable : 4005)	// macro redefinition
+#include "HAL/Platform.h"
+#include "Microsoft/AllowMicrosoftPlatformTypes.h"
 
 #if PLATFORM_64BITS
 	#pragma pack(push,16)
@@ -12,15 +11,14 @@
 	#pragma pack(push,8)
 #endif
 
-#include "Microsoft/AllowMicrosoftPlatformTypes.h"
 THIRD_PARTY_INCLUDES_START
 	#include <d3d11_2.h>
 	#include <dxgi1_6.h>
 	#include <dxgidebug.h>
 THIRD_PARTY_INCLUDES_END
-#include "Microsoft/HideMicrosoftPlatformTypes.h"
 
 #undef DrawText
 
 #pragma pack(pop)
-#pragma warning(pop)
+
+#include "Microsoft/HideMicrosoftPlatformTypes.h"
