Fix: Move creation of FWindowsVideoRecordingSystem away from FWindowsPlatformFeaturesModule constructor

diff --git a/Engine/Source/Runtime/Windows/WindowsPlatformFeatures/Private/WindowsPlatformFeatures.cpp b/Engine/Source/Runtime/Windows/WindowsPlatformFeatures/Private/WindowsPlatformFeatures.cpp
--- a/Engine/Source/Runtime/Windows/WindowsPlatformFeatures/Private/WindowsPlatformFeatures.cpp
+++ b/Engine/Source/Runtime/Windows/WindowsPlatformFeatures/Private/WindowsPlatformFeatures.cpp
@@ -11,11 +11,15 @@ IMPLEMENT_MODULE(FWindowsPlatformFeaturesModule, WindowsPlatformFeatures);
 
 FWindowsPlatformFeaturesModule::FWindowsPlatformFeaturesModule()
 {
-	VideoRecordingSystem = MakeShared<FWindowsVideoRecordingSystem>();
 }
 
 IVideoRecordingSystem* FWindowsPlatformFeaturesModule::GetVideoRecordingSystem()
 {
+	if (!VideoRecordingSystem)
+	{
+		VideoRecordingSystem = MakeShared<FWindowsVideoRecordingSystem>();
+	}
+	
 	return VideoRecordingSystem.Get();
 }
 
