[Backout] - CL32236602
Original CL Desc
-----------------------------------------------------------------
Fix: Move creation of FWindowsVideoRecordingSystem away from FWindowsPlatformFeaturesModule constructor

diff --git a/Engine/Source/Runtime/Windows/WindowsPlatformFeatures/Private/WindowsPlatformFeatures.cpp b/Engine/Source/Runtime/Windows/WindowsPlatformFeatures/Private/WindowsPlatformFeatures.cpp
--- a/Engine/Source/Runtime/Windows/WindowsPlatformFeatures/Private/WindowsPlatformFeatures.cpp
+++ b/Engine/Source/Runtime/Windows/WindowsPlatformFeatures/Private/WindowsPlatformFeatures.cpp
@@ -15,17 +15,8 @@ FWindowsPlatformFeaturesModule::FWindowsPlatformFeaturesModule()
 
 IVideoRecordingSystem* FWindowsPlatformFeaturesModule::GetVideoRecordingSystem()
 {
-	if (!VideoRecordingSystem)
-	{
-		VideoRecordingSystem = MakeShared<FWindowsVideoRecordingSystem>();
-	}
-	
-	return VideoRecordingSystem.Get();
-}
-
-void FWindowsPlatformFeaturesModule::RegisterVideoRecordingSystem(TSharedPtr<IVideoRecordingSystem> InVideoRecordingSystem)
-{
-	VideoRecordingSystem = InVideoRecordingSystem;
+	static FWindowsVideoRecordingSystem VideoRecordingSystem;
+	return &VideoRecordingSystem;
 }
 
 ISaveGameSystem* FWindowsPlatformFeaturesModule::GetSaveGameSystem()
