Removing unused getters from the D3D11 state cache.

diff --git a/Engine/Source/Runtime/Windows/D3D11RHI/Private/D3D11StateCachePrivate.h b/Engine/Source/Runtime/Windows/D3D11RHI/Private/D3D11StateCachePrivate.h
--- a/Engine/Source/Runtime/Windows/D3D11RHI/Private/D3D11StateCachePrivate.h
+++ b/Engine/Source/Runtime/Windows/D3D11RHI/Private/D3D11StateCachePrivate.h
@@ -271,43 +271,6 @@ public:
 		InternalSetShaderResourceView<ShaderFrequency>(SRV, ResourceIndex);
 	}
 
-	template <EShaderFrequency ShaderFrequency>
-	D3D11_STATE_CACHE_INLINE void GetShaderResourceViews(uint32 StartResourceIndex, uint32 NumResources, ID3D11ShaderResourceView** SRV)
-	{
-#if D3D11_ALLOW_STATE_CACHE
-		{
-			check(StartResourceIndex + NumResources <= D3D11_COMMONSHADER_INPUT_RESOURCE_SLOT_COUNT);
-			for (uint32 ResourceLoop = 0; ResourceLoop < NumResources; ResourceLoop++)
-			{
-				SRV[ResourceLoop] = CurrentShaderResourceViews[ShaderFrequency][ResourceLoop + StartResourceIndex];
-				if (SRV[ResourceLoop])
-				{
-					SRV[ResourceLoop]->AddRef();
-				}
-			}
-		}
-#else	// !D3D11_ALLOW_STATE_CACHE
-		{
-			switch (ShaderFrequency)
-			{
-			case SF_Vertex:		
-				Direct3DDeviceIMContext->VSGetShaderResources(StartResourceIndex, NumResources, SRV); 
-				break;
-			case SF_Geometry:	
-				Direct3DDeviceIMContext->GSGetShaderResources(StartResourceIndex, NumResources, SRV); 
-				break;
-			case SF_Pixel:		
-				Direct3DDeviceIMContext->PSGetShaderResources(StartResourceIndex, NumResources, SRV); 
-				break;
-			case SF_Compute:	
-				Direct3DDeviceIMContext->CSGetShaderResources(StartResourceIndex, NumResources, SRV); 
-				break;
-			}
-		}
-
-#endif	// D3D11_ALLOW_STATE_CACHE
-	}
-
 	D3D11_STATE_CACHE_INLINE void SetViewport(D3D11_VIEWPORT Viewport)
 	{
 #if D3D11_ALLOW_STATE_CACHE
@@ -402,34 +365,6 @@ public:
 #endif
 	}
 
-template <EShaderFrequency ShaderFrequency>
-	D3D11_STATE_CACHE_INLINE void GetConstantBuffers(uint32 StartSlotIndex, uint32 NumBuffers, ID3D11Buffer** ConstantBuffers)
-	{
-#if D3D11_ALLOW_STATE_CACHE
-		check(StartSlotIndex + NumBuffers <= D3D11_COMMONSHADER_CONSTANT_BUFFER_API_SLOT_COUNT);
-		for (uint32 constantLoop = 0; constantLoop < NumBuffers; constantLoop++)
-		{
-			FD3D11ConstantBufferState& cb = CurrentConstantBuffers[ShaderFrequency][constantLoop + StartSlotIndex];
-			ConstantBuffers[constantLoop] = cb.Buffer;
-			if (ConstantBuffers[constantLoop])
-			{
-				ConstantBuffers[constantLoop]->AddRef();
-			}
-		}
-#else
-		{
-			switch (ShaderFrequency)
-			{
-			case SF_Vertex:		Direct3DDeviceIMContext1->VSGetConstantBuffers(StartSlot, NumBuffers, ConstantBuffers); break;
-			case SF_Geometry:	Direct3DDeviceIMContext1->GSGetConstantBuffers(StartSlot, NumBuffers, ConstantBuffers); break;
-			case SF_Pixel:		Direct3DDeviceIMContext1->PSGetConstantBuffers(StartSlot, NumBuffers, ConstantBuffers); break;
-			case SF_Compute:	Direct3DDeviceIMContext1->CSGetConstantBuffers(StartSlot, NumBuffers, ConstantBuffers); break;
-			}
-
-		}
-#endif
-	}
-
 	D3D11_STATE_CACHE_INLINE void SetRasterizerState(ID3D11RasterizerState* State)
 	{
 #if D3D11_ALLOW_STATE_CACHE
@@ -445,19 +380,6 @@ template <EShaderFrequency ShaderFrequency>
 #endif
 	}
 
-	D3D11_STATE_CACHE_INLINE void GetRasterizerState(ID3D11RasterizerState** RasterizerState)
-	{
-#if D3D11_ALLOW_STATE_CACHE
-		*RasterizerState = CurrentRasterizerState;
-		if (CurrentRasterizerState)
-		{
-			CurrentRasterizerState->AddRef();
-		}
-#else
-		Direct3DDeviceIMContext->RSGetState(RasterizerState);
-#endif
-	}
-
 	D3D11_STATE_CACHE_INLINE void SetBlendState(ID3D11BlendState* State, const float BlendFactor[4], uint32 SampleMask)
 	{
 #if D3D11_ALLOW_STATE_CACHE
@@ -491,21 +413,6 @@ template <EShaderFrequency ShaderFrequency>
 #endif
 	}
 
-	D3D11_STATE_CACHE_INLINE void GetBlendState(ID3D11BlendState** BlendState, float BlendFactor[4], uint32* SampleMask)
-	{
-#if D3D11_ALLOW_STATE_CACHE
-		*BlendState = CurrentBlendState;
-		if (CurrentBlendState)
-		{
-			CurrentBlendState->AddRef();
-		}
-		*SampleMask = CurrentBlendSampleMask;
-		FMemory::Memcpy(BlendFactor, CurrentBlendFactor, sizeof(CurrentBlendFactor));
-#else
-		Direct3DDeviceIMContext->OMGetBlendState(BlendState, BlendFactor, SampleMask);
-#endif
-	}
-
 	D3D11_STATE_CACHE_INLINE void SetDepthStencilState(ID3D11DepthStencilState* State, uint32 RefStencil)
 	{
 #if D3D11_ALLOW_STATE_CACHE
@@ -537,20 +444,6 @@ template <EShaderFrequency ShaderFrequency>
 #endif
 	}
 
-	D3D11_STATE_CACHE_INLINE void GetDepthStencilState(ID3D11DepthStencilState** DepthStencilState, uint32* StencilRef)
-	{
-#if D3D11_ALLOW_STATE_CACHE
-		*DepthStencilState = CurrentDepthStencilState;
-		*StencilRef = CurrentReferenceStencil;
-		if (CurrentDepthStencilState)
-		{
-			CurrentDepthStencilState->AddRef();
-		}
-#else
-		Direct3DDeviceIMContext->OMGetDepthStencilState(DepthStencilState, StencilRef);
-#endif
-	}
-
 	D3D11_STATE_CACHE_INLINE void SetVertexShader(ID3D11VertexShader* Shader)
 	{
 #if D3D11_ALLOW_STATE_CACHE
@@ -683,19 +576,6 @@ template <EShaderFrequency ShaderFrequency>
 #endif
 	}
 
-	D3D11_STATE_CACHE_INLINE void GetInputLayout(ID3D11InputLayout** InputLayout)
-	{
-#if D3D11_ALLOW_STATE_CACHE
-		*InputLayout = CurrentInputLayout;
-		if (CurrentInputLayout)
-		{
-			CurrentInputLayout->AddRef();
-		}
-#else
-		Direct3DDeviceIMContext->IAGetInputLayout(InputLayout);
-#endif
-	}
-
 	D3D11_STATE_CACHE_INLINE void SetStreamSource(ID3D11Buffer* VertexBuffer, uint32 StreamIndex, uint32 Stride, uint32 Offset)
 	{
 		ensure(Stride == StreamStrides[StreamIndex]);
@@ -707,26 +587,6 @@ template <EShaderFrequency ShaderFrequency>
 		InternalSetStreamSource(VertexBuffer, StreamIndex, StreamStrides[StreamIndex], Offset);
 	}
 
-	D3D11_STATE_CACHE_INLINE void GetStreamSources(uint32 StartStreamIndex, uint32 NumStreams, ID3D11Buffer** VertexBuffers, uint32* Strides, uint32* Offsets)
-	{
-#if D3D11_ALLOW_STATE_CACHE
-		check (StartStreamIndex + NumStreams <= D3D11_IA_VERTEX_INPUT_RESOURCE_SLOT_COUNT);
-		for (uint32 StreamLoop = 0; StreamLoop < NumStreams; StreamLoop++)
-		{
-			FD3D11VertexBufferState& Slot = CurrentVertexBuffers[StreamLoop + StartStreamIndex];
-			VertexBuffers[StreamLoop] = Slot.VertexBuffer;
-			Strides[StreamLoop] = Slot.Stride;
-			Offsets[StreamLoop] = Slot.Offset;
-			if (Slot.VertexBuffer)
-			{
-				Slot.VertexBuffer->AddRef();
-			}
-		}
-#else
-		Direct3DDeviceIMContext->IAGetVertexBuffers(StartStreamIndex, NumStreams, VertexBuffers, Strides, Offsets);
-#endif
-	}
-
 public:
 
 	D3D11_STATE_CACHE_INLINE void SetIndexBuffer(ID3D11Buffer* IndexBuffer, DXGI_FORMAT Format, uint32 Offset)
@@ -734,21 +594,6 @@ public:
 		InternalSetIndexBuffer(IndexBuffer, Format, Offset);
 	}
 
-	D3D11_STATE_CACHE_INLINE void GetIndexBuffer(ID3D11Buffer** IndexBuffer, DXGI_FORMAT* Format, uint32* Offset)
-	{
-#if D3D11_ALLOW_STATE_CACHE
-		*IndexBuffer = CurrentIndexBuffer;
-		*Format = CurrentIndexFormat;
-		*Offset = CurrentIndexOffset;
-		if (CurrentIndexBuffer)
-		{
-			CurrentIndexBuffer->AddRef();
-		}
-#else
-		Direct3DDeviceIMContext->IAGetIndexBuffer(IndexBuffer,Format,Offset);
-#endif
-	}
-
 	D3D11_STATE_CACHE_INLINE void SetPrimitiveTopology(D3D11_PRIMITIVE_TOPOLOGY PrimitiveTopology)
 	{
 #if D3D11_ALLOW_STATE_CACHE
@@ -764,16 +609,6 @@ public:
 #endif
 	}
 
-	D3D11_STATE_CACHE_INLINE void GetPrimitiveTopology(D3D11_PRIMITIVE_TOPOLOGY* PrimitiveTopology)
-	{
-#if D3D11_ALLOW_STATE_CACHE
-		*PrimitiveTopology = CurrentPrimitiveTopology;
-#else
-		Direct3DDeviceIMContext->IAGetPrimitiveTopology(PrimitiveTopology);
-#endif
-	}
-
-	
 	FD3D11StateCacheBase()
 		: Direct3DDeviceIMContext(nullptr)
 	{
