Log Windows memory stats in addition to GPU memory stats in case of a driver-reported OOM.

#tested locally by triggering a OOM and verifying memory stats are logged

diff --git a/Engine/Source/Runtime/D3D12RHI/Private/D3D12Util.cpp b/Engine/Source/Runtime/D3D12RHI/Private/D3D12Util.cpp
--- a/Engine/Source/Runtime/D3D12RHI/Private/D3D12Util.cpp
+++ b/Engine/Source/Runtime/D3D12RHI/Private/D3D12Util.cpp
@@ -837,6 +837,9 @@ void FD3D12DynamicRHI::TerminateOnOutOfMemory(ID3D12Device* InDevice, HRESULT D3
 		FD3D12Adapter* Adapter = IterationDevice->GetParentAdapter();
 		LogMemoryStats(Adapter);
 	});
+
+	// Also log Windows memory stats.
+	FPlatformMemory::DumpStats(*GLog);
 	
 	UE_LOG(LogD3D12RHI, Fatal, TEXT("Out of video memory trying to allocate a rendering resource"));
 	if (!bGPUCrashOOM)
diff --git a/Engine/Source/Runtime/Windows/D3D11RHI/Private/D3D11Util.cpp b/Engine/Source/Runtime/Windows/D3D11RHI/Private/D3D11Util.cpp
--- a/Engine/Source/Runtime/Windows/D3D11RHI/Private/D3D11Util.cpp
+++ b/Engine/Source/Runtime/Windows/D3D11RHI/Private/D3D11Util.cpp
@@ -154,6 +154,7 @@ static void TerminateOnOutOfMemory(HRESULT D3DResult, bool bCreatingTextures)
 	{
 		uint64 VRAMBudgetBytes = 0, VRAMUsageBytes = 0;
 		GetAndLogMemoryStats(GD3D11RHI->GetAdapter(), VRAMBudgetBytes, VRAMUsageBytes);
+		FPlatformMemory::DumpStats(*GLog);
 		FCoreDelegates::GetGPUOutOfMemoryDelegate().Broadcast(VRAMBudgetBytes, VRAMUsageBytes);
 
 		if (bCreatingTextures)
