Kill legacy code path in XMA to tidy it up

diff --git a/Engine/Source/Runtime/Windows/AudioMixerXAudio2/Private/AudioMixerPlatformXAudio2.cpp b/Engine/Source/Runtime/Windows/AudioMixerXAudio2/Private/AudioMixerPlatformXAudio2.cpp
--- a/Engine/Source/Runtime/Windows/AudioMixerXAudio2/Private/AudioMixerPlatformXAudio2.cpp
+++ b/Engine/Source/Runtime/Windows/AudioMixerXAudio2/Private/AudioMixerPlatformXAudio2.cpp
@@ -643,7 +643,7 @@ namespace Audio
 
 #if WITH_XMA2
 		//Initialize our XMA2 decoder context
-		XMA2_INFO_CALL(FXMAAudioInfo::Initialize());
+		FXMAAudioInfo::Initialize();
 #endif //#if WITH_XMA2
 
 		if(IAudioMixer::ShouldRecycleThreads())
@@ -687,7 +687,7 @@ namespace Audio
 		SAFE_RELEASE(XAudio2System);
 
 #if WITH_XMA2
-		XMA2_INFO_CALL(FXMAAudioInfo::Shutdown());
+		FXMAAudioInfo::Shutdown();
 #endif
 
 #if PLATFORM_WINDOWS
@@ -1603,7 +1603,7 @@ namespace Audio
 	void FMixerPlatformXAudio2::OnHardwareUpdate()
 	{
 #if WITH_XMA2
-		XMA2_INFO_CALL(FXMAAudioInfo::Tick());
+		FXMAAudioInfo::Tick();
 #endif //WITH_XMA2
 
 	}
