Deprecated RHIGetColorSpace as it always returns rec709.

diff --git a/Engine/Source/Runtime/NullDrv/Public/NullRHI.h b/Engine/Source/Runtime/NullDrv/Public/NullRHI.h
--- a/Engine/Source/Runtime/NullDrv/Public/NullRHI.h
+++ b/Engine/Source/Runtime/NullDrv/Public/NullRHI.h
@@ -379,11 +379,6 @@ public:
 
 	}
 
-	virtual EColorSpaceAndEOTF RHIGetColorSpace(FRHIViewport* Viewport ) final override
-	{
-		return EColorSpaceAndEOTF::EColorSpace_Rec709;
-	}
-
 	virtual void RHICheckViewportHDRStatus(FRHIViewport* Viewport) final override
 	{
 	}
diff --git a/Engine/Source/Runtime/RHI/Private/DynamicRHI.cpp b/Engine/Source/Runtime/RHI/Private/DynamicRHI.cpp
--- a/Engine/Source/Runtime/RHI/Private/DynamicRHI.cpp
+++ b/Engine/Source/Runtime/RHI/Private/DynamicRHI.cpp
@@ -798,11 +798,6 @@ FRenderQueryPoolRHIRef RHICreateRenderQueryPool(ERenderQueryType QueryType, uint
 	return new FDefaultRHIRenderQueryPool(QueryType);
 }
 
-EColorSpaceAndEOTF FDynamicRHI::RHIGetColorSpace(FRHIViewport* Viewport)
-{
-	return EColorSpaceAndEOTF::ERec709_sRGB;
-}
-
 void FDynamicRHI::RHICheckViewportHDRStatus(FRHIViewport* Viewport)
 {
 }
diff --git a/Engine/Source/Runtime/RHI/Public/DynamicRHI.h b/Engine/Source/Runtime/RHI/Public/DynamicRHI.h
--- a/Engine/Source/Runtime/RHI/Public/DynamicRHI.h
+++ b/Engine/Source/Runtime/RHI/Public/DynamicRHI.h
@@ -718,9 +718,6 @@ public:
 		RHIResizeViewport(Viewport, SizeX, SizeY, bIsFullscreen);
 	}
 
-	// Return what colour space the viewport is in. Used for HDR displays
-	RHI_API virtual EColorSpaceAndEOTF RHIGetColorSpace(FRHIViewport* Viewport);
-
 	// Return preferred pixel format if given format is unsupported.
 	virtual EPixelFormat RHIPreferredPixelFormatHint(EPixelFormat PreferredPixelFormat)
 	{
@@ -1351,9 +1348,10 @@ FORCEINLINE void RHIResizeViewport(FRHIViewport* Viewport, uint32 SizeX, uint32
 	GDynamicRHI->RHIResizeViewport(Viewport, SizeX, SizeY, bIsFullscreen, PreferredPixelFormat);
 }
 
+// UE_DEPRECATED(5.5, "This method is no longer used.")
 FORCEINLINE EColorSpaceAndEOTF RHIGetColorSpace(FRHIViewport* Viewport)
 {
-	return GDynamicRHI->RHIGetColorSpace(Viewport);
+	return EColorSpaceAndEOTF::ERec709_sRGB;
 }
 
 FORCEINLINE void RHICheckViewportHDRStatus(FRHIViewport* Viewport)
diff --git a/Engine/Source/Runtime/RHI/Public/RHIValidation.h b/Engine/Source/Runtime/RHI/Public/RHIValidation.h
--- a/Engine/Source/Runtime/RHI/Public/RHIValidation.h
+++ b/Engine/Source/Runtime/RHI/Public/RHIValidation.h
@@ -922,11 +922,6 @@ public:
 		RHI->RHIResizeViewport(Viewport, SizeX, SizeY, bIsFullscreen, PreferredPixelFormat);
 	}
 
-	virtual EColorSpaceAndEOTF RHIGetColorSpace(FRHIViewport* Viewport) override final
-	{
-		return RHI->RHIGetColorSpace(Viewport);
-	}
-
 	virtual EPixelFormat RHIPreferredPixelFormatHint(EPixelFormat PreferredPixelFormat) override final
 	{
 		return RHI->RHIPreferredPixelFormatHint(PreferredPixelFormat);
diff --git a/Engine/Source/Runtime/Windows/D3D11RHI/Private/D3D11RHIPrivate.h b/Engine/Source/Runtime/Windows/D3D11RHI/Private/D3D11RHIPrivate.h
--- a/Engine/Source/Runtime/Windows/D3D11RHI/Private/D3D11RHIPrivate.h
+++ b/Engine/Source/Runtime/Windows/D3D11RHI/Private/D3D11RHIPrivate.h
@@ -485,7 +485,6 @@ public:
 	virtual FViewportRHIRef RHICreateViewport(void* WindowHandle, uint32 SizeX, uint32 SizeY, bool bIsFullscreen, EPixelFormat PreferredPixelFormat) override;
 	virtual void RHIResizeViewport(FRHIViewport* Viewport, uint32 SizeX, uint32 SizeY, bool bIsFullscreen) final override;
 	virtual void RHIResizeViewport(FRHIViewport* Viewport, uint32 SizeX, uint32 SizeY, bool bIsFullscreen, EPixelFormat PreferredPixelFormat) final override;
-	virtual EColorSpaceAndEOTF RHIGetColorSpace(FRHIViewport* Viewport) final override;
 	virtual void RHICheckViewportHDRStatus(FRHIViewport* ViewportRHI) final override;
 	virtual void RHITick(float DeltaTime) final override;
 	virtual void RHIBlockUntilGPUIdle() final override;
diff --git a/Engine/Source/Runtime/Windows/D3D11RHI/Private/D3D11Viewport.cpp b/Engine/Source/Runtime/Windows/D3D11RHI/Private/D3D11Viewport.cpp
--- a/Engine/Source/Runtime/Windows/D3D11RHI/Private/D3D11Viewport.cpp
+++ b/Engine/Source/Runtime/Windows/D3D11RHI/Private/D3D11Viewport.cpp
@@ -664,12 +664,6 @@ void* FD3D11Viewport::GetNativeBackBufferRT() const
 	return GetBackBuffer()->GetRenderTargetView(0, 0);
 }
 
-EColorSpaceAndEOTF FD3D11DynamicRHI::RHIGetColorSpace(FRHIViewport* ViewportRHI)
-{
-	FD3D11Viewport* Viewport = ResourceCast(ViewportRHI);
-	return Viewport->GetPixelColorSpace();
-}
-
 void  FD3D11DynamicRHI::RHICheckViewportHDRStatus(FRHIViewport* ViewportRHI)
 {
 	FD3D11Viewport* Viewport = ResourceCast(ViewportRHI);
diff --git a/Engine/Source/Runtime/Windows/D3D11RHI/Private/Windows/WindowsD3D11Viewport.cpp b/Engine/Source/Runtime/Windows/D3D11RHI/Private/Windows/WindowsD3D11Viewport.cpp
--- a/Engine/Source/Runtime/Windows/D3D11RHI/Private/Windows/WindowsD3D11Viewport.cpp
+++ b/Engine/Source/Runtime/Windows/D3D11RHI/Private/Windows/WindowsD3D11Viewport.cpp
@@ -44,7 +44,6 @@ FD3D11Viewport::FD3D11Viewport(FD3D11DynamicRHI* InD3DRHI,HWND InWindowHandle,ui
 	PresentFailCount(0),
 	ValidState(0),
 	PixelFormat(InPreferredPixelFormat),
-	PixelColorSpace(EColorSpaceAndEOTF::ERec709_sRGB),
 	DisplayColorGamut(EDisplayColorGamut::sRGB_D65),
 	DisplayOutputFormat(EDisplayOutputFormat::SDR_sRGB),
 	bIsFullscreen(bInIsFullscreen),
@@ -385,7 +384,6 @@ inline void EnsureColorSpace(IDXGISwapChain* SwapChain, EDisplayColorGamut Displ
 void FD3D11Viewport::CheckHDRMonitorStatus()
 {
 #if WITH_EDITOR
-
 	static auto CVarHDREnable = IConsoleManager::Get().FindConsoleVariable(TEXT("Editor.HDRSupport"));
 	if (CVarHDREnable->GetInt() != 0)
 	{
@@ -393,12 +391,6 @@ void FD3D11Viewport::CheckHDRMonitorStatus()
 
 		EnsureColorSpace(SwapChain, DisplayColorGamut, DisplayOutputFormat, PixelFormat);
 	}
-	
-	{
-		PixelColorSpace =  EColorSpaceAndEOTF::ERec709_sRGB;
-	}
-#else
-	PixelColorSpace =  EColorSpaceAndEOTF::ERec709_sRGB;
 #endif
 }
 
diff --git a/Engine/Source/Runtime/Windows/D3D11RHI/Public/D3D11Viewport.h b/Engine/Source/Runtime/Windows/D3D11RHI/Public/D3D11Viewport.h
--- a/Engine/Source/Runtime/Windows/D3D11RHI/Public/D3D11Viewport.h
+++ b/Engine/Source/Runtime/Windows/D3D11RHI/Public/D3D11Viewport.h
@@ -69,7 +69,6 @@ public:
 	// Accessors.
 	FIntPoint GetSizeXY() const { return FIntPoint(SizeX, SizeY); }
 	FD3D11Texture* GetBackBuffer() const { return BackBuffer; }
-	EColorSpaceAndEOTF GetPixelColorSpace() const { return PixelColorSpace; }
 
 	virtual void WaitForFrameEventCompletion() override
 	{
@@ -131,7 +130,6 @@ protected:
 	uint32 PresentFailCount;
 	TAtomic<uint32> ValidState;
 	EPixelFormat PixelFormat;
-	EColorSpaceAndEOTF PixelColorSpace;
 	EDisplayColorGamut DisplayColorGamut;
 	EDisplayOutputFormat DisplayOutputFormat;
 	bool bIsFullscreen;
