Fix for assert on AudioDevice shutdown when there's a active DeviceSwap running

[RN] Fix for assert on AudioDevice shutdown when there's a active device swap running
#tests 10 instances with PIE using cvar au.ThreadedSwapDebugExtraTime 1000000

diff --git a/Engine/Source/Runtime/Windows/AudioMixerXAudio2/Private/AudioMixerPlatformXAudio2.cpp b/Engine/Source/Runtime/Windows/AudioMixerXAudio2/Private/AudioMixerPlatformXAudio2.cpp
--- a/Engine/Source/Runtime/Windows/AudioMixerXAudio2/Private/AudioMixerPlatformXAudio2.cpp
+++ b/Engine/Source/Runtime/Windows/AudioMixerXAudio2/Private/AudioMixerPlatformXAudio2.cpp
@@ -1192,14 +1192,10 @@ namespace Audio
 			return false;
 		}
 
-		// If we have a active device swap in flight, we need to wait for it.
-		static const FTimespan TimeoutOnDeviceSwap(0,0,0,3,0); // 3 secs.
-		if (ActiveDeviceSwap.IsValid() && !ActiveDeviceSwap.IsReady() && !ActiveDeviceSwap.WaitFor(TimeoutOnDeviceSwap))
-		{
-			UE_LOG(LogAudioMixer, Warning, TEXT("Timeout waiting for inflight device-swap. InstanceID=%d"), InstanceID);
-			return false;
-		}
-
+		// If we're closing the stream, we're not interested in the results of the device swap. 
+		// Reset the handle to the future.
+		ActiveDeviceSwap.Reset();
+			
 		if (bIsDeviceOpen && !StopAudioStream())
 		{
 			return false;
