Fix bug in D3D11 RHI ActiveQueries linked list.
 - "Last" pointer was not properly fixed up when the final node in the list is unlinked.

diff --git a/Engine/Source/Runtime/Windows/D3D11RHI/Private/D3D11Query.cpp b/Engine/Source/Runtime/Windows/D3D11RHI/Private/D3D11Query.cpp
--- a/Engine/Source/Runtime/Windows/D3D11RHI/Private/D3D11Query.cpp
+++ b/Engine/Source/Runtime/Windows/D3D11RHI/Private/D3D11Query.cpp
@@ -219,16 +219,32 @@ void FD3D11RenderQuery::Unlink()
 
 	auto& List = FD3D11DynamicRHI::Get().ActiveQueries;
 
+	if (List.Last == this)
+	{
+		// This is the last node in the list, so the "List.Last" pointer needs fixing up.
+		if (Prev == &List.First)
+		{
+			// This is also the first node in the list, meaning there's only 1 node total.
+			// Just clear the "List.Last" pointer.
+			List.Last = nullptr;
+		}
+		else
+		{
+			//
+			// There's at least one real node before us.
+			// 
+			// "Prev" points to the "Next" member field of the previous node.
+			// Subtract the "Next" field offset to get the actual previous node address.
+			//
+			List.Last = reinterpret_cast<FD3D11RenderQuery*>(reinterpret_cast<uintptr_t>(Prev) - UFIELD_OFFSET(FD3D11RenderQuery, Next));
+		}
+	}
+
 	if (Next) { Next->Prev = Prev; }
 	if (Prev) { *Prev = Next; }
 
 	Next = nullptr;
 	Prev = nullptr;
-
-	if (List.Last == this)
-	{
-		List.Last = nullptr;
-	}
 }
 
 void FD3D11DynamicRHI::PollQueryResults()
