Temporarily dial up logging around Suspend/Resume operations and Null Render Audio Device.

#rnx

diff --git a/Engine/Source/Runtime/AudioMixerCore/Private/AudioMixer.cpp b/Engine/Source/Runtime/AudioMixerCore/Private/AudioMixer.cpp
--- a/Engine/Source/Runtime/AudioMixerCore/Private/AudioMixer.cpp
+++ b/Engine/Source/Runtime/AudioMixerCore/Private/AudioMixer.cpp
@@ -456,7 +456,7 @@ namespace Audio
 
 	void IAudioMixerPlatformInterface::StartRunningNullDevice()
 	{
-		UE_LOG(LogAudioMixer, Verbose, TEXT("StartRunningNullDevice() called, InstanceID=%d"), InstanceID);
+		UE_LOG(LogAudioMixer, Display, TEXT("StartRunningNullDevice() called, InstanceID=%d"), InstanceID);
 		SCOPED_NAMED_EVENT(FMixerPlatformXAudio2_StartRunningNullDevice, FColor::Blue);
 		
 		auto ThrowAwayBuffer = [this]() { this->ReadNextBuffer(); };
@@ -485,7 +485,7 @@ namespace Audio
 
 	void IAudioMixerPlatformInterface::StopRunningNullDevice()
 	{		
-		UE_LOG(LogAudioMixer, Verbose, TEXT("StopRunningNullDevice() called, InstanceID=%d"), InstanceID);
+		UE_LOG(LogAudioMixer, Display, TEXT("StopRunningNullDevice() called, InstanceID=%d"), InstanceID);
 		SCOPED_NAMED_EVENT(FMixerPlatformXAudio2_StopRunningNullDevice, FColor::Blue);
 
 		if (NullDeviceCallback.IsValid())
diff --git a/Engine/Source/Runtime/Windows/AudioMixerXAudio2/Private/AudioMixerPlatformXAudio2.cpp b/Engine/Source/Runtime/Windows/AudioMixerXAudio2/Private/AudioMixerPlatformXAudio2.cpp
--- a/Engine/Source/Runtime/Windows/AudioMixerXAudio2/Private/AudioMixerPlatformXAudio2.cpp
+++ b/Engine/Source/Runtime/Windows/AudioMixerXAudio2/Private/AudioMixerPlatformXAudio2.cpp
@@ -539,7 +539,10 @@ namespace Audio
 		{				
 			if( XAudio2System )
 			{
+				const uint64 StartCycles = FPlatformTime::Cycles64();						
 				XAudio2System->StopEngine();
+				UE_LOG(LogAudioMixer, Display, TEXT("XAudio2System StopEngine(), taking '%2.2f' ms"),
+					FPlatformTime::ToMilliseconds64(FPlatformTime::Cycles64() - StartCycles));
 				StartRunningNullDevice();
 				bIsSuspended = true;
 			}					
@@ -576,6 +579,8 @@ namespace Audio
 				if(SUCCEEDED(Result))
 				{		
 					bIsSuspended = false;								
+					UE_CLOG(StartAttempt == 1, LogAudioMixer, Display, 
+						TEXT("XAudio2System StartEngine() - Sucessfully started, taking '%2.2f' ms"), FPlatformTime::ToMilliseconds64(FPlatformTime::Cycles64() - StartCycles));
 					UE_CLOG(StartAttempt > 1, LogAudioMixer, Warning, 
 						TEXT("StartEngine() took %d attempts to start, taking '%f' ms"), 
 						StartAttempt, FPlatformTime::ToMilliseconds64(FPlatformTime::Cycles64() - StartCycles));
