Implement new RHI GPU profiler event stream in D3D11

diff --git a/Engine/Source/Runtime/Windows/D3D11RHI/Private/D3D11Commands.cpp b/Engine/Source/Runtime/Windows/D3D11RHI/Private/D3D11Commands.cpp
--- a/Engine/Source/Runtime/Windows/D3D11RHI/Private/D3D11Commands.cpp
+++ b/Engine/Source/Runtime/Windows/D3D11RHI/Private/D3D11Commands.cpp
@@ -1397,13 +1397,27 @@ IRHIComputeContext* FD3D11DynamicRHI::RHIGetCommandContext(ERHIPipeline Pipeline
 	return nullptr;
 }
 
+struct FD3D11PlatformCommandList : public IRHIPlatformCommandList
+{
+	virtual ~FD3D11PlatformCommandList() = default;
+};
+
 void FD3D11DynamicRHI::RHIFinalizeContext(FRHIFinalizeContextArgs&& Args, TRHIPipelineArray<IRHIPlatformCommandList*>& Output)
 {
+#if RHI_NEW_GPU_PROFILER
+	FlushProfilerStats();
+#endif
+
 	// "Context" will always be the default context, since we don't implement parallel execution.
-	for(IRHIComputeContext* Context : Args.Contexts)
+	for (IRHIComputeContext* Context : Args.Contexts)
 	{
 		// "Context" will always be the default context, since we don't implement parallel execution.
 		check(Context == this);
+
+#if RHI_NEW_GPU_PROFILER && WITH_RHI_BREADCRUMBS
+		// We need platform command lists to contain the breadcrumb allocators
+		Output[Context->GetPipeline()] = new FD3D11PlatformCommandList;
+#endif
 	}
 
 	// Reset some context state
@@ -1422,6 +1436,22 @@ void FD3D11DynamicRHI::RHISubmitCommandLists(FRHISubmitCommandListsArgs&& Args)
 {
 	// Attempt to readback completed queries
 	PollQueryResults();
+
+#if RHI_NEW_GPU_PROFILER && WITH_RHI_BREADCRUMBS
+	for (IRHIPlatformCommandList* CmdList : Args.CommandLists)
+	{
+		FD3D11PlatformCommandList* D3DCmdList = static_cast<FD3D11PlatformCommandList*>(CmdList);
+
+		// Preserve the breadcrumb allocators in the profiler frame
+		// so they are kept alive until the frame's data is processed.
+		for (auto const& Allocator : D3DCmdList->BreadcrumbAllocators)
+		{
+			Profiler.Current.BreadcrumbAllocators.AddUnique(&Allocator.Get());
+		}
+
+		delete D3DCmdList;
+	}
+#endif
 }
 
 void FD3D11DynamicRHI::EnableUAVOverlap()
diff --git a/Engine/Source/Runtime/Windows/D3D11RHI/Private/D3D11Device.cpp b/Engine/Source/Runtime/Windows/D3D11RHI/Private/D3D11Device.cpp
--- a/Engine/Source/Runtime/Windows/D3D11RHI/Private/D3D11Device.cpp
+++ b/Engine/Source/Runtime/Windows/D3D11RHI/Private/D3D11Device.cpp
@@ -267,13 +267,18 @@ void FD3D11DynamicRHI::Shutdown()
 
 		if (ShouldEmitBreadcrumbs())
 		{
+		#if RHI_NEW_GPU_PROFILER
+			FlushProfilerStats();
+
+			auto& Event = EmplaceProfilerEvent<UE::RHI::GPUProfiler::FEvent::FBeginBreadcrumb>(Breadcrumb);
+			InsertProfilerTimestamp(&Event.GPUTimestampTOP);
+		#else
 			// @todo dev-pr avoid TCHAR -> ANSI conversion
 			FRHIBreadcrumb::FBuffer Buffer;
 			TCHAR const* Name = Breadcrumb->Name.GetTCHAR(Buffer);
 
-#if (RHI_NEW_GPU_PROFILER == 0)
 			GPUProfilingData.PushEvent(Name, FColor::White);
-#endif
+		#endif
 		}
 	}
 
@@ -281,9 +286,14 @@ void FD3D11DynamicRHI::Shutdown()
 	{
 		if (ShouldEmitBreadcrumbs())
 		{
-#if (RHI_NEW_GPU_PROFILER == 0)
+		#if RHI_NEW_GPU_PROFILER
+			FlushProfilerStats();
+
+			auto& Event = EmplaceProfilerEvent<UE::RHI::GPUProfiler::FEvent::FEndBreadcrumb>(Breadcrumb);
+			InsertProfilerTimestamp(&Event.GPUTimestampBOP);
+		#else
 			GPUProfilingData.PopEvent();
-#endif
+		#endif
 		}
 
 	#if NV_AFTERMATH
@@ -547,6 +557,24 @@ void FD3D11DynamicRHI::SetupAfterDeviceCreation()
 			UE_LOG(LogD3D11RHI, Log, TEXT("Array index from any shader is supported"));
 		}
 	}
+
+	TimestampCalibration = CalibrateTimers();
+
+#if RHI_NEW_GPU_PROFILER
+	// Register the single graphics GPU queue we have access to in D3D11.
+	UE::RHI::GPUProfiler::FQueue Queue(UE::RHI::GPUProfiler::FQueue::EType::Graphics, 0, 0);
+	UE::RHI::GPUProfiler::InitializeQueues(MakeConstArrayView(&Queue, 1));
+
+	//
+	// Since we can't tell when the GPU is actually executing engine work in D3D11,
+	// just mark the GPU as always busy in the frame. We also push begin/end work
+	// markers in RHIEndFrame either side of the frame boundary.
+	//
+	{
+		auto& Event = EmplaceProfilerEvent<UE::RHI::GPUProfiler::FEvent::FBeginWork>(FPlatformTime::Cycles64());
+		InsertProfilerTimestamp(&Event.GPUTimestampTOP);
+	}
+#endif
 }
 
 void FD3D11DynamicRHI::UpdateMSAASettings()
diff --git a/Engine/Source/Runtime/Windows/D3D11RHI/Private/D3D11Query.cpp b/Engine/Source/Runtime/Windows/D3D11RHI/Private/D3D11Query.cpp
--- a/Engine/Source/Runtime/Windows/D3D11RHI/Private/D3D11Query.cpp
+++ b/Engine/Source/Runtime/Windows/D3D11RHI/Private/D3D11Query.cpp
@@ -23,20 +23,60 @@ static FAutoConsoleVariableRef CVarD3D11QueryTimeoutValue(
 	ECVF_Default
 );
 
-FRenderQueryRHIRef FD3D11DynamicRHI::RHICreateRenderQuery(ERenderQueryType QueryType)
+FD3D11RenderQuery::FD3D11RenderQuery(EType Type)
+	: Type(Type)
 {
-	TRefCountPtr<ID3D11Query> Query;
 	D3D11_QUERY_DESC Desc {};
 
-	switch (QueryType)
+	switch (Type)
 	{
-		case RQT_Occlusion   : Desc.Query = D3D11_QUERY_OCCLUSION; break;
-		case RQT_AbsoluteTime: Desc.Query = D3D11_QUERY_TIMESTAMP; break;
-		default: checkNoEntry(); return nullptr;
+	case EType::Occlusion: Desc.Query = D3D11_QUERY_OCCLUSION; break;
+	case EType::Timestamp: Desc.Query = D3D11_QUERY_TIMESTAMP; break;
+	case EType::Profiler : Desc.Query = D3D11_QUERY_TIMESTAMP; break;
+
+	default:
+		checkNoEntry();
+		return;
 	}
 
-	VERIFYD3D11RESULT_EX(Direct3DDevice->CreateQuery(&Desc, Query.GetInitReference()), Direct3DDevice);
-	return new FD3D11RenderQuery(Query, QueryType);
+	ID3D11Device* Device = FD3D11DynamicRHI::Get().GetDevice();
+	VERIFYD3D11RESULT_EX(Device->CreateQuery(&Desc, Resource.GetInitReference()), Device);
+}
+
+FD3D11RenderQuery::~FD3D11RenderQuery()
+{
+	Unlink();
+}
+
+void FD3D11RenderQuery::Begin(ID3D11DeviceContext* Context)
+{
+	check(Type == EType::Occlusion);
+	check(State == EState::None);
+
+	Context->Begin(Resource);
+}
+
+void FD3D11RenderQuery::End(ID3D11DeviceContext* Context, uint64* NewTarget)
+{
+	check(State == EState::None);
+	Context->End(Resource);
+	Target = NewTarget;
+
+	Link();
+	State = FD3D11RenderQuery::EState::Ended;
+}
+
+FRenderQueryRHIRef FD3D11DynamicRHI::RHICreateRenderQuery(ERenderQueryType QueryType)
+{
+	switch (QueryType)
+	{
+	case RQT_AbsoluteTime: return new FD3D11RenderQuery_RHI(FD3D11RenderQuery::EType::Timestamp);
+	case RQT_Occlusion   : return new FD3D11RenderQuery_RHI(FD3D11RenderQuery::EType::Occlusion);
+
+	default:
+		checkNoEntry();
+		return nullptr;
+	}
 }
 
 void FD3D11DynamicRHI::RHIBeginRenderQuery_TopOfPipe(FRHICommandListBase& RHICmdList, FRHIRenderQuery* RenderQuery)
@@ -57,33 +97,24 @@ void FD3D11DynamicRHI::RHIEndRenderQuery_TopOfPipe(FRHICommandListBase& RHICmdLi
 
 void FD3D11DynamicRHI::RHIBeginRenderQuery(FRHIRenderQuery* RenderQuery)
 {
-	FD3D11RenderQuery* Query = ResourceCast(RenderQuery);
-	check(Query->QueryType == RQT_Occlusion);
-	check(Query->State == FD3D11RenderQuery::EState::None);
-
-	Direct3DDeviceIMContext->Begin(Query->Resource.GetReference());
+	ResourceCast(RenderQuery)->Begin(Direct3DDeviceIMContext);
 }
 
 void FD3D11DynamicRHI::RHIEndRenderQuery(FRHIRenderQuery* RenderQuery)
 {
-	FD3D11RenderQuery* Query = ResourceCast(RenderQuery);
-	check(Query->State == FD3D11RenderQuery::EState::None);
-
-	Direct3DDeviceIMContext->End(Query->Resource);
-
-	Query->Link();
-	Query->State = FD3D11RenderQuery::EState::Ended;
+	FD3D11RenderQuery_RHI* Query = ResourceCast(RenderQuery);
+	Query->End(Direct3DDeviceIMContext, &Query->Result);
 }
 
 bool FD3D11DynamicRHI::RHIGetRenderQueryResult(FRHIRenderQuery* QueryRHI, uint64& OutResult, bool bWait, uint32 GPUIndex)
 {
 	check(IsInRenderingThread());
-	FD3D11RenderQuery* Query = ResourceCast(QueryRHI);
+	FD3D11RenderQuery_RHI* Query = ResourceCast(QueryRHI);
 
 	bool bRHIThreadFlushed = false;
 
 Retry:
-	if (Query->State == FD3D11RenderQuery::EState::Cached)
+	if (Query->State == FD3D11RenderQuery::EState::Completed)
 	{
 		// Early return for queries we already have the result for.
 		check(!Query->IsLinked());
@@ -96,7 +127,7 @@ Retry:
 		if (!bWait)
 		{
 			//
-			// The RHI thread is still processing work, the query has not yet been cached, and we don't want to wait for the query result.
+			// The RHI thread is still processing work, the query has not yet completed, and we don't want to wait for the query result.
 			// Return. The RHI thread will poll for results later.
 			//
 			OutResult = 0;
@@ -105,7 +136,7 @@ Retry:
 		else
 		{
 			//
-			// The RHI thread is active, the query has not yet been cached, and we want to wait for results.
+			// The RHI thread is active, the query has not yet completed, and we want to wait for results.
 			// 
 			// Flushing the RHI thread will ensure a query poll operation has happened before the render thread resumes, which might successfully cache the results.
 			// It will also make it safe for us to use the immediate device context in case the query still wasn't done when the RHI thread last polled for results.
@@ -140,7 +171,7 @@ Retry:
 		return false;
 	}
 
-	check(Query->State == FD3D11RenderQuery::EState::Cached);
+	check(Query->State == FD3D11RenderQuery::EState::Completed);
 	check(!Query->IsLinked());
 	OutResult = Query->Result;
 
@@ -149,38 +180,67 @@ Retry:
 
 bool FD3D11RenderQuery::CacheResult(FD3D11DynamicRHI& RHI, bool bWait)
 {
-	if (State == FD3D11RenderQuery::EState::Cached)
+	if (State == FD3D11RenderQuery::EState::Completed)
 	{
 		check(!IsLinked());
 		return true;
 	}
 
+	check(Target);
+
 	// Attempt to read the result from the GPU.
 	uint64 Temp;
-	if (!RHI.GetQueryData(Resource, &Temp, sizeof(Temp), QueryType, /*bWait = */ bWait, /*bStallRHIThread = */ false))
+	if (!RHI.GetQueryData(Resource, &Temp, sizeof(Temp), Type == EType::Timestamp, /*bWait = */ bWait, /*bStallRHIThread = */ false))
 	{
 		return false;
 	}
 
 	// Data retrieved.
 	// Adjust timer queries to engine-clock ticks.
-	if (QueryType == RQT_AbsoluteTime)
+	switch (Type)
 	{
 #if RHI_NEW_GPU_PROFILER
-		checkNoEntry(); // @todo - new gpu profiler
-#else
-		// GetTimingFrequency is the number of ticks per second
-		uint64 Div = FMath::Max(1llu, FGPUTiming::GetTimingFrequency() / (1000 * 1000));
+	case EType::Profiler:
+		{
+			// Convert from GPU timestamp to CPU timestamp (relative to FPlatformTime::Cycles64())
+			uint64 GPUDelta = Temp - RHI.TimestampCalibration->GPUTimestamp;
+			uint64 CPUDelta = (GPUDelta * RHI.TimestampCalibration->CPUFrequency) / RHI.TimestampCalibration->GPUFrequency;
 
-		// convert from GPU specific timestamp to micro sec (1 / 1 000 000 s) which seems a reasonable resolution
-		Temp = Temp / Div;
+			Temp = CPUDelta + RHI.TimestampCalibration->CPUTimestamp;
+		}
+		break;
 #endif
+
+	case EType::Timestamp:
+		{
+			// GetTimingFrequency is the number of ticks per second
+			uint64 Div = FMath::Max(1llu, RHI.TimestampCalibration->GPUFrequency / (1000 * 1000));
+
+			// convert from GPU specific timestamp to micro sec (1 / 1 000 000 s) which seems a reasonable resolution
+			Temp = Temp / Div;
+		}
+		break;
 	}
 
-	Result = Temp;
+	*Target = Temp;
+	Target = nullptr;
+
 	Unlink();
 
-	State = FD3D11RenderQuery::EState::Cached;
+#if RHI_NEW_GPU_PROFILER
+	if (Type == EType::Profiler)
+	{
+		State = FD3D11RenderQuery::EState::None;
+
+		// Return the query to the pool
+		RHI.Profiler.TimestampPool.Push(this);
+	}
+	else
+#endif
+	{
+		State = FD3D11RenderQuery::EState::Completed;
+	}
+
 	return true;
 }
 
@@ -256,7 +316,7 @@ void FD3D11DynamicRHI::PollQueryResults()
 	}
 }
 
-bool FD3D11DynamicRHI::GetQueryData(ID3D11Query* Query, void* Data, SIZE_T DataSize, ERenderQueryType QueryType, bool bWait, bool bStallRHIThread)
+bool FD3D11DynamicRHI::GetQueryData(ID3D11Query* Query, void* Data, SIZE_T DataSize, bool bTimestamp, bool bWait, bool bStallRHIThread)
 {
 	// Request the data from the query.
 	HRESULT Result;
@@ -277,7 +337,7 @@ bool FD3D11DynamicRHI::GetQueryData(ID3D11Query* Query, void* Data, SIZE_T DataS
 		double StartTime = FPlatformTime::Seconds();
 		double TimeoutWarningLimit = 5.0;
 		// timer queries are used for Benchmarks which can stall a bit more
-		double TimeoutValue = (QueryType == RQT_AbsoluteTime) ? GD3D11AbsoluteTimeQueryTimeoutValue : GD3D11QueryTimeoutValue;
+		double TimeoutValue = bTimestamp ? GD3D11AbsoluteTimeQueryTimeoutValue : GD3D11QueryTimeoutValue;
 
 		do
 		{
@@ -341,7 +401,7 @@ void FD3D11EventQuery::WaitForCompletion()
 {
 	BOOL bRenderingIsFinished = false;
 	while(
-		D3DRHI->GetQueryData(Query,&bRenderingIsFinished,sizeof(bRenderingIsFinished),RQT_Undefined,true,true) &&
+		D3DRHI->GetQueryData(Query, &bRenderingIsFinished, sizeof(bRenderingIsFinished), false, true, true) &&
 		!bRenderingIsFinished
 		)
 	{};
@@ -356,6 +416,123 @@ FD3D11EventQuery::FD3D11EventQuery(class FD3D11DynamicRHI* InD3DRHI):
 	VERIFYD3D11RESULT_EX(D3DRHI->GetDevice()->CreateQuery(&QueryDesc,Query.GetInitReference()), D3DRHI->GetDevice());
 }
 
+
+TOptional<FD3D11DynamicRHI::FTimestampCalibration> FD3D11DynamicRHI::CalibrateTimers()
+{
+	// Attempt to generate a timestamp on GPU and CPU as closely to each other as possible.
+	// This works by first flushing any pending GPU work, then writing a GPU timestamp and waiting for GPU to finish.
+	// CPU timestamp is continuously captured while we are waiting on GPU.
+
+	HRESULT D3DResult = E_FAIL;
+
+	TRefCountPtr<ID3D11Query> DisjointQuery;
+	{
+		D3D11_QUERY_DESC QueryDesc;
+		QueryDesc.Query = D3D11_QUERY_TIMESTAMP_DISJOINT;
+		QueryDesc.MiscFlags = 0;
+		D3DResult = Direct3DDevice->CreateQuery(&QueryDesc, DisjointQuery.GetInitReference());
+
+		if (D3DResult != S_OK)
+			return {};
+	}
+
+	TRefCountPtr<ID3D11Query> TimestampQuery;
+	{
+		D3D11_QUERY_DESC QueryDesc;
+		QueryDesc.Query = D3D11_QUERY_TIMESTAMP;
+		QueryDesc.MiscFlags = 0;
+		D3DResult = Direct3DDevice->CreateQuery(&QueryDesc, TimestampQuery.GetInitReference());
+
+		if (D3DResult != S_OK)
+			return {};
+	}
+
+	TRefCountPtr<ID3D11Query> PendingWorkDoneQuery;
+	TRefCountPtr<ID3D11Query> TimestampDoneQuery;
+	{
+		D3D11_QUERY_DESC QueryDesc;
+		QueryDesc.Query = D3D11_QUERY_EVENT;
+		QueryDesc.MiscFlags = 0;
+
+		D3DResult = Direct3DDevice->CreateQuery(&QueryDesc, PendingWorkDoneQuery.GetInitReference());
+		if (D3DResult != S_OK)
+			return {};
+
+		D3DResult = Direct3DDevice->CreateQuery(&QueryDesc, TimestampDoneQuery.GetInitReference());
+		if (D3DResult != S_OK)
+			return {};
+	}
+
+	// Flush any currently pending GPU work and wait for it to finish
+	Direct3DDeviceIMContext->End(PendingWorkDoneQuery);
+	Direct3DDeviceIMContext->Flush();
+
+	for (;;)
+	{
+		BOOL EventComplete = false;
+		Direct3DDeviceIMContext->GetData(PendingWorkDoneQuery, &EventComplete, sizeof(EventComplete), 0);
+
+		if (EventComplete)
+			break;
+
+		FPlatformProcess::Sleep(0.001f);
+	}
+
+	const uint32 MaxCalibrationAttempts = 10;
+	for (uint32 CalibrationAttempt = 0; CalibrationAttempt < MaxCalibrationAttempts; ++CalibrationAttempt)
+	{
+		Direct3DDeviceIMContext->Begin(DisjointQuery);
+		Direct3DDeviceIMContext->End(TimestampQuery);
+		Direct3DDeviceIMContext->End(DisjointQuery);
+		Direct3DDeviceIMContext->End(TimestampDoneQuery);
+
+		Direct3DDeviceIMContext->Flush();
+
+		uint64 CPUTimestamp = 0;
+		uint64 GPUTimestamp = 0;
+
+		// Busy-wait for GPU to finish and capture CPU timestamp approximately when GPU work is done
+		for (;;)
+		{
+			BOOL EventComplete = false;
+
+			CPUTimestamp = FPlatformTime::Cycles64();
+			Direct3DDeviceIMContext->GetData(TimestampDoneQuery, &EventComplete, sizeof(EventComplete), 0);
+
+			if (EventComplete)
+				break;
+		}
+
+		D3D11_QUERY_DATA_TIMESTAMP_DISJOINT DisjointQueryData = {};
+		D3DResult = Direct3DDeviceIMContext->GetData(DisjointQuery, &DisjointQueryData, sizeof(DisjointQueryData), 0);
+
+		// If timestamp was unreliable, try again
+		if (D3DResult != S_OK || DisjointQueryData.Disjoint)
+		{
+			continue;
+		}
+
+		D3DResult = Direct3DDeviceIMContext->GetData(TimestampQuery, &GPUTimestamp, sizeof(GPUTimestamp), 0);
+
+		// If we managed to get valid timestamps, save both of them (CPU & GPU) and return
+		if (D3DResult == S_OK && GPUTimestamp)
+		{
+			return FD3D11DynamicRHI::FTimestampCalibration
+			{
+				.CPUTimestamp = CPUTimestamp,
+				.CPUFrequency = uint64(1.0 / FPlatformTime::GetSecondsPerCycle64()),
+
+				.GPUTimestamp = GPUTimestamp,
+				.GPUFrequency = DisjointQueryData.Frequency
+			};
+		}
+	}
+
+	return {};
+}
+
+
+
 /*=============================================================================
  * class FD3D11BufferedGPUTiming
  *=============================================================================*/
@@ -450,106 +627,13 @@ void FD3D11BufferedGPUTiming::PlatformStaticInitialize(void* UserData)
 
 void FD3D11BufferedGPUTiming::CalibrateTimers(FD3D11DynamicRHI* InD3DRHI)
 {
-	// Attempt to generate a timestamp on GPU and CPU as closely to each other as possible.
-	// This works by first flushing any pending GPU work, then writing a GPU timestamp and waiting for GPU to finish.
-	// CPU timestamp is continuously captured while we are waiting on GPU.
-
-	HRESULT D3DResult = E_FAIL;
-
-	TRefCountPtr<ID3D11Query> DisjointQuery;
-	{
-		D3D11_QUERY_DESC QueryDesc;
-		QueryDesc.Query = D3D11_QUERY_TIMESTAMP_DISJOINT;
-		QueryDesc.MiscFlags = 0;
-		D3DResult = InD3DRHI->GetDevice()->CreateQuery(&QueryDesc, DisjointQuery.GetInitReference());
-
-		if (D3DResult != S_OK) return;
-	}
-
-	TRefCountPtr<ID3D11Query> TimestampQuery;
-	{
-		D3D11_QUERY_DESC QueryDesc;
-		QueryDesc.Query = D3D11_QUERY_TIMESTAMP;
-		QueryDesc.MiscFlags = 0;
-		D3DResult = InD3DRHI->GetDevice()->CreateQuery(&QueryDesc, TimestampQuery.GetInitReference());
-
-		if (D3DResult != S_OK) return;
-	}
-
-	TRefCountPtr<ID3D11Query> PendingWorkDoneQuery;
-	TRefCountPtr<ID3D11Query> TimestampDoneQuery;
-	{
-		D3D11_QUERY_DESC QueryDesc;
-		QueryDesc.Query = D3D11_QUERY_EVENT;
-		QueryDesc.MiscFlags = 0;
-
-		D3DResult = InD3DRHI->GetDevice()->CreateQuery(&QueryDesc, PendingWorkDoneQuery.GetInitReference());
-		if (D3DResult != S_OK) return;
-
-		D3DResult = InD3DRHI->GetDevice()->CreateQuery(&QueryDesc, TimestampDoneQuery.GetInitReference());
-		if (D3DResult != S_OK) return;
-	}
-
-	ID3D11DeviceContext* D3D11DeviceContext = InD3DRHI->GetDeviceContext();
-
-	// Flush any currently pending GPU work and wait for it to finish
-	D3D11DeviceContext->End(PendingWorkDoneQuery);
-	D3D11DeviceContext->Flush();
-
-	for(;;)
+	TOptional<FD3D11DynamicRHI::FTimestampCalibration> Data = InD3DRHI->CalibrateTimers();
+	if (Data.IsSet())
 	{
-		BOOL EventComplete = false;
-		D3D11DeviceContext->GetData(PendingWorkDoneQuery, &EventComplete, sizeof(EventComplete), 0);
-		if (EventComplete) break;
-		FPlatformProcess::Sleep(0.001f);
-	}
-
-	const uint32 MaxCalibrationAttempts = 10;
-	for (uint32 CalibrationAttempt = 0; CalibrationAttempt < MaxCalibrationAttempts; ++CalibrationAttempt)
-	{
-		D3D11DeviceContext->Begin(DisjointQuery);
-		D3D11DeviceContext->End(TimestampQuery);
-		D3D11DeviceContext->End(DisjointQuery);
-		D3D11DeviceContext->End(TimestampDoneQuery);
-
-		D3D11DeviceContext->Flush();
-
-		uint64 CPUTimestamp = 0;
-		uint64 GPUTimestamp = 0;
-
-		// Busy-wait for GPU to finish and capture CPU timestamp approximately when GPU work is done
-		for (;;)
-		{
-			BOOL EventComplete = false;
-			CPUTimestamp = FPlatformTime::Cycles64();
-			D3D11DeviceContext->GetData(TimestampDoneQuery, &EventComplete, sizeof(EventComplete), 0);
-			if (EventComplete) break;
-		}
-
-		D3D11_QUERY_DATA_TIMESTAMP_DISJOINT DisjointQueryData = {};
-		D3DResult = D3D11DeviceContext->GetData(DisjointQuery, &DisjointQueryData, sizeof(DisjointQueryData), 0);
-
-		// If timestamp was unreliable, try again
-		if (D3DResult != S_OK || DisjointQueryData.Disjoint)
-		{
-			continue;
-		}
-
-		D3DResult = D3D11DeviceContext->GetData(TimestampQuery, &GPUTimestamp, sizeof(GPUTimestamp), 0);
-
-		// If we managed to get valid timestamps, save both of them (CPU & GPU) and return
-		if (D3DResult == S_OK && GPUTimestamp)
-		{
-			FGPUTimingCalibrationTimestamp CalibrationTimestamp;
-			CalibrationTimestamp.CPUMicroseconds = uint64(FPlatformTime::ToSeconds64(CPUTimestamp) * 1e6);
-			CalibrationTimestamp.GPUMicroseconds = uint64(GPUTimestamp * (1e6 / GetTimingFrequency()));
-			SetCalibrationTimestamp(CalibrationTimestamp);
-			break;
-		}
-		else
-		{
-			continue;
-		}
+		FGPUTimingCalibrationTimestamp CalibrationTimestamp;
+		CalibrationTimestamp.CPUMicroseconds = uint64(FPlatformTime::ToSeconds64(Data->CPUTimestamp) * 1e6);
+		CalibrationTimestamp.GPUMicroseconds = uint64(Data->GPUTimestamp * (1e6 / Data->GPUFrequency));
+		SetCalibrationTimestamp(CalibrationTimestamp);
 	}
 }
 
diff --git a/Engine/Source/Runtime/Windows/D3D11RHI/Private/D3D11RHI.cpp b/Engine/Source/Runtime/Windows/D3D11RHI/Private/D3D11RHI.cpp
--- a/Engine/Source/Runtime/Windows/D3D11RHI/Private/D3D11RHI.cpp
+++ b/Engine/Source/Runtime/Windows/D3D11RHI/Private/D3D11RHI.cpp
@@ -32,9 +32,68 @@ extern void UniformBufferBeginFrame();
 void FD3D11DynamicRHI::RHIEndFrame(const FRHIEndFrameArgs& Args)
 {
 	// End Frame
-#if (RHI_NEW_GPU_PROFILER == 0)
+#if RHI_NEW_GPU_PROFILER
+	{
+		// End GPU work
+		auto& EndWork = EmplaceProfilerEvent<UE::RHI::GPUProfiler::FEvent::FEndWork>();
+		InsertProfilerTimestamp(&EndWork.GPUTimestampBOP);
+
+		// Insert frame boundary
+		EmplaceProfilerEvent<UE::RHI::GPUProfiler::FEvent::FFrameBoundary>(Args.FrameNumber
+	#if WITH_RHI_BREADCRUMBS
+			, Args.GPUBreadcrumbs[ERHIPipeline::Graphics]
+	#endif
+		);
+
+		// Issue a completion query so we know when to readback these profiler results.
+		{
+			if (Profiler.EventPool.IsEmpty())
+			{
+				D3D11_QUERY_DESC QueryDesc {};
+				QueryDesc.Query = D3D11_QUERY_EVENT;
+				VERIFYD3D11RESULT(Direct3DDevice->CreateQuery(&QueryDesc, Profiler.Current.CompletionQuery.GetInitReference()));
+			}
+			else
+			{
+				Profiler.Current.CompletionQuery = Profiler.EventPool.Pop();
+			}
+
+			Direct3DDeviceIMContext->End(Profiler.Current.CompletionQuery);
+			Profiler.Pending.Enqueue(MakeUnique<FProfiler::FFrame>(MoveTemp(Profiler.Current)));
+		}
+
+		// Attempt to process historic results
+		while (TUniquePtr<FProfiler::FFrame>* PreviousFramePtr = Profiler.Pending.Peek())
+		{
+			TUniquePtr<FProfiler::FFrame>& PreviousFrame = *PreviousFramePtr;
+
+			BOOL EventComplete = false;
+			VERIFYD3D11RESULT(Direct3DDeviceIMContext->GetData(PreviousFrame->CompletionQuery, &EventComplete, sizeof(EventComplete), 0));
+			if (!EventComplete)
+			{
+				// Frame not yet finished on the GPU
+				break;
+			}
+
+			// Ensure we have the latest timestamp data
+			PollQueryResults();
+
+			// Previous frame has completed and the data is available. Publish the profiler events.
+			UE::RHI::GPUProfiler::ProcessEvents(UE::RHI::GPUProfiler::FQueue(UE::RHI::GPUProfiler::FQueue::EType::Graphics, 0, 0), PreviousFrame->EventStream);
+
+			Profiler.EventPool.Push(MoveTemp(PreviousFrame->CompletionQuery));
+			Profiler.Pending.Pop();
+		}
+
+		// Start the next frame's GPU work
+		auto& BeginWork = EmplaceProfilerEvent<UE::RHI::GPUProfiler::FEvent::FBeginWork>(FPlatformTime::Cycles64());
+		InsertProfilerTimestamp(&BeginWork.GPUTimestampTOP);
+	}
+
+#else
 	GPUProfilingData.EndFrame();
 #endif
+
 	UpdateMemoryStats();
 	CurrentComputeShader = nullptr;
 
@@ -638,3 +697,20 @@ void FD3D11DynamicRHI::RHIVerifyResult(ID3D11Device* Device, HRESULT Result, con
 {
 	VerifyD3D11Result(Result, Code, Filename, Line, Device);
 }
+
+#if RHI_NEW_GPU_PROFILER
+void FD3D11DynamicRHI::InsertProfilerTimestamp(uint64* Target)
+{
+	FD3D11RenderQuery* Query;
+	if (Profiler.TimestampPool.IsEmpty())
+	{
+		Query = new FD3D11RenderQuery(FD3D11RenderQuery::EType::Profiler);
+	}
+	else
+	{
+		Query = Profiler.TimestampPool.Pop();
+	}
+
+	Query->End(Direct3DDeviceIMContext, Target);
+}
+#endif // RHI_NEW_GPU_PROFILER
diff --git a/Engine/Source/Runtime/Windows/D3D11RHI/Private/D3D11RHIPrivate.h b/Engine/Source/Runtime/Windows/D3D11RHI/Private/D3D11RHIPrivate.h
--- a/Engine/Source/Runtime/Windows/D3D11RHI/Private/D3D11RHIPrivate.h
+++ b/Engine/Source/Runtime/Windows/D3D11RHI/Private/D3D11RHIPrivate.h
@@ -424,7 +424,7 @@ public:
 	 * @param bStallRHIThread - if true, stall RHIT before accessing immediate context
 	 * @return true if the query finished.
 	 */
-	bool GetQueryData(ID3D11Query* Query, void* Data, SIZE_T DataSize, ERenderQueryType QueryType, bool bWait, bool bStallRHIThread);
+	bool GetQueryData(ID3D11Query* Query, void* Data, SIZE_T DataSize, bool bTimestamp, bool bWait, bool bStallRHIThread);
 
 	virtual FSamplerStateRHIRef RHICreateSamplerState(const FSamplerStateInitializerRHI& Initializer) final override;
 	virtual FRasterizerStateRHIRef RHICreateRasterizerState(const FRasterizerStateInitializerRHI& Initializer) final override;
@@ -793,6 +793,64 @@ protected:
 	} ActiveQueries;
 	friend class FD3D11RenderQuery;
 
+	struct FProfiler
+	{
+		struct FFrame
+		{
+			TRefCountPtr<ID3D11Query> CompletionQuery;
+
+		#if RHI_NEW_GPU_PROFILER
+			UE::RHI::GPUProfiler::FEventStream EventStream;
+		#endif
+
+		#if WITH_RHI_BREADCRUMBS
+			FRHIBreadcrumbAllocatorArray BreadcrumbAllocators{};
+		#endif
+		};
+
+		FFrame Current;
+		TQueue<TUniquePtr<FFrame>> Pending;
+
+		TArray<FD3D11RenderQuery*> TimestampPool;
+		TArray<TRefCountPtr<ID3D11Query>> EventPool;
+
+	} Profiler;
+
+#if RHI_NEW_GPU_PROFILER
+	template <typename TEventType, typename... TArgs>
+	TEventType& EmplaceProfilerEvent(TArgs&&... Args)
+	{
+		return Profiler.Current.EventStream.Emplace<TEventType>(Forward<TArgs>(Args)...);
+	}
+
+	void InsertProfilerTimestamp(uint64* Target);
+
+	void FlushProfilerStats()
+	{
+		// Flush accumulated draw stats
+		if (StatEvent)
+		{
+			EmplaceProfilerEvent<UE::RHI::GPUProfiler::FEvent::FStats>() = StatEvent;
+			StatEvent = {};
+		}
+	}
+#endif // RHI_NEW_GPU_PROFILER
+
+public:
+	struct FTimestampCalibration
+	{
+		uint64 CPUTimestamp = 0;
+		uint64 CPUFrequency = 0;
+
+		uint64 GPUTimestamp = 0;
+		uint64 GPUFrequency = 0;
+	};
+
+	TOptional<FTimestampCalibration> CalibrateTimers();
+
+protected:
+	TOptional<FTimestampCalibration> TimestampCalibration;
+
 	/** D3D11 defines a maximum of 14 constant buffers per shader stage. */
 	enum { MAX_UNIFORM_BUFFERS_PER_SHADER_STAGE = 14 };
 
@@ -832,8 +890,8 @@ protected:
 
 public:
 #if RHI_NEW_GPU_PROFILER
-	void RegisterGPUWork(uint32 NumPrimitives = 0, uint32 NumVertices = 0)	{ checkNoEntry(); } // @todo - new gpu profiler
-	void RegisterGPUDispatch(FIntVector GroupCount)	                        { checkNoEntry(); } // @todo - new gpu profiler
+	void RegisterGPUWork(uint32 NumPrimitives = 0, uint32 NumVertices = 0)	{ /*checkNoEntry();*/ } // @todo - new gpu profiler
+	void RegisterGPUDispatch(FIntVector GroupCount)	                        { /*checkNoEntry();*/ } // @todo - new gpu profiler
 #else
 	void RegisterGPUWork(uint32 NumPrimitives = 0, uint32 NumVertices = 0)	{ GPUProfilingData.RegisterGPUWork(NumPrimitives, NumVertices); }
 	void RegisterGPUDispatch(FIntVector GroupCount)	                        { GPUProfilingData.RegisterGPUDispatch(GroupCount); }
diff --git a/Engine/Source/Runtime/Windows/D3D11RHI/Public/D3D11Resources.h b/Engine/Source/Runtime/Windows/D3D11RHI/Public/D3D11Resources.h
--- a/Engine/Source/Runtime/Windows/D3D11RHI/Public/D3D11Resources.h
+++ b/Engine/Source/Runtime/Windows/D3D11RHI/Public/D3D11Resources.h
@@ -305,50 +305,60 @@ private:
 	uint8 bAlias : 1;
 };
 
-
-/** D3D11 render query */
-class FD3D11RenderQuery : public FRHIRenderQuery
+class FD3D11RenderQuery
 {
-public:
-	/** The query resource. */
+private:
 	TRefCountPtr<ID3D11Query> Resource;
 
-	/** The cached query result. */
-	uint64 Result = 0;
+	// Location the result is written to.
+	uint64* Target = nullptr;
+
+	// Linked list pointers. Used to build a list of "active" queries, i.e. queries that need data to be polled from the GPU.
+	FD3D11RenderQuery** Prev = nullptr;
+	FD3D11RenderQuery* Next = nullptr;
 
-	enum class EState
+public:
+	enum class EState : uint8
 	{
 		None,
 		Ended,
-		Cached
+		Completed
 	};
 	std::atomic<EState> State { EState::None };
 
-	// todo: memory optimize
-	ERenderQueryType QueryType;
-
-	// Linked list pointers. Used to build a list of "active" queries, i.e. queries that need data to be polled from the GPU.
-	FD3D11RenderQuery** Prev = nullptr;
-	FD3D11RenderQuery* Next = nullptr;
-
-	/** Initialization constructor. */
-	FD3D11RenderQuery(ID3D11Query* InResource, ERenderQueryType InQueryType)
-		: Resource(InResource)
-		, QueryType(InQueryType)
-	{}
-
-	virtual ~FD3D11RenderQuery()
+	enum class EType : uint8
 	{
-		Unlink();
-	}
+		Timestamp,
+		Occlusion,
+		Profiler
+	} const Type;
+
+	FD3D11RenderQuery(EType Type);
+	~FD3D11RenderQuery();
 
 	bool CacheResult(class FD3D11DynamicRHI& RHI, bool bWait);
 
+	void Begin(ID3D11DeviceContext* Context);
+	void End(ID3D11DeviceContext* Context, uint64* Target);
+
 	bool IsLinked() const { return Prev != nullptr; }
+
+private:
 	void Link();
 	void Unlink();
 };
 
+/** D3D11 render query */
+class FD3D11RenderQuery_RHI : public FRHIRenderQuery, public FD3D11RenderQuery
+{
+public:
+	uint64 Result = 0;
+
+	FD3D11RenderQuery_RHI(EType Type)
+		: FD3D11RenderQuery(Type)
+	{}
+};
+
 /** Forward declare the constants ring buffer. */
 class FD3D11ConstantsRingBuffer;
 
@@ -537,7 +547,7 @@ struct TD3D11ResourceTraits<FRHIBoundShaderState>
 template<>
 struct TD3D11ResourceTraits<FRHIRenderQuery>
 {
-	typedef FD3D11RenderQuery TConcreteType;
+	typedef FD3D11RenderQuery_RHI TConcreteType;
 };
 template<>
 struct TD3D11ResourceTraits<FRHIUniformBuffer>
