remove driver specific code for HDR since UE now requires 1703/DXGI 1.6
Older version should not crash, but won't be able to run HDR

#rnx

diff --git a/Engine/Source/Runtime/Engine/Private/GameUserSettings.cpp b/Engine/Source/Runtime/Engine/Private/GameUserSettings.cpp
--- a/Engine/Source/Runtime/Engine/Private/GameUserSettings.cpp
+++ b/Engine/Source/Runtime/Engine/Private/GameUserSettings.cpp
@@ -993,16 +993,6 @@ void UGameUserSettings::EnableHDRDisplayOutputInternal(bool bEnable, int32 Displ
 		// Apply device-specific output encoding
 		if (bEnable)
 		{
-#if PLATFORM_WINDOWS
-			if (GRHIHDRNeedsVendorExtensions)
-			{
-				// Force exclusive fullscreen
-				SetPreferredFullscreenMode(0);
-				SetFullscreenMode(GetPreferredFullscreenMode());
-				ApplyResolutionSettings(false);
-				RequestUIUpdate();
-			}
-#endif
 			CVarHDROutputEnabled->Set(1, bFromUserSettings ? ECVF_SetByGameSetting : ECVF_SetByCode);
 		}
 
diff --git a/Engine/Source/Runtime/RHI/Public/RHIGlobals.h b/Engine/Source/Runtime/RHI/Public/RHIGlobals.h
--- a/Engine/Source/Runtime/RHI/Public/RHIGlobals.h
+++ b/Engine/Source/Runtime/RHI/Public/RHIGlobals.h
@@ -666,6 +666,7 @@ struct FRHIGlobals
 	bool SupportsBarycentricsSemantic = false;
 
 	/** True if HDR requires vendor specific extensions */
+	UE_DEPRECATED(5.5, "GRHIGlobals.HDRNeedsVendorExtensions has been deprecated - Vendor extensions are not used anymore")
 	bool HDRNeedsVendorExtensions = false;
 
 	/** True if RHI supports MSAA resolve with a custom shader */
diff --git a/Engine/Source/Runtime/RenderCore/Private/RenderCore.cpp b/Engine/Source/Runtime/RenderCore/Private/RenderCore.cpp
--- a/Engine/Source/Runtime/RenderCore/Private/RenderCore.cpp
+++ b/Engine/Source/Runtime/RenderCore/Private/RenderCore.cpp
@@ -544,20 +544,6 @@ bool HdrHasWindowParamsFromCVars(void* OSWindow, FHDRMetaData& HDRMetaData)
 	return false;
 }
 
-static void HDRGetDeviceAndColorGamut(uint32 DeviceId, uint32 DisplayNitLevel, EDisplayOutputFormat& OutDisplayOutputFormat, EDisplayColorGamut& OutDisplayColorGamut)
-{
-	if (GRHIHDRNeedsVendorExtensions)
-	{
-		// all our implementations of HDR with vendor extensions happen with FP16 / ScRGB. See FD3D11DynamicRHI::EnableHDR / SetHDRMonitorMode*
-		OutDisplayOutputFormat = EDisplayOutputFormat::HDR_ACES_1000nit_ScRGB;
-		OutDisplayColorGamut = EDisplayColorGamut::sRGB_D65;
-	}
-	else
-	{
-		FPlatformMisc::ChooseHDRDeviceAndColorGamut(DeviceId, DisplayNitLevel, OutDisplayOutputFormat, OutDisplayColorGamut);
-	}
-}
-
 RENDERCORE_API void HDRGetMetaData(EDisplayOutputFormat& OutDisplayOutputFormat, EDisplayColorGamut& OutDisplayColorGamut, bool& OutbHDRSupported, 
 								   const FVector2D& WindowTopLeft, const FVector2D& WindowBottomRight, void* OSWindow)
 {
@@ -609,7 +595,7 @@ RENDERCORE_API void HDRGetMetaData(EDisplayOutputFormat& OutDisplayOutputFormat,
 
 	if (OutbHDRSupported)
 	{
-		HDRGetDeviceAndColorGamut(GRHIVendorId, CVarHDRDisplayMaxLuminance.GetValueOnAnyThread(), OutDisplayOutputFormat, OutDisplayColorGamut);
+		FPlatformMisc::ChooseHDRDeviceAndColorGamut(GRHIVendorId, CVarHDRDisplayMaxLuminance.GetValueOnAnyThread(), OutDisplayOutputFormat, OutDisplayColorGamut);
 	}
 
 }
@@ -643,7 +629,7 @@ RENDERCORE_API void HDRConfigureCVars(bool bIsHDREnabled, uint32 DisplayNits, bo
 	// If we are turning it off, we'll reset back to 0/0
 	if (bIsHDREnabled)
 	{
-		HDRGetDeviceAndColorGamut(GRHIVendorId, DisplayNits, OutputDevice, ColorGamut);
+		FPlatformMisc::ChooseHDRDeviceAndColorGamut(GRHIVendorId, DisplayNits, OutputDevice, ColorGamut);
 	}
 
 	//CVarHDRDisplayMaxLuminance is ECVF_SetByCode as it's only a mean of communicating the information from UGameUserSettings to the rest of the engine
diff --git a/Engine/Source/Runtime/Windows/D3D11RHI/Private/D3D11RHIPrivate.h b/Engine/Source/Runtime/Windows/D3D11RHI/Private/D3D11RHIPrivate.h
--- a/Engine/Source/Runtime/Windows/D3D11RHI/Private/D3D11RHIPrivate.h
+++ b/Engine/Source/Runtime/Windows/D3D11RHI/Private/D3D11RHIPrivate.h
@@ -375,10 +375,6 @@ public:
 	virtual void Shutdown() override;
 	virtual const TCHAR* GetName() override { return TEXT("D3D11"); }
 
-	// HDR display output
-	virtual void EnableHDR();
-	virtual void ShutdownHDR();
-
 	virtual void FlushPendingLogs() override;
 
 	static FD3D11DynamicRHI& Get() { return *GetDynamicRHI<FD3D11DynamicRHI>(); }
@@ -485,7 +481,7 @@ public:
 	virtual FViewportRHIRef RHICreateViewport(void* WindowHandle, uint32 SizeX, uint32 SizeY, bool bIsFullscreen, EPixelFormat PreferredPixelFormat) override;
 	virtual void RHIResizeViewport(FRHIViewport* Viewport, uint32 SizeX, uint32 SizeY, bool bIsFullscreen) final override;
 	virtual void RHIResizeViewport(FRHIViewport* Viewport, uint32 SizeX, uint32 SizeY, bool bIsFullscreen, EPixelFormat PreferredPixelFormat) final override;
-	virtual void RHICheckViewportHDRStatus(FRHIViewport* ViewportRHI) final override;
+	virtual void RHIHandleDisplayChange() final override;
 	virtual void RHITick(float DeltaTime) final override;
 	virtual void RHIBlockUntilGPUIdle() final override;
 	virtual bool RHIGetAvailableResolutions(FScreenResolutionArray& Resolutions, bool bIgnoreRefreshRate) final override;
@@ -689,17 +685,6 @@ public:
 	void ConditionalClearShaderResource(FD3D11ViewableResource* Resource, bool bCheckBoundInputAssembler);
 	void ClearAllShaderResources();
 
-	uint32 GetHDRDetectedDisplayIndex() const
-	{
-		return HDRDetectedDisplayIndex;
-	}
-
-	void SetHDRDetectedDisplayIndices(const uint32 DisplayIndex, const uint32 IHVIndex)
-	{
-		HDRDetectedDisplayIndex = DisplayIndex;
-		HDRDetectedDisplayIHVIndex = IHVIndex;
-	}
-
 	EPixelFormat GetDisplayFormat(EPixelFormat InPixelFormat) const;
 
 	FD3D11StateCache& GetStateCache() { return StateCache; }
@@ -707,6 +692,9 @@ public:
 protected:
 	/** The global D3D interface. */
 	TRefCountPtr<IDXGIFactory1> DXGIFactory1;
+#if PLATFORM_WINDOWS
+	TRefCountPtr<IDXGIFactory1> DXGIFactoryForDisplayList;
+#endif
 
 	// Whether HDR is available from the particular DXGI factories available
 	bool bDXGISupportsHDR;
@@ -878,10 +866,6 @@ protected:
 	TGlobalResource< TBoundShaderStateHistory<10000> > BoundShaderStateHistory;
 	FComputeShaderRHIRef CurrentComputeShader;
 
-	/** If HDR display detected, we store the output device. */
-	uint32 HDRDetectedDisplayIndex;
-	uint32 HDRDetectedDisplayIHVIndex;
-
 	FDisplayInformationArray DisplayList;
 
 	HANDLE ExceptionHandlerHandle = INVALID_HANDLE_VALUE;
diff --git a/Engine/Source/Runtime/Windows/D3D11RHI/Private/D3D11Viewport.cpp b/Engine/Source/Runtime/Windows/D3D11RHI/Private/D3D11Viewport.cpp
--- a/Engine/Source/Runtime/Windows/D3D11RHI/Private/D3D11Viewport.cpp
+++ b/Engine/Source/Runtime/Windows/D3D11RHI/Private/D3D11Viewport.cpp
@@ -205,9 +205,6 @@ FD3D11Viewport::~FD3D11Viewport()
 {
 	check(IsInRHIThread() || IsInRenderingThread());
 
-	// Turn off HDR display mode
-	D3DRHI->ShutdownHDR();
-
 	// If the swap chain was in fullscreen mode, switch back to windowed before releasing the swap chain.
 	// DXGI throws an error otherwise.
 	if (SwapChain)
@@ -334,18 +331,15 @@ void FD3D11Viewport::Resize(uint32 InSizeX, uint32 InSizeY, bool bInIsFullscreen
 	HDRGetMetaData(DisplayOutputFormat, DisplayColorGamut, bHDREnabled, WindowTopLeft, WindowBottomRight, (void*)WindowHandle);
 
 	// Float RGBA backbuffers are requested whenever HDR mode is desired
-	if (PixelFormat == GRHIHDRDisplayOutputFormat && bIsFullscreen)
+	if (bHDREnabled)
 	{
-		D3DRHI->EnableHDR();
+		EnableHDR();
 	}
 	else
 	{
-		D3DRHI->ShutdownHDR();
+		ShutdownHDR();
 	}
 
-	// If the window has been moved or resized it may have moved focus to and from a HDR monitor
-	CheckHDRMonitorStatus();
-	
 	// Create a RHI surface to represent the viewport's back buffer.
 	BackBuffer = GetSwapChainSurface(D3DRHI, PixelFormat, SizeX, SizeY, SwapChain);
 }
@@ -695,12 +689,6 @@ void* FD3D11Viewport::GetNativeBackBufferRT() const
 	return GetBackBuffer()->GetRenderTargetView(0, 0);
 }
 
-void  FD3D11DynamicRHI::RHICheckViewportHDRStatus(FRHIViewport* ViewportRHI)
-{
-	FD3D11Viewport* Viewport = ResourceCast(ViewportRHI);
-	return Viewport->CheckHDRMonitorStatus();
-}
-
 /*=============================================================================
  *	The following RHI functions must be called from the main thread.
  *=============================================================================*/
diff --git a/Engine/Source/Runtime/Windows/D3D11RHI/Private/Windows/WindowsD3D11Device.cpp b/Engine/Source/Runtime/Windows/D3D11RHI/Private/Windows/WindowsD3D11Device.cpp
--- a/Engine/Source/Runtime/Windows/D3D11RHI/Private/Windows/WindowsD3D11Device.cpp
+++ b/Engine/Source/Runtime/Windows/D3D11RHI/Private/Windows/WindowsD3D11Device.cpp
@@ -448,202 +448,57 @@ const DisplayChromacities DisplayChromacityList[] =
 	{ 0.71300f, 0.29300f, 0.16500f, 0.83000f, 0.12800f, 0.04400f, 0.32168f, 0.33767f }, // EDisplayColorGamut::ACEScg_D60
 };
 
-static void SetHDRMonitorModeNVIDIA(uint32 IHVDisplayIndex, bool bEnableHDR, EDisplayColorGamut DisplayGamut, float MaxOutputNits, float MinOutputNits, float MaxCLL, float MaxFALL)
-{
-#ifdef NVAPI_INTERFACE
-	NvAPI_Status NvStatus = NVAPI_OK;
-	NvDisplayHandle hNvDisplay = NULL;
-	NvU32 DisplayId = (NvU32)IHVDisplayIndex;
-
-	NV_HDR_CAPABILITIES HDRCapabilities = {};
-	HDRCapabilities.version = NV_HDR_CAPABILITIES_VER;
-
-	NvStatus = NvAPI_Disp_GetHdrCapabilities(DisplayId, &HDRCapabilities);
-
-	if (NvStatus == NVAPI_OK)
-	{
-		if (HDRCapabilities.isST2084EotfSupported)
-		{
-			NV_HDR_COLOR_DATA HDRColorData = {};
-			memset(&HDRColorData, 0, sizeof(HDRColorData));
-
-			HDRColorData.version = NV_HDR_COLOR_DATA_VER;
-			HDRColorData.cmd = NV_HDR_CMD_SET;
-			HDRColorData.static_metadata_descriptor_id = NV_STATIC_METADATA_TYPE_1;
-			HDRColorData.hdrMode = bEnableHDR ? NV_HDR_MODE_UHDBD : NV_HDR_MODE_OFF;
-
-			const DisplayChromacities& Chroma = DisplayChromacityList[(int32)DisplayGamut];
-
-			HDRColorData.mastering_display_data.displayPrimary_x0 = NvU16(Chroma.RedX * 50000.0f);
-			HDRColorData.mastering_display_data.displayPrimary_y0 = NvU16(Chroma.RedY * 50000.0f);
-			HDRColorData.mastering_display_data.displayPrimary_x1 = NvU16(Chroma.GreenX * 50000.0f);
-			HDRColorData.mastering_display_data.displayPrimary_y1 = NvU16(Chroma.GreenY * 50000.0f);
-			HDRColorData.mastering_display_data.displayPrimary_x2 = NvU16(Chroma.BlueX * 50000.0f);
-			HDRColorData.mastering_display_data.displayPrimary_y2 = NvU16(Chroma.BlueY * 50000.0f);
-			HDRColorData.mastering_display_data.displayWhitePoint_x = NvU16(Chroma.WpX * 50000.0f);
-			HDRColorData.mastering_display_data.displayWhitePoint_y = NvU16(Chroma.WpY * 50000.0f);
-			HDRColorData.mastering_display_data.max_display_mastering_luminance = NvU16(MaxOutputNits);
-			HDRColorData.mastering_display_data.min_display_mastering_luminance = NvU16(MinOutputNits);
-			HDRColorData.mastering_display_data.max_content_light_level = NvU16(MaxCLL);
-			HDRColorData.mastering_display_data.max_frame_average_light_level = NvU16(MaxFALL);
-
-			NvStatus = NvAPI_Disp_HdrColorControl(DisplayId, &HDRColorData);
-
-			// Ignore expected failures caused by insufficient driver version, remote desktop connections and similar
-			if (NvStatus != NVAPI_OK && NvStatus != NVAPI_ERROR && NvStatus != NVAPI_NVIDIA_DEVICE_NOT_FOUND)
-			{
-				NvAPI_ShortString SzDesc;
-				NvAPI_GetErrorMessage(NvStatus, SzDesc);
-				UE_LOG(LogD3D11RHI, Warning, TEXT("NvAPI_Disp_HdrColorControl returned %s (%x)"), ANSI_TO_TCHAR(SzDesc), int(NvStatus));
-			}
-		}
-	}
-#endif //NVAPI_INTERFACE
-}
-
-static void SetHDRMonitorModeAMD(uint32 IHVDisplayIndex, bool bEnableHDR, EDisplayColorGamut DisplayGamut, float MaxOutputNits, float MinOutputNits, float MaxCLL, float MaxFALL)
+bool FD3D11DynamicRHI::SetupDisplayHDRMetaData()
 {
-#ifdef AMD_AGS_API
-	const int32 AmdHDRDeviceIndex = (IHVDisplayIndex & 0xffff0000) >> 16;
-	const int32 AmdHDRDisplayIndex = IHVDisplayIndex & 0x0000ffff;
-
-	check(AmdInfo.AmdAgsContext != NULL && AmdHDRDeviceIndex != -1 && AmdHDRDisplayIndex != -1);
-	check(AmdInfo.AmdGpuInfo.numDevices > AmdHDRDeviceIndex && AmdInfo.AmdGpuInfo.devices[AmdHDRDeviceIndex].numDisplays > AmdHDRDisplayIndex);
+	check(GetDevice());
 
-	const AGSDeviceInfo& DeviceInfo = AmdInfo.AmdGpuInfo.devices[AmdHDRDeviceIndex];
-	const AGSDisplayInfo& DisplayInfo = DeviceInfo.displays[AmdHDRDisplayIndex];
+	TRefCountPtr<IDXGIAdapter> DXGIAdapter = Adapter.DXGIAdapter;
 
-	if (DisplayInfo.HDR10 != 0 || DisplayInfo.dolbyVision != 0)
+	if (DXGIFactory1.IsValid() && !DXGIFactory1->IsCurrent())
 	{
-		AGSDisplaySettings HDRDisplaySettings;
-		FMemory::Memzero(&HDRDisplaySettings, sizeof(HDRDisplaySettings));
-
-		HDRDisplaySettings.mode = bEnableHDR ? AGSDisplaySettings::Mode_HDR10_scRGB : AGSDisplaySettings::Mode_SDR;
-
-		if (bEnableHDR)
-		{
-			const DisplayChromacities& Chroma = DisplayChromacityList[(int32)DisplayGamut];
-			HDRDisplaySettings.chromaticityRedX   = Chroma.RedX;
-			HDRDisplaySettings.chromaticityRedY   = Chroma.RedY;
-			HDRDisplaySettings.chromaticityGreenX = Chroma.GreenX;
-			HDRDisplaySettings.chromaticityGreenY = Chroma.GreenY;
-			HDRDisplaySettings.chromaticityBlueX  = Chroma.BlueX;
-			HDRDisplaySettings.chromaticityBlueY  = Chroma.BlueY;
-			HDRDisplaySettings.chromaticityWhitePointX = Chroma.WpX;
-			HDRDisplaySettings.chromaticityWhitePointY = Chroma.WpY;
-			HDRDisplaySettings.maxLuminance = MaxOutputNits;
-			HDRDisplaySettings.minLuminance = MinOutputNits;
-			HDRDisplaySettings.maxContentLightLevel = MaxCLL;
-			HDRDisplaySettings.maxFrameAverageLightLevel = MaxFALL;
-		}
-
-		AGSReturnCode AmdStatus = agsSetDisplayMode(AmdInfo.AmdAgsContext, AmdHDRDeviceIndex, AmdHDRDisplayIndex, &HDRDisplaySettings);
-
-		// Ignore expected failures caused by insufficient driver version
-		if (AmdStatus != AGS_SUCCESS && AmdStatus != AGS_LEGACY_DRIVER)
+		if (!DXGIFactoryForDisplayList.IsValid() || !DXGIFactoryForDisplayList->IsCurrent())
 		{
-			UE_LOG(LogD3D11RHI, Warning, TEXT("agsSetDisplayMode returned (%x)"), int(AmdStatus));
+			SafeCreateDXGIFactory(DXGIFactoryForDisplayList.GetInitReference(), GRHIGlobals.IsDebugLayerEnabled);
 		}
-	}
-#endif //AMD_AGS_API
-}
 
-/** Enable HDR meta data transmission */
-void FD3D11DynamicRHI::EnableHDR()
-{
-	if ( GRHISupportsHDROutput && IsHDREnabled() )
-	{
-		const EDisplayOutputFormat OutputDevice = HDRGetDefaultDisplayOutputFormat();
-		const EDisplayColorGamut DisplayGamut = HDRGetDefaultDisplayColorGamut();
+		if (DXGIFactoryForDisplayList.IsValid() && DXGIFactoryForDisplayList->IsCurrent())
+		{
+			TRefCountPtr<IDXGIAdapter> TempAdapter;
+			bool bMatchingAdapterFound = false;
+			for (uint32 AdapterIndex = 0; DXGIFactoryForDisplayList->EnumAdapters(AdapterIndex, TempAdapter.GetInitReference()) != DXGI_ERROR_NOT_FOUND; ++AdapterIndex)
+			{
+				DXGI_ADAPTER_DESC AdapterDesc;
+				if (!TempAdapter.IsValid() || FAILED(TempAdapter->GetDesc(&AdapterDesc)))
+				{
+					continue;
+				}
 
-		const float DisplayMaxOutputNits = HDRGetDisplayMaximumLuminance();
-		const float DisplayMinOutputNits = 0.0f;	// Min output of the display
-		const float DisplayMaxCLL = 0.0f;			// Max content light level in lumens (0.0 == unknown)
-		const float DisplayFALL = 0.0f;				// Frame average light level (0.0 == unknown)
+				if (AdapterDesc.AdapterLuid.LowPart != Adapter.DXGIAdapterDesc.AdapterLuid.LowPart)
+				{
+					continue;
+				}
 
-		if (IsRHIDeviceNVIDIA())
-		{
-			SetHDRMonitorModeNVIDIA(
-				HDRDetectedDisplayIHVIndex,
-				true,
-				DisplayGamut,
-				DisplayMaxOutputNits,
-				DisplayMinOutputNits,
-				DisplayMaxCLL,
-				DisplayFALL);
-		}
-		else if (IsRHIDeviceAMD())
-		{
-			SetHDRMonitorModeAMD(
-				HDRDetectedDisplayIHVIndex,
-				true,
-				DisplayGamut,
-				DisplayMaxOutputNits,
-				DisplayMinOutputNits,
-				DisplayMaxCLL,
-				DisplayFALL);
-		}
-		else if (IsRHIDeviceIntel())
-		{
-			UE_LOG(LogD3D11RHI, Warning, TEXT("There is no HDR output implementation currently available for this hardware."));
-		}
-	}
-}
+				if (AdapterDesc.AdapterLuid.HighPart != Adapter.DXGIAdapterDesc.AdapterLuid.HighPart)
+				{
+					continue;
+				}
 
-/** Disable HDR meta data transmission */
-void FD3D11DynamicRHI::ShutdownHDR()
-{
-	if (GRHISupportsHDROutput)
-	{
-		// Default SDR display data
-		const float DisplayMaxOutputNits = 100.0f;	// Max output of the display
-		const float DisplayMinOutputNits = 0.0f;	// Min output of the display
-		const float DisplayMaxCLL = 100.0f;			// Max content light level in lumens
-		const float DisplayFALL = 20.0f;			// Frame average light level
+				bMatchingAdapterFound = true;
+				break;
+			}
 
-		if (IsRHIDeviceNVIDIA())
-		{
-			SetHDRMonitorModeNVIDIA(
-				HDRDetectedDisplayIHVIndex,
-				false,
-				EDisplayColorGamut::sRGB_D65,
-				DisplayMaxOutputNits,
-				DisplayMinOutputNits,
-				DisplayMaxCLL,
-				DisplayFALL);
-		}
-		else if (IsRHIDeviceAMD())
-		{
-			SetHDRMonitorModeAMD(
-				HDRDetectedDisplayIHVIndex,
-				false,
-				EDisplayColorGamut::sRGB_D65,
-				DisplayMaxOutputNits,
-				DisplayMinOutputNits,
-				DisplayMaxCLL,
-				DisplayFALL);
-		}
-		else if (IsRHIDeviceIntel())
-		{
-			// Not yet implemented
+			if (bMatchingAdapterFound && ensure(TempAdapter.IsValid()))
+			{
+				DXGIAdapter = TempAdapter;
+			}
 		}
 	}
-}
-
-bool FD3D11DynamicRHI::SetupDisplayHDRMetaData()
-{
-	check(GetDevice());
 
-	// Default to primary display
-	SetHDRDetectedDisplayIndices(0, 0);
-	
 	DisplayList.Empty();
 
-#if WITH_EDITOR
 	// Determines if any displays support HDR
 	bool bSupportsHDROutput = false;
 	{
-		IDXGIAdapter* DXGIAdapter = Adapter.DXGIAdapter;
 		if (!DXGIAdapter)
 		{
 			return false;
@@ -685,111 +540,15 @@ bool FD3D11DynamicRHI::SetupDisplayHDRMetaData()
 	}
 
 	return bSupportsHDROutput;
-#else
-
-	// Grab the adapter
-	TRefCountPtr<IDXGIDevice> DXGIDevice;
-	VERIFYD3D11RESULT(Direct3DDevice->QueryInterface(IID_IDXGIDevice, (void**)DXGIDevice.GetInitReference()));
-
-	TRefCountPtr<IDXGIAdapter> DXGIAdapter;
-	DXGIDevice->GetAdapter((IDXGIAdapter**)DXGIAdapter.GetInitReference());
-		
-	uint32 DisplayIndex = 0;
-	uint32 ForcedDisplayIndex = 0;
-	bool bForcedDisplay = FParse::Value(FCommandLine::Get(), TEXT("FullscreenDisplay="), ForcedDisplayIndex);
-
-	bool bSupportsHDROutput = false;
-	for (; true; ++DisplayIndex)
-	{
-		TRefCountPtr<IDXGIOutput> DXGIOutput;
-		if (S_OK != DXGIAdapter->EnumOutputs(DisplayIndex, DXGIOutput.GetInitReference()))
-		{
-			break;
-		}
-
-		// Query requested display only
-		if (bForcedDisplay && DisplayIndex != ForcedDisplayIndex)
-		{
-			continue;
-		}
-
-		DXGI_OUTPUT_DESC OutputDesc;
-		DXGIOutput->GetDesc(&OutputDesc);
-		FDisplayInformation DisplayInformation{};
-		const RECT& DisplayCoords = OutputDesc.DesktopCoordinates;
-		DisplayInformation.DesktopCoordinates = FIntRect(DisplayCoords.left, DisplayCoords.top, DisplayCoords.right, DisplayCoords.bottom);
-		DisplayInformation.bHDRSupported = false;
-
-		if (IsRHIDeviceNVIDIA())
-		{
-#ifdef NVAPI_INTERFACE
-			NvU32 DisplayId = 0;
-
-			// Technically, the DeviceName is a WCHAR however, UE makes the assumption elsewhere that TCHAR == WCHAR on Windows
-			NvAPI_Status Status = NvAPI_DISP_GetDisplayIdByDisplayName(TCHAR_TO_ANSI(OutputDesc.DeviceName), &DisplayId);
-
-			if (Status == NVAPI_OK)
-			{
-				NV_HDR_CAPABILITIES HdrCapabilities = {};
-
-				HdrCapabilities.version = NV_HDR_CAPABILITIES_VER;
-
-				if (NVAPI_OK == NvAPI_Disp_GetHdrCapabilities(DisplayId, &HdrCapabilities))
-				{		
-					// we're only choosing the first supported HDR monitor
-					if (HdrCapabilities.isST2084EotfSupported && !bSupportsHDROutput)
-					{
-						UE_LOG(LogD3D11RHI, Log, TEXT("HDR output is supported on display %i (NvId: 0x%x)."), DisplayIndex, DisplayId);
-						SetHDRDetectedDisplayIndices(DisplayIndex, DisplayId);
-						bSupportsHDROutput = true;
-					}
-					DisplayInformation.bHDRSupported = HdrCapabilities.isST2084EotfSupported;
-				}
-			}
-			else if (Status != NVAPI_ERROR && Status != NVAPI_NVIDIA_DEVICE_NOT_FOUND)
-			{
-				NvAPI_ShortString szDesc;
-				NvAPI_GetErrorMessage(Status, szDesc);
-				UE_LOG(LogD3D11RHI, Log, TEXT("Failed to enumerate display ID for NVAPI (%s) (%s) unable to"), OutputDesc.DeviceName, ANSI_TO_TCHAR(szDesc));
-			}
-#endif //NVAPI_INTERFACE
-		}
-		else if (IsRHIDeviceAMD())
-		{
-#ifdef AMD_AGS_API
-			// Search the device list for a matching display device name
-			for (uint16 AMDDeviceIndex = 0; AMDDeviceIndex < AmdInfo.AmdGpuInfo.numDevices; ++AMDDeviceIndex)
-			{
-				const AGSDeviceInfo& DeviceInfo = AmdInfo.AmdGpuInfo.devices[AMDDeviceIndex];
-				for (uint16 AMDDisplayIndex = 0; DeviceInfo.displays != nullptr && AMDDisplayIndex < DeviceInfo.numDisplays; ++AMDDisplayIndex)
-				{
-					const AGSDisplayInfo& DisplayInfo = DeviceInfo.displays[AMDDisplayIndex];
-					if (FCStringAnsi::Strcmp(TCHAR_TO_ANSI(OutputDesc.DeviceName), DisplayInfo.displayDeviceName) == 0)
-					{
-						// AGS has flags for HDR10 and Dolby Vision instead of a flag for the ST2084 transfer function.
-						// Both HDR10 and Dolby Vision use the ST2084 EOTF.
-						bool DisplaySupportsHDR = (DisplayInfo.HDR10 != 0 || DisplayInfo.dolbyVision != 0);
-						if (DisplaySupportsHDR && !bSupportsHDROutput)
-						{
-							UE_LOG(LogD3D11RHI, Log, TEXT("HDR output is supported on display %i (AMD Device: 0x%x, Display: 0x%x)."), DisplayIndex, AMDDeviceIndex, AMDDisplayIndex);
-							SetHDRDetectedDisplayIndices(DisplayIndex, (uint32)(AMDDeviceIndex << 16) | (uint32)AMDDisplayIndex);
-							bSupportsHDROutput = true;
-						}
-						DisplayInformation.bHDRSupported = DisplaySupportsHDR;
-					}
-				}
-			}
-#endif //AMD_AGS_API
-		}
-		else if (IsRHIDeviceIntel())
-		{
-			// Not yet implemented
-		}
-		DisplayList.Add(DisplayInformation);
-	}
+}
 
-	return bSupportsHDROutput;
-#endif
+extern void HDRSettingChangedSinkCallback();
+void FD3D11DynamicRHI::RHIHandleDisplayChange()
+{
+	RHIBlockUntilGPUIdle();
+	GRHISupportsHDROutput = SetupDisplayHDRMetaData();
+	// make sure CVars are being updated properly
+	HDRSettingChangedSinkCallback();
 }
 
 static bool IsDeviceOverclocked()
@@ -1883,9 +1642,11 @@ void FD3D11DynamicRHI::InitD3DDevice()
 
 		SetupAfterDeviceCreation();
 		GRHISupportsHDROutput = SetupDisplayHDRMetaData();
-#if !WITH_EDITOR
-		// cooked game D3D11 still needs to rely on vendor extensions to trigger HDR, which will then require exclusive fullscreen state / special RT formats
-		GRHIHDRNeedsVendorExtensions = true;
+
+#if WITH_EDITOR
+		GRHIHDRDisplayOutputFormat = PF_FloatRGBA;
+#else
+		GRHIHDRDisplayOutputFormat = PF_A2B10G10R10;
 #endif
 
 		// Add device overclock state to crash context
diff --git a/Engine/Source/Runtime/Windows/D3D11RHI/Private/Windows/WindowsD3D11Viewport.cpp b/Engine/Source/Runtime/Windows/D3D11RHI/Private/Windows/WindowsD3D11Viewport.cpp
--- a/Engine/Source/Runtime/Windows/D3D11RHI/Private/Windows/WindowsD3D11Viewport.cpp
+++ b/Engine/Source/Runtime/Windows/D3D11RHI/Private/Windows/WindowsD3D11Viewport.cpp
@@ -108,10 +108,10 @@ FD3D11Viewport::FD3D11Viewport(FD3D11DynamicRHI* InD3DRHI,HWND InWindowHandle,ui
 	}
 
 	// If requested, keep a handle to a DXGIOutput so we can force that display on fullscreen swap
-	uint32 DisplayIndex = D3DRHI->GetHDRDetectedDisplayIndex();
+	uint32 DisplayIndex = 0;
 	bForcedFullscreenDisplay = FParse::Value(FCommandLine::Get(), TEXT("FullscreenDisplay="), DisplayIndex);
 
-	if (bForcedFullscreenDisplay || GRHISupportsHDROutput)
+	if (bForcedFullscreenDisplay)
 	{
 		TRefCountPtr<IDXGIAdapter> DXGIAdapter;
 		DXGIDevice->GetAdapter(DXGIAdapter.GetInitReference());
@@ -128,12 +128,6 @@ FD3D11Viewport::FD3D11Viewport(FD3D11DynamicRHI* InD3DRHI,HWND InWindowHandle,ui
 		ForcedFullscreenOutput = nullptr;
 	}
 
-	if (PixelFormat == PF_FloatRGBA && bIsFullscreen)
-	{
-		// Send HDR meta data to enable
-		D3DRHI->EnableHDR();
-	}
-	
 	DXGI_FORMAT SwapChainFormat = GetRenderTargetFormat(PixelFormat);
 
 	// Skip swap chain creation in off-screen rendering mode
@@ -226,10 +220,8 @@ FD3D11Viewport::FD3D11Viewport(FD3D11DynamicRHI* InD3DRHI,HWND InWindowHandle,ui
 				{
 					DisplayOutputFormat = LocalDisplayOutputFormat;
 					DisplayColorGamut = LocalDisplayColorGamut;
+					EnableHDR();
 				}
-
-				// See if we are running on a HDR monitor 
-				CheckHDRMonitorStatus();
 			}
 			else
 			{
@@ -383,15 +375,25 @@ inline void EnsureColorSpace(IDXGISwapChain* SwapChain, EDisplayColorGamut Displ
 // When a window has moved or resized we need to check whether it is on a HDR monitor or not. Set the correct color space of the monitor
 void FD3D11Viewport::CheckHDRMonitorStatus()
 {
-#if WITH_EDITOR
-	static auto CVarHDREnable = IConsoleManager::Get().FindConsoleVariable(TEXT("Editor.HDRSupport"));
-	if (CVarHDREnable->GetInt() != 0)
-	{
-		FlushRenderingCommands();
+}
 
+void FD3D11Viewport::EnableHDR()
+{
+	if ( GRHISupportsHDROutput && IsHDREnabled() )
+	{
+		// Ensure we have the correct color space set.
 		EnsureColorSpace(SwapChain, DisplayColorGamut, DisplayOutputFormat, PixelFormat);
 	}
-#endif
+}
+
+void FD3D11Viewport::ShutdownHDR()
+{
+	// Make sure to set the appropriate color space even if GRHISupportsHDROutput is false because we 
+	// might have toggled HDR on and off in the windows settings
+	{
+		// Ensure we have the correct color space set.
+		EnsureColorSpace(SwapChain, EDisplayColorGamut::sRGB_D65, EDisplayOutputFormat::SDR_sRGB, PixelFormat);
+	}
 }
 
 void FD3D11Viewport::ConditionalResetSwapChain(bool bIgnoreFocus)
diff --git a/Engine/Source/Runtime/Windows/D3D11RHI/Public/D3D11Viewport.h b/Engine/Source/Runtime/Windows/D3D11RHI/Public/D3D11Viewport.h
--- a/Engine/Source/Runtime/Windows/D3D11RHI/Public/D3D11Viewport.h
+++ b/Engine/Source/Runtime/Windows/D3D11RHI/Public/D3D11Viewport.h
@@ -116,6 +116,13 @@ protected:
 	 */
 	D3D11RHI_API bool PresentChecked(IRHICommandContext& RHICmdContext, int32 SyncInterval);
 
+	/** Enable HDR meta data transmission and set the necessary color space. */
+	void EnableHDR();
+
+	/** Disable HDR meta data transmission and set the necessary color space. */
+	void ShutdownHDR();
+
+
 	FD3D11DynamicRHI* D3DRHI;
 	uint64 LastFlipTime;
 	uint64 LastFrameComplete;
