- Fixed a NonUnity issue on xbox

diff --git a/Engine/Source/Runtime/Windows/AudioMixerXAudio2/Private/ToStringHelpers.cpp b/Engine/Source/Runtime/Windows/AudioMixerXAudio2/Private/ToStringHelpers.cpp
--- a/Engine/Source/Runtime/Windows/AudioMixerXAudio2/Private/ToStringHelpers.cpp
+++ b/Engine/Source/Runtime/Windows/AudioMixerXAudio2/Private/ToStringHelpers.cpp
@@ -1,7 +1,6 @@
 // Copyright Epic Games, Inc. All Rights Reserved.
 
 #include "ToStringHelpers.h"
-#include "AudioMixerPlatformXAudio2.h"
 
 #include "Windows/AllowWindowsPlatformTypes.h"
 
diff --git a/Engine/Source/Runtime/Windows/AudioMixerXAudio2/Private/ToStringHelpers.h b/Engine/Source/Runtime/Windows/AudioMixerXAudio2/Private/ToStringHelpers.h
--- a/Engine/Source/Runtime/Windows/AudioMixerXAudio2/Private/ToStringHelpers.h
+++ b/Engine/Source/Runtime/Windows/AudioMixerXAudio2/Private/ToStringHelpers.h
@@ -3,6 +3,7 @@
 #pragma once
 
 #include "AudioMixer.h"
+#include "AudioMixerPlatformXAudio2.h"
 
 namespace Audio
 {
