Speculative fix for OOMs on Microsoft platforms - Enables device swap style behavior on consoles.

diff --git a/Engine/Source/Runtime/Windows/AudioMixerXAudio2/Private/AudioMixerPlatformXAudio2.cpp b/Engine/Source/Runtime/Windows/AudioMixerXAudio2/Private/AudioMixerPlatformXAudio2.cpp
--- a/Engine/Source/Runtime/Windows/AudioMixerXAudio2/Private/AudioMixerPlatformXAudio2.cpp
+++ b/Engine/Source/Runtime/Windows/AudioMixerXAudio2/Private/AudioMixerPlatformXAudio2.cpp
@@ -31,13 +31,14 @@
 #include "Logging/LogMacros.h"
 #include "ToStringHelpers.h"
 
-#if PLATFORM_WINDOWS
 THIRD_PARTY_INCLUDES_START
 #include <mmdeviceapi.h>
-#include <FunctionDiscoveryKeys_devpkey.h>
 #include <AudioClient.h>
+#if PLATFORM_WINDOWS
+#include <FunctionDiscoveryKeys_devpkey.h>
+#endif //PLATFORM_WINDOWS
 THIRD_PARTY_INCLUDES_END
-#endif
+
 
 #include "Misc/CoreDelegates.h"
 #include "ProfilingDebugging/ScopedTimers.h"
@@ -108,6 +109,7 @@ const FString& GetDllName(FName Current = NAME_None)
 */
 #define XAUDIO2_DEBUG_ENABLED 0
 
+
 namespace Audio
 {
 	void FXAudio2VoiceCallback::OnBufferEnd(void* BufferContext)
@@ -282,6 +284,7 @@ namespace Audio
 									MasterVoice = nullptr;
 								}
 
+
 								// Destroy System
 								{
 									SCOPED_NAMED_EVENT(FMixerPlatformXAudio2_AsyncDeleteCreate_DestroySystem, FColor::Blue);
@@ -305,7 +308,8 @@ namespace Audio
 								if (FAILED(Result))
 								{
 									XAUDIO2_LOG_RESULT("XAudio2Create", Result);
-									return {};// FAIL.
+
+									return {};
 								}
 							}
 
@@ -490,7 +494,7 @@ namespace Audio
 
 	bool FMixerPlatformXAudio2::AllowDeviceSwap()
 	{
-#if PLATFORM_WINDOWS
+
 		double CurrentTime = FPlatformTime::Seconds();
 
 		// If we're already in the process of swapping, we do not want to "double-trigger" a swap
@@ -507,7 +511,7 @@ namespace Audio
 			LastDeviceSwapTime = CurrentTime;
 			return true;
 		}
-#endif
+
 		return false;
 	}
 
@@ -516,18 +520,27 @@ namespace Audio
 		SAFE_RELEASE(XAudio2System);
 
 		uint32 Flags = 0;
+		HRESULT Result;
 
 #if WITH_XMA2
 		// We need to raise this flag explicitly to prevent initializing SHAPE twice, because we are allocating SHAPE in FXMAAudioInfo
 		Flags |= XAUDIO2_DO_NOT_USE_SHAPE;
 #endif
 
-		if (FAILED(XAudio2Create(&XAudio2System, Flags, GetXAudio2ProcessorsToUse())))
+		Result = XAudio2Create(&XAudio2System, Flags, GetXAudio2ProcessorsToUse());
+		if (FAILED(Result))
 		{
 			XAudio2System = nullptr;
+			UE_LOG(LogAudioMixer, Error, TEXT("Failed to CreateXAudio2 HRESULT=0x%X"), (uint32)Result);
 			return false;
 		}
 
+		Result = XAudio2System->RegisterForCallbacks(this);
+		if (FAILED(Result))
+		{
+			UE_LOG(LogAudioMixer, Error, TEXT("Failed to register for callbacks. HRESULT=0x%X"), (uint32)Result);
+			return false;
+		}
 		return true;
 	}
 
@@ -556,33 +569,58 @@ namespace Audio
 				StopRunningNullDevice();			
 
 				HRESULT Result = XAudio2System->StartEngine();				
-				
+
+				// Log any failed state as "Display" so we always appear in the logs.
+				UE_CLOG(FAILED(Result), LogAudioMixer, Display, TEXT("XAudio2->StartEngine() returned (%s) 0x%X"), *ToErrorFString(Result), (uint32)Result);
+
+				// If we see this error here, the audio device has been invalidated, and we must swap to a new one.
+				if (Result == XAUDIO2_E_DEVICE_INVALIDATED)
+				{
+					bIsSuspended = false;
+					RequestDeviceSwap(TEXT(""), /*force*/ true, *FString::Printf(TEXT("XAudio2->StartEngine() returned (%s) 0x%X"), *ToErrorFString(Result), (uint32)Result));
+					return;
+				}
+
 				// Suggestion from Microsoft for this returning 0x80070490 on a quick cycle. 
 				// Try again with delay.
-				constexpr int32 NumStartAttempts = 16;
-				constexpr float DelayBetweenAttemptsSecs = 0.1f;
-				
-				uint64 StartCycles = FPlatformTime::Cycles64();				
-				int32 StartAttempt = 0;
-				for(; StartAttempt < NumStartAttempts && FAILED(Result); ++StartAttempt)
+				if (Result == 0x80070490)
 				{
-					FPlatformProcess::Sleep(DelayBetweenAttemptsSecs);
-					Result = XAudio2System->StartEngine();	
+					constexpr int32 NumStartAttempts = 16;
+				
+					uint64 StartCycles = FPlatformTime::Cycles64();				
+					int32 StartAttempt = 0;
+					for (; StartAttempt < NumStartAttempts && FAILED(Result); ++StartAttempt)
+					{
+						constexpr float DelayBetweenAttemptsSecs = 0.1f;
+						FPlatformProcess::Sleep(DelayBetweenAttemptsSecs);
+						Result = XAudio2System->StartEngine();	
+					}
+
+					if (SUCCEEDED(Result))
+					{		
+						UE_CLOG(StartAttempt == 1, LogAudioMixer, Display, 
+							TEXT("XAudio2System StartEngine() - Sucessfully started, taking '%2.2f' ms"), FPlatformTime::ToMilliseconds64(FPlatformTime::Cycles64() - StartCycles));
+						
+						UE_CLOG(StartAttempt > 1, LogAudioMixer, Warning, 
+							TEXT("StartEngine() took %d attempts to start, taking '%f' ms"), 
+							StartAttempt, FPlatformTime::ToMilliseconds64(FPlatformTime::Cycles64() - StartCycles));
+					}
+					else
+					{
+						UE_CLOG(FAILED(Result), LogAudioMixer, Error,
+								TEXT("Could not resume XAudio2, StartEngine() returned 0x%X, after %d attempts"),
+								(uint32)Result, StartAttempt);
+					}
 				}
 
-				UE_CLOG(FAILED(Result), LogAudioMixer, Error,
-					TEXT("Could not resume XAudio2, StartEngine() returned %#010x, after %d attempts"),
-					(uint32)Result, StartAttempt
-				);
-				
-				if(SUCCEEDED(Result))
-				{		
-					bIsSuspended = false;								
-					UE_CLOG(StartAttempt == 1, LogAudioMixer, Display, 
-						TEXT("XAudio2System StartEngine() - Sucessfully started, taking '%2.2f' ms"), FPlatformTime::ToMilliseconds64(FPlatformTime::Cycles64() - StartCycles));
-					UE_CLOG(StartAttempt > 1, LogAudioMixer, Warning, 
-						TEXT("StartEngine() took %d attempts to start, taking '%f' ms"), 
-						StartAttempt, FPlatformTime::ToMilliseconds64(FPlatformTime::Cycles64() - StartCycles));
+				// Always turn off the suspended state even in a failure.
+				bIsSuspended = false;
+
+				// If we're still in failed state, do our best here 
+				if (FAILED(Result))
+				{
+					UE_LOG(LogAudioMixer, Error, TEXT("XAudio2->StartEngine() returned (%s) 0x%X"), *ToErrorFString(Result), (uint32)Result);
+					StartRunningNullDevice();
 				}
 			}						
 		}		
@@ -594,7 +632,6 @@ namespace Audio
 		{
 			AUDIO_PLATFORM_LOG_ONCE(TEXT("XAudio2 already initialized."), Warning);
 			return false;
-
 		}
 
 #if PLATFORM_NEEDS_SUSPEND_ON_BACKGROUND
@@ -639,12 +676,18 @@ namespace Audio
 			FMessageDialog::Open(EAppMsgType::Ok, NSLOCTEXT("Audio", "XAudio2Error", "Failed to initialize audio. This may be an issue with your installation of XAudio 2.7. XAudio2 is available in the DirectX End-User Runtime (June 2010)."));
 			return false;
 		}
+	 
 #if XAUDIO2_DEBUG_ENABLED
 		XAUDIO2_DEBUG_CONFIGURATION DebugConfiguration = { 0 };
 		DebugConfiguration.TraceMask = XAUDIO2_LOG_ERRORS | XAUDIO2_LOG_WARNINGS;
 		XAudio2System->SetDebugConfiguration(&DebugConfiguration, 0);
 #endif // #if XAUDIO2_DEBUG_ENABLED
 
+		if (FAILED(XAudio2System->RegisterForCallbacks(this)))
+		{
+			UE_LOG(LogAudioMixer, Error, TEXT("Failed to register for callbacks."));
+		}
+
 #if WITH_XMA2
 		//Initialize our XMA2 decoder context
 		FXMAAudioInfo::Initialize();
@@ -688,6 +731,10 @@ namespace Audio
 
 #endif //PLATFORM_NEEDS_SUSPEND_ON_BACKGROUND
 
+		if (XAudio2System)
+		{
+			XAudio2System->UnregisterForCallbacks(this);
+		}
 		SAFE_RELEASE(XAudio2System);
 
 #if WITH_XMA2
@@ -736,7 +783,7 @@ namespace Audio
 
 		// XAudio2 for HoloLens doesn't have GetDeviceCount, use Windows::Devices::Enumeration instead
 		// See https://blogs.msdn.microsoft.com/chuckw/2012/04/02/xaudio2-and-windows-8/
-#if PLATFORM_WINDOWS && XAUDIO_SUPPORTS_DEVICE_DETAILS
+#if XAUDIO_SUPPORTS_DEVICE_DETAILS
 
 		IMMDeviceEnumerator* DeviceEnumerator = nullptr;
 		IMMDeviceCollection* DeviceCollection = nullptr;
@@ -764,7 +811,7 @@ namespace Audio
 #endif 
 	}
 
-#if PLATFORM_WINDOWS
+
 	static bool GetMMDeviceInfo(IMMDevice* MMDevice, FAudioPlatformDeviceInfo& OutInfo)
 	{
 		SCOPED_NAMED_EVENT(FMixerPlatformXAudio2_GetMMDeviceInfo, FColor::Blue);
@@ -791,12 +838,13 @@ namespace Audio
 		Result = MMDevice->OpenPropertyStore(STGM_READ, &PropertyStore);
 		XAUDIO2_GOTO_CLEANUP_ON_FAIL(Result);
 
+#if PLATFORM_WINDOWS
 		// Grab the friendly name
 		PropVariantInit(&FriendlyName);
 		Result = PropertyStore->GetValue(PKEY_Device_FriendlyName, &FriendlyName);
 		XAUDIO2_GOTO_CLEANUP_ON_FAIL(Result);
-
 		OutInfo.Name = FString(FriendlyName.pwszVal);
+#endif 
 
 		// Retrieve the DeviceFormat prop variant
 		Result = PropertyStore->GetValue(PKEY_AudioEngine_DeviceFormat, &DeviceFormat);
@@ -920,7 +968,6 @@ namespace Audio
 
 		return SUCCEEDED(Result);
 	}
-#endif //  #if PLATFORM_WINDOWS
 
 	bool FMixerPlatformXAudio2::GetOutputDeviceInfo(const uint32 InDeviceIndex, FAudioPlatformDeviceInfo& OutInfo)
 	{
@@ -955,7 +1002,6 @@ namespace Audio
 			return false;
 		}
 
-#if PLATFORM_WINDOWS
 		IMMDeviceEnumerator* DeviceEnumerator = nullptr;
 		IMMDeviceCollection* DeviceCollection = nullptr;
 		IMMDevice* DefaultDevice = nullptr;
@@ -1028,24 +1074,6 @@ namespace Audio
 
 		return bSucceeded && SUCCEEDED(Result);
 
-#else // #elif PLATFORM_WINDOWS
-		OutInfo.bIsSystemDefault = true;
-		OutInfo.SampleRate = 44100;
-		OutInfo.DeviceId = 0;
-		OutInfo.Format = EAudioMixerStreamDataFormat::Float;
-		OutInfo.Name = TEXT("Audio Device.");
-		OutInfo.NumChannels = 8;
-
-		OutInfo.OutputChannelArray.Add(EAudioMixerChannel::FrontLeft);
-		OutInfo.OutputChannelArray.Add(EAudioMixerChannel::FrontRight);
-		OutInfo.OutputChannelArray.Add(EAudioMixerChannel::FrontCenter);
-		OutInfo.OutputChannelArray.Add(EAudioMixerChannel::LowFrequency);
-		OutInfo.OutputChannelArray.Add(EAudioMixerChannel::BackLeft);
-		OutInfo.OutputChannelArray.Add(EAudioMixerChannel::BackRight);
-		OutInfo.OutputChannelArray.Add(EAudioMixerChannel::SideLeft);
-		OutInfo.OutputChannelArray.Add(EAudioMixerChannel::SideRight);
-		return true;
-#endif 
 	}
 
 	bool FMixerPlatformXAudio2::GetDefaultOutputDeviceIndex(uint32& OutDefaultDeviceIndex) const
@@ -1107,13 +1135,17 @@ namespace Audio
 				*AudioStreamInfo.DeviceInfo.DeviceId,
 				nullptr,
 				AudioCategory_GameEffects);
-#else //PLATFORM_WINDOWS			
+#else //PLATFORM_WINDOWS
+			// Passing the device-id to CreateMasteringVoice on a non-windows platform, will prevent
+			// the creation of virtualized device which handles disconnection state for us, but
+			// if we need to handle these errors i.e. OnCriticalError callback, we need to pass the device here.
+			// This allows device swapping, so protect with a CVAR.
 			Result = XAudio2System->CreateMasteringVoice(
 				&OutputAudioStreamMasteringVoice,
 				AudioStreamInfo.DeviceInfo.NumChannels,
 				AudioStreamInfo.DeviceInfo.SampleRate,
 				0,
-				0,
+				*AudioStreamInfo.DeviceInfo.DeviceId, 
 				nullptr);
 #endif //#else PLATFORM_WINDOWS
 
@@ -1287,7 +1319,7 @@ namespace Audio
 
 	bool FMixerPlatformXAudio2::CheckAudioDeviceChange()
 	{
-#if PLATFORM_WINDOWS && XAUDIO_SUPPORTS_DEVICE_DETAILS
+#if XAUDIO_SUPPORTS_DEVICE_DETAILS
 
 		SCOPED_NAMED_EVENT(FMixerPlatformXAudio2_CheckAudioDeviceChange, FColor::Blue);
 
@@ -1348,7 +1380,7 @@ namespace Audio
 	{
 		bool bDidStopGeneratingAudio = false;
 
-#if PLATFORM_WINDOWS && XAUDIO_SUPPORTS_DEVICE_DETAILS
+#if XAUDIO_SUPPORTS_DEVICE_DETAILS
 
 		SCOPED_NAMED_EVENT(FMixerPlatformXAudio2_MoveAudioStreamToNewAudioDevice, FColor::Blue);
 
@@ -1642,4 +1674,24 @@ namespace Audio
 	{
 		return true;
 	}
+
+
+	void FMixerPlatformXAudio2::OnProcessingPassStart()
+	{
+	}
+
+	void FMixerPlatformXAudio2::OnProcessingPassEnd()
+	{
+	}
+
+	void FMixerPlatformXAudio2::OnCriticalError(HRESULT Error)
+	{
+		// Windows handles this via session events, so no need to handle that here..
+#if !PLATFORM_WINDOWS
+		UE_LOG(LogAudioMixer, Display, TEXT("FMixerPlatformXAudio2::OnCriticalError() returned (%s) 0x%X"), *ToErrorFString(Error), (uint32)Error);
+		RequestDeviceSwap(TEXT(""), /*force*/ true, *FString::Printf(TEXT("FMixerPlatformXAudio2::OnCriticalError() returned (%s) 0x%X"), *ToErrorFString(Error), (uint32)Error));
+#endif //!PLATFORM_WINDOWS
+	}
+	
 }
+
diff --git a/Engine/Source/Runtime/Windows/AudioMixerXAudio2/Private/AudioMixerPlatformXAudio2.h b/Engine/Source/Runtime/Windows/AudioMixerXAudio2/Private/AudioMixerPlatformXAudio2.h
--- a/Engine/Source/Runtime/Windows/AudioMixerXAudio2/Private/AudioMixerPlatformXAudio2.h
+++ b/Engine/Source/Runtime/Windows/AudioMixerXAudio2/Private/AudioMixerPlatformXAudio2.h
@@ -71,7 +71,8 @@ namespace Audio
 
 	};
 
-	class FMixerPlatformXAudio2 : public IAudioMixerPlatformInterface
+	class FMixerPlatformXAudio2 : public IAudioMixerPlatformInterface,
+								  public IXAudio2EngineCallback
 	{
 
 	public:
@@ -98,7 +99,7 @@ namespace Audio
 		virtual void ResumePlaybackOnNewDevice() override;
 		virtual FAudioPlatformDeviceInfo GetPlatformDeviceInfo() const override;
 		virtual void SubmitBuffer(const uint8* Buffer) override;
-		virtual bool DisablePCMAudioCaching() const override;
+		virtual bool DisablePCMAudioCaching() const override;		
 		virtual FString GetDefaultDeviceName() override;
 		virtual FAudioPlatformSettings GetPlatformSettings() const override;
 		virtual void OnHardwareUpdate() override;
@@ -118,7 +119,12 @@ namespace Audio
 		//~ End IAudioMixerDeviceChangedListener
 
 	private:
-	
+		//~ Begin IXAudio2EngineCallback
+		virtual void OnCriticalError(HRESULT Error) override;
+		virtual void OnProcessingPassStart() override;
+		virtual void OnProcessingPassEnd() override;
+		//~ End IXAudio2EngineCallback
+			
 		struct FXAudio2AsyncCreateResult
 		{
 			IXAudio2* XAudio2System = nullptr;
