Fix crashes in D3D11 occlusion queries when queries are re-used.
 - The scene occlusion query manager pools and reuses query objects, sometimes before the old result of the query has been read back. This led to cases where the State field of the query was still "Ended", triggering an assert.
 - Replace the "State" field with a set of counters to maintain correct ordering. If a query is reused, the old unread result is discarded, and the counters are used to track this.

diff --git a/Engine/Source/Runtime/Windows/D3D11RHI/Private/D3D11Query.cpp b/Engine/Source/Runtime/Windows/D3D11RHI/Private/D3D11Query.cpp
--- a/Engine/Source/Runtime/Windows/D3D11RHI/Private/D3D11Query.cpp
+++ b/Engine/Source/Runtime/Windows/D3D11RHI/Private/D3D11Query.cpp
@@ -51,19 +51,17 @@ FD3D11RenderQuery::~FD3D11RenderQuery()
 void FD3D11RenderQuery::Begin(ID3D11DeviceContext* Context)
 {
 	check(Type == EType::Occlusion);
-	check(State == EState::None);
-
 	Context->Begin(Resource);
 }
 
 void FD3D11RenderQuery::End(ID3D11DeviceContext* Context, uint64* NewTarget)
 {
-	check(State == EState::None);
+	BOPCounter++;
+
 	Context->End(Resource);
 	Target = NewTarget;
 
 	Link();
-	State = FD3D11RenderQuery::EState::Ended;
 }
 
 FRenderQueryRHIRef FD3D11DynamicRHI::RHICreateRenderQuery(ERenderQueryType QueryType)
@@ -79,18 +77,10 @@ FRenderQueryRHIRef FD3D11DynamicRHI::RHICreateRenderQuery(ERenderQueryType Query
 	}
 }
 
-void FD3D11DynamicRHI::RHIBeginRenderQuery_TopOfPipe(FRHICommandListBase& RHICmdList, FRHIRenderQuery* RenderQuery)
-{
-	FD3D11RenderQuery* Query = ResourceCast(RenderQuery);
-	Query->State = FD3D11RenderQuery::EState::None;
-
-	FDynamicRHI::RHIBeginRenderQuery_TopOfPipe(RHICmdList, RenderQuery);
-}
-
 void FD3D11DynamicRHI::RHIEndRenderQuery_TopOfPipe(FRHICommandListBase& RHICmdList, FRHIRenderQuery* RenderQuery)
 {
 	FD3D11RenderQuery* Query = ResourceCast(RenderQuery);
-	Query->State = FD3D11RenderQuery::EState::None;
+	Query->TOPCounter++;
 
 	FDynamicRHI::RHIEndRenderQuery_TopOfPipe(RHICmdList, RenderQuery);
 }
@@ -114,7 +104,7 @@ bool FD3D11DynamicRHI::RHIGetRenderQueryResult(FRHIRenderQuery* QueryRHI, uint64
 	bool bRHIThreadFlushed = false;
 
 Retry:
-	if (Query->State == FD3D11RenderQuery::EState::Completed)
+	if (Query->TOPCounter == Query->LastCachedBOPCounter.load(std::memory_order_acquire))
 	{
 		// Early return for queries we already have the result for.
 		check(!Query->IsLinked());
@@ -154,9 +144,10 @@ Retry:
 	// The query is unresolved. Either the GPU isn't done, or the commands to signal the query were never submitted (still recorded in the immediate command list).
 	//
 
-	if (Query->State == FD3D11RenderQuery::EState::None && !bRHIThreadFlushed)
+	if (Query->TOPCounter != Query->BOPCounter && !bRHIThreadFlushed)
 	{
-		// The query state is None, meaning the End() hasn't been processed on the RHI thread.
+		// When TOPCounter != BOPCounter, there's an End() operation that was recorded at the TOP, but has not yet been submitted for translation by the RHI thread.
+		// Flush the immediate command list to push this command into the RHI pipeline.
 		FRHICommandListImmediate::Get().ImmediateFlush(EImmediateFlushType::FlushRHIThread);
 		bRHIThreadFlushed = true;
 
@@ -164,14 +155,13 @@ Retry:
 		goto Retry;
 	}
 
-	checkf(Query->State != FD3D11RenderQuery::EState::None, TEXT("Attempting to get data from an RHI render query which was never issued."));
+	checkf(Query->TOPCounter == Query->BOPCounter, TEXT("Attempting to get data from an RHI render query which was never issued."));
 	if (!Query->CacheResult(*this, bWait))
 	{
 		OutResult = 0;
 		return false;
 	}
 
-	check(Query->State == FD3D11RenderQuery::EState::Completed);
 	check(!Query->IsLinked());
 	OutResult = Query->Result;
 
@@ -180,8 +170,9 @@ Retry:
 
 bool FD3D11RenderQuery::CacheResult(FD3D11DynamicRHI& RHI, bool bWait)
 {
-	if (State == FD3D11RenderQuery::EState::Completed)
+	if (BOPCounter == LastCachedBOPCounter.load(std::memory_order_relaxed))
 	{
+		// Value has been cached and no newer query operation has started.
 		check(!IsLinked());
 		return true;
 	}
@@ -227,19 +218,15 @@ bool FD3D11RenderQuery::CacheResult(FD3D11DynamicRHI& RHI, bool bWait)
 
 	Unlink();
 
+	LastCachedBOPCounter.store(BOPCounter, std::memory_order_release);
+
 #if RHI_NEW_GPU_PROFILER
 	if (Type == EType::Profiler)
 	{
-		State = FD3D11RenderQuery::EState::None;
-
 		// Return the query to the pool
 		RHI.Profiler.TimestampPool.Push(this);
 	}
-	else
 #endif
-	{
-		State = FD3D11RenderQuery::EState::Completed;
-	}
 
 	return true;
 }
diff --git a/Engine/Source/Runtime/Windows/D3D11RHI/Private/D3D11RHIPrivate.h b/Engine/Source/Runtime/Windows/D3D11RHI/Private/D3D11RHIPrivate.h
--- a/Engine/Source/Runtime/Windows/D3D11RHI/Private/D3D11RHIPrivate.h
+++ b/Engine/Source/Runtime/Windows/D3D11RHI/Private/D3D11RHIPrivate.h
@@ -513,7 +513,6 @@ public:
 	virtual void RHIBeginTransitions(TArrayView<const FRHITransition*> Transitions) override final;
 	virtual void RHIEndTransitions(TArrayView<const FRHITransition*> Transitions) override final;
 
-	virtual void RHIBeginRenderQuery_TopOfPipe(FRHICommandListBase& RHICmdList, FRHIRenderQuery* RenderQuery) override final;
 	virtual void RHIEndRenderQuery_TopOfPipe  (FRHICommandListBase& RHICmdList, FRHIRenderQuery* RenderQuery) override final;
 	virtual void RHIBeginRenderQuery(FRHIRenderQuery* RenderQuery) final override;
 	virtual void RHIEndRenderQuery  (FRHIRenderQuery* RenderQuery) final override;
diff --git a/Engine/Source/Runtime/Windows/D3D11RHI/Public/D3D11Resources.h b/Engine/Source/Runtime/Windows/D3D11RHI/Public/D3D11Resources.h
--- a/Engine/Source/Runtime/Windows/D3D11RHI/Public/D3D11Resources.h
+++ b/Engine/Source/Runtime/Windows/D3D11RHI/Public/D3D11Resources.h
@@ -318,13 +318,9 @@ private:
 	FD3D11RenderQuery* Next = nullptr;
 
 public:
-	enum class EState : uint8
-	{
-		None,
-		Ended,
-		Completed
-	};
-	std::atomic<EState> State { EState::None };
+	uint8 TOPCounter = 0;
+	uint8 BOPCounter = 0;
+	std::atomic<uint8> LastCachedBOPCounter = 0;
 
 	enum class EType : uint8
 	{
