Fix for unresolved pkey symbol

#rnx

diff --git a/Engine/Source/Runtime/Windows/AudioMixerXAudio2/Private/AudioMixerPlatformXAudio2.cpp b/Engine/Source/Runtime/Windows/AudioMixerXAudio2/Private/AudioMixerPlatformXAudio2.cpp
--- a/Engine/Source/Runtime/Windows/AudioMixerXAudio2/Private/AudioMixerPlatformXAudio2.cpp
+++ b/Engine/Source/Runtime/Windows/AudioMixerXAudio2/Private/AudioMixerPlatformXAudio2.cpp
@@ -32,6 +32,8 @@
 #include "ToStringHelpers.h"
 
 THIRD_PARTY_INCLUDES_START
+// Including initguid.h will define the PKEY symbols below which area used cross-platform
+#include <initguid.h>
 #include <mmdeviceapi.h>
 #include <AudioClient.h>
 #if PLATFORM_WINDOWS
