[Experimental] WASAPI audio backend

diff --git a/Engine/Source/Runtime/AudioCaptureImplementations/Windows/AudioCaputureWasapi/AudioCaptureWasapi.Build.cs b/Engine/Source/Runtime/AudioCaptureImplementations/Windows/AudioCaputureWasapi/AudioCaptureWasapi.Build.cs
--- a/Engine/Source/Runtime/AudioCaptureImplementations/Windows/AudioCaputureWasapi/AudioCaptureWasapi.Build.cs
+++ b/Engine/Source/Runtime/AudioCaptureImplementations/Windows/AudioCaputureWasapi/AudioCaptureWasapi.Build.cs
@@ -6,8 +6,13 @@ public class AudioCaptureWasapi : ModuleRules
 {
 	public AudioCaptureWasapi(ReadOnlyTargetRules Target) : base(Target)
 	{
-		PrivateDependencyModuleNames.Add("Core");
-		PrivateDependencyModuleNames.Add("AudioCaptureCore");
+		PrivateDependencyModuleNames.AddRange(
+			new string[] {
+				"Core",
+				"AudioCaptureCore",
+				"AudioPlatformSupportWasapi"
+			}
+		);
 
 		if (Target.Platform.IsInGroup(UnrealPlatformGroup.Windows))
 		{
diff --git a/Engine/Source/Runtime/AudioCaptureImplementations/Windows/AudioCaputureWasapi/Private/WasapiAudioFormat.cpp b/Engine/Source/Runtime/AudioCaptureImplementations/Windows/AudioCaputureWasapi/Private/WasapiAudioFormat.cpp
deleted file mode 100644
diff --git a/Engine/Source/Runtime/AudioCaptureImplementations/Windows/AudioCaputureWasapi/Private/WasapiAudioFormat.h b/Engine/Source/Runtime/AudioCaptureImplementations/Windows/AudioCaputureWasapi/Private/WasapiAudioFormat.h
deleted file mode 100644
diff --git a/Engine/Source/Runtime/AudioPlatformSupport/Windows/WASAPI/AudioPlatformSupportWasapi.Build.cs b/Engine/Source/Runtime/AudioPlatformSupport/Windows/WASAPI/AudioPlatformSupportWasapi.Build.cs
new file mode 100644
--- /dev/null
+++ b/Engine/Source/Runtime/AudioPlatformSupport/Windows/WASAPI/AudioPlatformSupportWasapi.Build.cs
@@ -0,0 +1,17 @@
+// Copyright Epic Games, Inc. All Rights Reserved.
+
+namespace UnrealBuildTool.Rules
+{
+	public class AudioPlatformSupportWasapi : ModuleRules
+	{
+		public AudioPlatformSupportWasapi(ReadOnlyTargetRules Target) : base(Target)
+		{
+			PrivateDependencyModuleNames.AddRange(
+				new string[]
+				{
+					"Core"
+				}
+			);
+		}
+	}
+}
diff --git a/Engine/Source/Runtime/AudioPlatformSupport/Windows/WASAPI/Private/AudioPlatformSupportWasapiModule.cpp b/Engine/Source/Runtime/AudioPlatformSupport/Windows/WASAPI/Private/AudioPlatformSupportWasapiModule.cpp
new file mode 100644
--- /dev/null
+++ b/Engine/Source/Runtime/AudioPlatformSupport/Windows/WASAPI/Private/AudioPlatformSupportWasapiModule.cpp
@@ -0,0 +1,24 @@
+// Copyright Epic Games, Inc. All Rights Reserved.
+
+#include "AudioPlatformSupportWasapiLog.h"
+
+#include "Modules/ModuleManager.h"
+
+DEFINE_LOG_CATEGORY(LogAudioPlatformSupportWasapi)
+
+
+class FAudioPlatformSupportWasapiModule : public IModuleInterface
+{
+public:
+	virtual void StartupModule() override
+	{
+
+	}
+
+	virtual void ShutdownModule() override
+	{
+
+	}
+};
+
+IMPLEMENT_MODULE(FAudioPlatformSupportWasapiModule, AudioPlatformSupportWasapi);
diff --git a/Engine/Source/Runtime/AudioPlatformSupport/Windows/WASAPI/Private/WasapiAudioFormat.cpp b/Engine/Source/Runtime/AudioPlatformSupport/Windows/WASAPI/Private/WasapiAudioFormat.cpp
new file mode 100644
--- /dev/null
+++ b/Engine/Source/Runtime/AudioPlatformSupport/Windows/WASAPI/Private/WasapiAudioFormat.cpp
@@ -0,0 +1,275 @@
+// Copyright Epic Games, Inc. All Rights Reserved.
+
+#include "WasapiAudioFormat.h"
+
+#include "AudioPlatformSupportWasapiLog.h"
+
+namespace Audio
+{
+	// For explanation on WAVEFORMATEX structure and cbSize requirements, see: 
+	// https://learn.microsoft.com/en-us/windows/win32/api/mmeapi/ns-mmeapi-waveformatex
+	static constexpr uint16 WaveFormatExtensibleCbsize = 22;
+	static constexpr uint16 DefaultChannelCount = 2;
+	static constexpr uint32 DefaultSampleRate = 44100;
+	static constexpr EWasapiAudioEncoding DefaultEncoding = EWasapiAudioEncoding::PCM_16;
+
+	FWasapiAudioFormat::FWasapiAudioFormat()
+	{
+		InitAudioFormat(DefaultChannelCount, DefaultSampleRate, DefaultEncoding);
+		InitAudioEncoding();
+	}
+
+	FWasapiAudioFormat::FWasapiAudioFormat(WAVEFORMATEX* InFormat)
+	{
+		if ((InFormat->wFormatTag == WAVE_FORMAT_EXTENSIBLE) && (InFormat->cbSize >= WaveFormatExtensibleCbsize))
+		{
+			FMemory::Memcpy(&WaveFormat, InFormat, sizeof(WAVEFORMATEXTENSIBLE)); //-V512
+		}
+		else
+		{
+			FMemory::Memcpy(&WaveFormat, InFormat, sizeof(WAVEFORMATEX));
+		}
+
+		InitAudioEncoding();
+	}
+
+	FWasapiAudioFormat::FWasapiAudioFormat(uint16 InChannels, uint32 InSampleRate, EWasapiAudioEncoding InEncoding)
+	{
+		InitAudioFormat(InChannels, InSampleRate, InEncoding);
+		InitAudioEncoding();
+	}
+
+	void FWasapiAudioFormat::InitAudioFormat(uint16 InChannels, uint32 InSampleRate, EWasapiAudioEncoding InEncoding)
+	{
+		static constexpr uint32 BitsPerByte = 8;
+		uint16 BitDepth = EncodingToBitDepth(InEncoding);
+		ensure(BitDepth != 0);
+
+		if (BitDepth != 0)
+		{
+			WaveFormat.Format.wFormatTag = WAVE_FORMAT_EXTENSIBLE;
+			WaveFormat.Format.cbSize = WaveFormatExtensibleCbsize;
+			WaveFormat.Format.nChannels = InChannels;
+			WaveFormat.Format.nSamplesPerSec = InSampleRate;
+			WaveFormat.Format.wBitsPerSample = BitDepth;
+			WaveFormat.Format.nBlockAlign = InChannels * (BitDepth / BitsPerByte);
+			WaveFormat.Format.nAvgBytesPerSec = InSampleRate * WaveFormat.Format.nBlockAlign;
+			WaveFormat.Samples.wValidBitsPerSample = BitDepth;
+
+			if (InEncoding == EWasapiAudioEncoding::FLOATING_POINT_32 || InEncoding == EWasapiAudioEncoding::FLOATING_POINT_64)
+			{
+				WaveFormat.SubFormat = KSDATAFORMAT_SUBTYPE_IEEE_FLOAT;
+			}
+			else
+			{
+				WaveFormat.SubFormat = KSDATAFORMAT_SUBTYPE_PCM;
+				if (InEncoding == EWasapiAudioEncoding::PCM_24_IN_32)
+				{
+					WaveFormat.Samples.wValidBitsPerSample = 24;
+				}
+			}
+		}
+	}
+
+	void FWasapiAudioFormat::InitAudioEncoding()
+	{
+		Encoding = DetermineAudioEncoding(WaveFormat);
+		if (Encoding == EWasapiAudioEncoding::UNKNOWN)
+		{
+			UE_LOG(LogAudioPlatformSupportWasapi, Error, TEXT("FWasapiAudioFormat unknown audio format"));
+		}
+	}
+
+	EWasapiAudioEncoding FWasapiAudioFormat::GetEncoding() const
+	{
+		return Encoding;
+	}
+
+	uint32 FWasapiAudioFormat::GetNumChannels() const
+	{
+		return WaveFormat.Format.nChannels;
+	}
+
+	uint32 FWasapiAudioFormat::GetSampleRate() const
+	{
+		return WaveFormat.Format.nSamplesPerSec;
+	}
+
+	uint32 FWasapiAudioFormat::GetBitsPerSample() const
+	{
+		return WaveFormat.Format.wBitsPerSample;
+	}
+
+	uint32 FWasapiAudioFormat::GetBytesPerSample() const
+	{
+		return WaveFormat.Format.wBitsPerSample / 8;
+	}
+
+	uint32 FWasapiAudioFormat::GetFrameSizeInBytes() const
+	{
+		return WaveFormat.Format.nBlockAlign;
+	}
+
+	const WAVEFORMATEX* FWasapiAudioFormat::GetWaveFormat() const
+	{
+		return &(WaveFormat.Format);
+	}
+
+	EWasapiAudioEncoding FWasapiAudioFormat::DetermineAudioEncoding(const WAVEFORMATEXTENSIBLE& InFormat)
+	{
+		if (InFormat.Format.wFormatTag == WAVE_FORMAT_IEEE_FLOAT)
+		{
+			if (InFormat.Format.wBitsPerSample == (sizeof(float) * 8))
+			{
+				return EWasapiAudioEncoding::FLOATING_POINT_32;
+			}
+			else if (InFormat.Format.wBitsPerSample == (sizeof(double) * 8))
+			{
+				return EWasapiAudioEncoding::FLOATING_POINT_64;
+			}
+			else
+			{
+				return EWasapiAudioEncoding::UNKNOWN;
+			}
+		}
+		else
+		{
+			if (InFormat.Format.wFormatTag == WAVE_FORMAT_PCM)
+			{
+				switch (InFormat.Format.wBitsPerSample)
+				{
+				case 8:
+					return EWasapiAudioEncoding::PCM_8;
+
+				case 16:
+					return EWasapiAudioEncoding::PCM_16;
+
+				case 24:
+					return EWasapiAudioEncoding::PCM_24;
+
+				case 32:
+					return EWasapiAudioEncoding::PCM_32;
+
+				default:
+					return EWasapiAudioEncoding::UNKNOWN;
+				}
+			}
+			else
+			{
+				if ((InFormat.Format.wFormatTag == WAVE_FORMAT_EXTENSIBLE) && (InFormat.Format.cbSize >= WaveFormatExtensibleCbsize))
+				{
+					if (InFormat.SubFormat == KSDATAFORMAT_SUBTYPE_IEEE_FLOAT)
+					{
+						if (InFormat.Format.wBitsPerSample == (sizeof(float) * 8))
+						{
+							return EWasapiAudioEncoding::FLOATING_POINT_32;
+						}
+						else if (InFormat.Format.wBitsPerSample == (sizeof(double) * 8))
+						{
+							return EWasapiAudioEncoding::FLOATING_POINT_64;
+						}
+						else
+						{
+							return EWasapiAudioEncoding::UNKNOWN;
+						}
+					}
+					else
+					{
+						if (InFormat.SubFormat == KSDATAFORMAT_SUBTYPE_PCM)
+						{
+							switch (InFormat.Samples.wValidBitsPerSample)
+							{
+							case 8:
+								return (InFormat.Format.wBitsPerSample == 8) ? EWasapiAudioEncoding::PCM_8 : EWasapiAudioEncoding::UNKNOWN;
+
+							case 16:
+								return (InFormat.Format.wBitsPerSample == 16) ? EWasapiAudioEncoding::PCM_16 : EWasapiAudioEncoding::UNKNOWN;
+
+							case 32:
+								return (InFormat.Format.wBitsPerSample == 32) ? EWasapiAudioEncoding::PCM_32 : EWasapiAudioEncoding::UNKNOWN;
+
+							case 24:
+								if (InFormat.Format.wBitsPerSample == 32)
+								{
+									return EWasapiAudioEncoding::PCM_24_IN_32;
+								}
+								else
+								{
+									if (InFormat.Format.wBitsPerSample == 24)
+									{
+										return EWasapiAudioEncoding::PCM_24;
+									}
+									else
+									{
+										return EWasapiAudioEncoding::UNKNOWN;
+									}
+								}
+							default:
+								break;
+							}
+						}
+					}
+				}
+			}
+		}
+
+		return EWasapiAudioEncoding::UNKNOWN;
+	}
+
+	uint16 FWasapiAudioFormat::EncodingToBitDepth(EWasapiAudioEncoding InEncoding)
+	{
+		switch (InEncoding)
+		{
+		case EWasapiAudioEncoding::PCM_8:
+			return 8;
+
+		case EWasapiAudioEncoding::PCM_16:
+			return 16;
+
+		case EWasapiAudioEncoding::PCM_24:
+			return 24;
+
+		case EWasapiAudioEncoding::PCM_24_IN_32:
+		case EWasapiAudioEncoding::PCM_32:
+		case EWasapiAudioEncoding::FLOATING_POINT_32:
+			return 32;
+
+		case EWasapiAudioEncoding::FLOATING_POINT_64:
+			return 64;
+
+		default:
+			return 0;
+		}
+	}
+
+	FString FWasapiAudioFormat::GetEncodingString() const
+	{
+		switch (GetEncoding())
+		{
+		case EWasapiAudioEncoding::PCM_8:
+			return TEXT("8-Bit PCM");
+
+		case EWasapiAudioEncoding::PCM_16:
+			return TEXT("16-Bit PCM");
+
+		case EWasapiAudioEncoding::PCM_24:
+			return TEXT("24-Bit PCM");
+
+		case EWasapiAudioEncoding::PCM_24_IN_32:
+			return TEXT("24-Bit PCM (32-bit Container)");
+
+		case EWasapiAudioEncoding::PCM_32:
+			return TEXT("PCM_32 PCM");
+
+		case EWasapiAudioEncoding::FLOATING_POINT_32:
+			return TEXT("32-Bit Floating Point");
+
+		case EWasapiAudioEncoding::FLOATING_POINT_64:
+			return TEXT("64-Bit Floating Point");
+
+		default:
+			return TEXT("UNKNOWN");
+		}
+	}
+
+}
diff --git a/Engine/Source/Runtime/AudioPlatformSupport/Windows/WASAPI/Private/WasapiAudioUtils.cpp b/Engine/Source/Runtime/AudioPlatformSupport/Windows/WASAPI/Private/WasapiAudioUtils.cpp
new file mode 100644
--- /dev/null
+++ b/Engine/Source/Runtime/AudioPlatformSupport/Windows/WASAPI/Private/WasapiAudioUtils.cpp
@@ -0,0 +1,25 @@
+// Copyright Epic Games, Inc. All Rights Reserved.
+
+#include "WasapiAudioUtils.h"
+
+#include "AudioPlatformSupportWasapiLog.h"
+
+namespace Audio
+{
+	// REFERENCE_TIME base which is in units of 100 nanoseconds
+	static constexpr REFERENCE_TIME ReferenceTimeBase = 1e7;
+
+	REFERENCE_TIME FWasapiAudioUtils::FramesToRefTime(const uint32 InNumFrames, const uint32 InSampleRate)
+	{
+		REFERENCE_TIME RefTimeDuration = FMath::RoundToInt(static_cast<double>(InNumFrames * ReferenceTimeBase) / InSampleRate);
+
+		return RefTimeDuration;
+	}
+
+	uint64 FWasapiAudioUtils::RefTimeToFrames(const REFERENCE_TIME InRefTime, const uint32 InSampleRate)
+	{
+		uint64 NumFrames = FMath::RoundToInt(static_cast<double>(InRefTime * InSampleRate) / ReferenceTimeBase);
+
+		return NumFrames;
+	}
+}
diff --git a/Engine/Source/Runtime/AudioPlatformSupport/Windows/WASAPI/Public/AudioPlatformSupportWasapiLog.h b/Engine/Source/Runtime/AudioPlatformSupport/Windows/WASAPI/Public/AudioPlatformSupportWasapiLog.h
new file mode 100644
--- /dev/null
+++ b/Engine/Source/Runtime/AudioPlatformSupport/Windows/WASAPI/Public/AudioPlatformSupportWasapiLog.h
@@ -0,0 +1,7 @@
+// Copyright Epic Games, Inc. All Rights Reserved.
+
+#pragma once
+
+#include "Logging/LogMacros.h"
+
+AUDIOPLATFORMSUPPORTWASAPI_API DECLARE_LOG_CATEGORY_EXTERN(LogAudioPlatformSupportWasapi, Log, All);
diff --git a/Engine/Source/Runtime/AudioPlatformSupport/Windows/WASAPI/Public/WasapiAudioFormat.h b/Engine/Source/Runtime/AudioPlatformSupport/Windows/WASAPI/Public/WasapiAudioFormat.h
new file mode 100644
--- /dev/null
+++ b/Engine/Source/Runtime/AudioPlatformSupport/Windows/WASAPI/Public/WasapiAudioFormat.h
@@ -0,0 +1,102 @@
+// Copyright Epic Games, Inc. All Rights Reserved.
+
+#pragma once
+
+#include "Windows/AllowWindowsPlatformTypes.h"
+#include "Windows/AllowWindowsPlatformAtomics.h"
+
+THIRD_PARTY_INCLUDES_START
+#include <AudioClient.h>
+THIRD_PARTY_INCLUDES_END
+
+#include "Windows/HideWindowsPlatformAtomics.h"
+#include "Windows/HideWindowsPlatformTypes.h"
+
+
+namespace Audio
+{
+	/**
+	 * EWasapiAudioEncoding - Audio bit depths supported by WASAPI implementation
+	 */
+	enum class EWasapiAudioEncoding : uint8
+	{
+		UNKNOWN,
+		PCM_8,
+		PCM_16,
+		PCM_24,
+		PCM_24_IN_32,
+		PCM_32,
+		FLOATING_POINT_32,
+		FLOATING_POINT_64
+	};
+
+	/**
+	 * FWasapiAudioFormat - Wrapper class for Windows WAVEFORMATEX and WAVEFORMATEXTENSIBLE structs.
+	 */
+	class AUDIOPLATFORMSUPPORTWASAPI_API FWasapiAudioFormat
+	{
+	public:
+		/** FWasapiAudioFormat - Constructs underlying WAVEFORMATEX with default values. */
+		FWasapiAudioFormat();
+
+		/**
+		 * FWasapiAudioFormat - Constructs underlying data structure with given WAVEFORMATEX parameter.
+		 * 
+		 * @param InFormat - Pointer to a WAVEFORMATEX struct from which to construct this object. 
+		 * Correctly handles the extended form, WAVEFORMATEXTENSIBLE. 
+		 */
+		explicit FWasapiAudioFormat(WAVEFORMATEX* InFormat);
+
+		/**
+		 * FWasapiAudioFormat - Constructs underlying data structure using given parameters.
+		 * 
+		 * @param InChannels - Number of audio channels.
+		 * @param InSampleRate - Samples per second (Hz).
+		 * @param InEncoding - For PCM data, this is the bit depth of data (e.g. 16-bit integer, 32-bit float, etc.)
+		 */
+		FWasapiAudioFormat(uint16 InChannels, uint32 InSampleRate, EWasapiAudioEncoding InEncoding);
+
+		/** GetEncoding - Returns the encoding for this audio format. */
+		EWasapiAudioEncoding GetEncoding() const;
+
+		/** GetNumChannels - Returns the number of audio channels present in this format. */
+		uint32 GetNumChannels() const;
+
+		/** GetSampleRate - Returns the audio sample rate for this format. */
+		uint32 GetSampleRate() const;
+
+		/** GetBitsPerSample - Returns the bit depth for the audio data as an integer value. */
+		uint32 GetBitsPerSample() const;
+
+		/** GetBytesPerSample - Returns the number of bytes needed to store a single sample of audio data. */
+		uint32 GetBytesPerSample() const;
+
+		/** GetFrameSizeInBytes - Returns the number of bytes needed to store a singe frame of audio data. */
+		uint32 GetFrameSizeInBytes() const;
+		
+		/** GetEncodingString - Returns string form of encoding. Suitable for logging. */
+		FString GetEncodingString() const;
+
+		/**
+		 * GetWaveFormat - Returns pointer to raw WAVEFORMATEX struct. Use this for specific OS API calls
+		 * which require it. Otherwise use the accessors above.
+		 */
+		const WAVEFORMATEX* GetWaveFormat() const;
+
+	private:
+
+		/** InitAudioFormat - Init method used by constructors to initialize data. */
+		void InitAudioFormat(uint16 InChannels, uint32 InSampleRate, EWasapiAudioEncoding InEncoding);
+		/** InitAudioEncoding - Parses underlying WAVEFORMATEXTENSIBLE structure and determines corresponding encoding. */
+		void InitAudioEncoding();
+		/** DetermineAudioEncoding - Helper function for determining audio encoding from WAVEFORMATEXTENSIBLE struct. */
+		static EWasapiAudioEncoding DetermineAudioEncoding(const WAVEFORMATEXTENSIBLE& InFormat);
+		/** EncodingToBitDepth - Helper function for getting bit depth from InEncoding parameter. */
+		static uint16 EncodingToBitDepth(EWasapiAudioEncoding InEncoding);
+
+		/** WaveFormat - Windows wave format struct which this object wraps. */
+		WAVEFORMATEXTENSIBLE WaveFormat = {};
+		/** Encoding - Audio encoding for this object. */
+		EWasapiAudioEncoding Encoding = EWasapiAudioEncoding::UNKNOWN;
+	};
+}
diff --git a/Engine/Source/Runtime/AudioPlatformSupport/Windows/WASAPI/Public/WasapiAudioUtils.h b/Engine/Source/Runtime/AudioPlatformSupport/Windows/WASAPI/Public/WasapiAudioUtils.h
new file mode 100644
--- /dev/null
+++ b/Engine/Source/Runtime/AudioPlatformSupport/Windows/WASAPI/Public/WasapiAudioUtils.h
@@ -0,0 +1,30 @@
+// Copyright Epic Games, Inc. All Rights Reserved.
+
+#pragma once
+
+#include "Windows/AllowWindowsPlatformTypes.h"
+#include "Windows/AllowWindowsPlatformAtomics.h"
+
+THIRD_PARTY_INCLUDES_START
+#include <AudioClient.h>
+THIRD_PARTY_INCLUDES_END
+
+#include "Windows/HideWindowsPlatformAtomics.h"
+#include "Windows/HideWindowsPlatformTypes.h"
+
+
+namespace Audio
+{
+	/**
+	 * FWasapiAudioUtils - WASAPI utility class
+	 */
+	class AUDIOPLATFORMSUPPORTWASAPI_API FWasapiAudioUtils
+	{
+	public:
+		/** FramesToRefTime - Converts a given number of frames at the given sample rate to REFERENCE_TIME. */
+		static REFERENCE_TIME FramesToRefTime(const uint32 InNumFrames, const uint32 InSampleRate);
+
+		/** RefTimeToFrames - Converts the given REFERENCE_TIME to a number of frames at the given sample rate. */
+		static uint64 RefTimeToFrames(const REFERENCE_TIME InRefTime, const uint32 InSampleRate);
+	};
+}
diff --git a/Engine/Source/Runtime/Launch/Launch.Build.cs b/Engine/Source/Runtime/Launch/Launch.Build.cs
--- a/Engine/Source/Runtime/Launch/Launch.Build.cs
+++ b/Engine/Source/Runtime/Launch/Launch.Build.cs
@@ -84,6 +84,7 @@ public class Launch : ModuleRules
 			{
 				DynamicallyLoadedModuleNames.AddRange(new string[] {
 					"AudioMixerXAudio2",
+					"AudioMixerWasapi",
 				});
 			}
 			else if (Target.Platform == UnrealTargetPlatform.Mac)
diff --git a/Engine/Source/Runtime/Windows/AudioMixerWasapi/AudioMixerWasapi.Build.cs b/Engine/Source/Runtime/Windows/AudioMixerWasapi/AudioMixerWasapi.Build.cs
new file mode 100644
--- /dev/null
+++ b/Engine/Source/Runtime/Windows/AudioMixerWasapi/AudioMixerWasapi.Build.cs
@@ -0,0 +1,39 @@
+// Copyright Epic Games, Inc. All Rights Reserved.
+
+using UnrealBuildTool;
+
+public class AudioMixerWasapi : ModuleRules
+{
+	public AudioMixerWasapi(ReadOnlyTargetRules Target) : base(Target)
+	{
+		PrivateIncludePathModuleNames.Add("TargetPlatform");
+
+		if (Target.bCompileAgainstEngine)
+        {
+            PrivateDependencyModuleNames.AddRange(
+				new string[]
+				{
+					"Engine",
+					"AudioPlatformSupportWasapi",
+					"WindowsMMDeviceEnumeration"
+				}
+			);
+        }
+
+        PrivateDependencyModuleNames.AddRange(
+			new string[] {
+					"Core",
+					"CoreUObject",
+					"AudioMixer",
+					"AudioMixerCore"
+                }
+		);
+
+		PrecompileForTargets = PrecompileTargetsType.None;
+
+		if (Target.Platform == UnrealTargetPlatform.Win64)
+		{
+			PrecompileForTargets = PrecompileTargetsType.Any;
+		}
+	}
+}
diff --git a/Engine/Source/Runtime/Windows/AudioMixerWasapi/Private/AudioMixerWasapi.cpp b/Engine/Source/Runtime/Windows/AudioMixerWasapi/Private/AudioMixerWasapi.cpp
new file mode 100644
--- /dev/null
+++ b/Engine/Source/Runtime/Windows/AudioMixerWasapi/Private/AudioMixerWasapi.cpp
@@ -0,0 +1,312 @@
+// Copyright Epic Games, Inc. All Rights Reserved.
+
+#include "AudioMixerWasapi.h"
+#include "WasapiAudioUtils.h"
+
+namespace Audio
+{
+	FAudioMixerWasapi::FAudioMixerWasapi()
+	{
+	}
+
+	FAudioMixerWasapi::~FAudioMixerWasapi()
+	{
+	}
+
+	bool FAudioMixerWasapi::InitializeHardware()
+	{
+		SCOPED_NAMED_EVENT(FAudioMixerWasapi_InitializeHardware, FColor::Blue);
+
+		RegisterDeviceChangedListener();
+
+		FAudioPlatformSettings EngineSettings = GetPlatformSettings();
+
+		MainRenderStreamDevice = MakeUnique<FAudioMixerWasapiRenderStream>();
+		if (!MainRenderStreamDevice.IsValid())
+		{
+			UE_LOG(LogAudioMixer, Error, TEXT("FAudioMixerWasapi::InitializeHardware failed to create MainRenderStreamDevice"));
+			return false;
+		}
+
+		MainRenderStreamDevice->OnReadNextBuffer().BindRaw(this, &FAudioMixerWasapi::ReadNextBuffer);
+
+		FWasapiRenderStreamParams StreamParams;
+		if (!InitStreamParams(AUDIO_MIXER_DEFAULT_DEVICE_INDEX, EngineSettings.CallbackBufferFrameSize, EngineSettings.SampleRate, StreamParams))
+		{
+			return false;
+		}
+
+		if (!InitAudioStreamDevice(StreamParams))
+		{
+			UE_LOG(LogAudioMixer, Error, TEXT("FAudioMixerWasapi::InitializeHardware failed InitAudioStreamDevice"));
+			return false;
+		}
+
+		return true;
+	}
+
+	bool FAudioMixerWasapi::TeardownHardware()
+	{
+		if (!bIsInitialized)
+		{
+			AUDIO_PLATFORM_LOG_ONCE(TEXT("FAudioMixerWasapi::TeardownHardware failed...not initialized."), Warning);
+			return false;
+		}
+
+		if (MainRenderStreamDevice.IsValid())
+		{
+			// Teardown the main device which will also unbind our delegate
+			MainRenderStreamDevice->TeardownHardware();
+			MainRenderStreamDevice.Reset();
+		}
+
+		bIsInitialized = false;
+
+		return true;
+	}
+
+	bool FAudioMixerWasapi::IsInitialized() const
+	{
+		return bIsInitialized && MainRenderStreamDevice.IsValid() && MainRenderStreamDevice->IsInitialized();
+	}
+
+	int32 FAudioMixerWasapi::GetNumFrames(const int32 InNumReqestedFrames)
+	{
+		if (MainRenderStreamDevice.IsValid())
+		{
+			return MainRenderStreamDevice->GetNumFrames(InNumReqestedFrames);
+		}
+
+		return InNumReqestedFrames;
+	}
+
+	bool FAudioMixerWasapi::GetNumOutputDevices(uint32& OutNumOutputDevices)
+	{
+		SCOPED_NAMED_EVENT(FAudioMixerWasapi_GetNumOutputDevices, FColor::Blue);
+
+		OutNumOutputDevices = 0;
+
+		if (IAudioPlatformDeviceInfoCache* Cache = GetDeviceInfoCache())
+		{
+			OutNumOutputDevices = Cache->GetAllActiveOutputDevices().Num();
+			return true;
+		}
+		else
+		{
+			AUDIO_PLATFORM_LOG_ONCE(TEXT("FAudioMixerWasapi device cache not initialized"), Warning);
+			return false;
+		}
+	}
+
+	bool FAudioMixerWasapi::GetOutputDeviceInfo(const uint32 InDeviceIndex, FAudioPlatformDeviceInfo& OutInfo)
+	{
+		SCOPED_NAMED_EVENT(FAudioMixerWasapi_GetOutputDeviceInfo, FColor::Blue);
+
+		if (IAudioPlatformDeviceInfoCache* Cache = GetDeviceInfoCache())
+		{
+			if (InDeviceIndex == AUDIO_MIXER_DEFAULT_DEVICE_INDEX)
+			{
+				if (TOptional<FAudioPlatformDeviceInfo> Defaults = Cache->FindDefaultOutputDevice())
+				{
+					OutInfo = *Defaults;
+					return true;
+				}
+			}
+			else
+			{
+				TArray<FAudioPlatformDeviceInfo> ActiveDevices = Cache->GetAllActiveOutputDevices();
+				if (ActiveDevices.IsValidIndex(InDeviceIndex))
+				{
+					OutInfo = ActiveDevices[InDeviceIndex];
+					return true;
+				}
+			}
+		}
+
+		return false;
+	}
+
+	FString FAudioMixerWasapi::GetCurrentDeviceName() const
+	{
+		return AudioStreamInfo.DeviceInfo.Name;
+	}
+
+	bool FAudioMixerWasapi::GetDefaultOutputDeviceIndex(uint32& OutDefaultDeviceIndex) const
+	{
+		OutDefaultDeviceIndex = AUDIO_MIXER_DEFAULT_DEVICE_INDEX;
+		return true;
+	}
+
+	bool FAudioMixerWasapi::InitStreamParams(const uint32 InDeviceIndex, const int32 InNumBufferFrames, const int32 InSampleRate, FWasapiRenderStreamParams& OutParams)
+	{
+		FAudioPlatformDeviceInfo DeviceInfo;
+		if (!GetOutputDeviceInfo(InDeviceIndex, DeviceInfo))
+		{
+			UE_LOG(LogAudioMixer, Error, TEXT("FAudioMixerWasapi::InitStreamParams unable to find default device"));
+			return false;
+		}
+
+		TComPtr<IMMDevice> MMDevice = GetMMDevice(DeviceInfo.DeviceId);
+		if (!MMDevice)
+		{
+			UE_LOG(LogAudioMixer, Error, TEXT("FAudioMixerWasapi::InitStreamParams null MMDevice"));
+			return false;
+		}
+
+		OutParams = FWasapiRenderStreamParams(MMDevice, DeviceInfo, InNumBufferFrames, InSampleRate);
+
+		return true;
+	}
+
+	bool FAudioMixerWasapi::InitAudioStreamDevice(const FWasapiRenderStreamParams& InStreamParams)
+	{
+		if (ensure(MainRenderStreamDevice.IsValid()))
+		{
+			return MainRenderStreamDevice->InitializeHardware(InStreamParams);
+		}
+
+		return false;
+	}
+
+	bool FAudioMixerWasapi::OpenAudioStream(const FAudioMixerOpenStreamParams& Params)
+	{
+		OpenStreamParams = Params;
+
+		AudioStreamInfo.Reset();
+
+		AudioStreamInfo.OutputDeviceIndex = OpenStreamParams.OutputDeviceIndex;
+		AudioStreamInfo.NumOutputFrames = OpenStreamParams.NumFrames;
+		AudioStreamInfo.NumBuffers = OpenStreamParams.NumBuffers;
+		AudioStreamInfo.AudioMixer = OpenStreamParams.AudioMixer;
+
+		FWasapiRenderStreamParams StreamParams;
+		if (!InitStreamParams(OpenStreamParams.OutputDeviceIndex, OpenStreamParams.NumFrames, OpenStreamParams.SampleRate, StreamParams))
+		{
+			return false;
+		}
+
+		AudioStreamInfo.DeviceInfo = StreamParams.HardwareDeviceInfo;
+		
+		HANDLE EventHandle = nullptr;
+		const TFunction<void()> RenderCallback = [this]() { MainRenderStreamDevice->DeviceRenderCallback(); };
+		RenderDeviceThread = MakeUnique<FAudioMixerWasapiDeviceThread>(RenderCallback, EventHandle);
+		
+		if (!RenderDeviceThread.IsValid())
+		{
+			UE_LOG(LogAudioMixer, Error, TEXT("Unable to create RenderDeviceThread"));
+			return false;
+		}
+
+		if (EventHandle == nullptr)
+		{
+			UE_LOG(LogAudioMixer, Error, TEXT("OpenAudioStream null EventHandle"));
+			return false;
+		}
+
+		if (!MainRenderStreamDevice->OpenAudioStream(StreamParams, EventHandle))
+		{
+			UE_LOG(LogAudioMixer, Error, TEXT("OpenAudioStream failed to open main audio device"));
+			return false;
+		}
+
+		bIsInitialized = true;
+
+		UE_LOG(LogAudioMixer, Display, TEXT("FAudioMixerWasapi initialized SampeRate: %d"), OpenStreamParams.SampleRate);
+
+		return true;
+	}
+
+	bool FAudioMixerWasapi::CloseAudioStream()
+	{
+		if (!bIsInitialized || AudioStreamInfo.StreamState == EAudioOutputStreamState::Closed)
+		{
+			return false;
+		}
+
+		MainRenderStreamDevice->CloseAudioStream();
+
+		AudioStreamInfo.StreamState = EAudioOutputStreamState::Closed;
+
+		return true;
+	}
+
+	bool FAudioMixerWasapi::StartAudioStream()
+	{
+		if (IsInitialized() && RenderDeviceThread.IsValid())
+		{
+			BeginGeneratingAudio();
+
+			AudioStreamInfo.StreamState = EAudioOutputStreamState::Running;
+
+			MainRenderStreamDevice->StartAudioStream();
+			RenderDeviceThread->Start();
+
+			return true;
+		}
+
+		return false;
+	}
+
+	bool FAudioMixerWasapi::StopAudioStream()
+	{
+		if (!bIsInitialized)
+		{
+			AUDIO_PLATFORM_LOG_ONCE(TEXT("FAudioMixerWasapi::StopAudioStream() not initialized."), Warning);
+			return false;
+		}
+
+		if (AudioStreamInfo.StreamState != EAudioOutputStreamState::Stopped && AudioStreamInfo.StreamState != EAudioOutputStreamState::Closed)
+		{
+			if (AudioStreamInfo.StreamState == EAudioOutputStreamState::Running)
+			{
+				StopGeneratingAudio();
+			}
+
+			if (RenderDeviceThread.IsValid())
+			{
+				RenderDeviceThread->Stop();
+			}
+
+			if (MainRenderStreamDevice.IsValid())
+			{
+				MainRenderStreamDevice->StopAudioStream();
+			}
+
+			check(AudioStreamInfo.StreamState == EAudioOutputStreamState::Stopped);
+		}
+
+		return true;
+	}
+
+	FAudioPlatformDeviceInfo FAudioMixerWasapi::GetPlatformDeviceInfo() const
+	{
+		return AudioStreamInfo.DeviceInfo;
+	}
+
+	void FAudioMixerWasapi::SubmitBuffer(const uint8* InBuffer)
+	{
+		if (MainRenderStreamDevice.IsValid())
+		{
+			MainRenderStreamDevice->SubmitBuffer(InBuffer, OpenStreamParams.NumFrames);
+		}
+	}
+
+	FString FAudioMixerWasapi::GetDefaultDeviceName()
+	{
+		return FString();
+	}
+
+	FAudioPlatformSettings FAudioMixerWasapi::GetPlatformSettings() const
+	{
+#if WITH_ENGINE
+		return FAudioPlatformSettings::GetPlatformSettings(FPlatformProperties::GetRuntimeSettingsClassName());
+#else
+		return FAudioPlatformSettings();
+#endif // WITH_ENGINE
+	}
+
+	IAudioPlatformDeviceInfoCache* FAudioMixerWasapi::GetDeviceInfoCache() const
+	{
+		return DeviceInfoCache.Get();
+	}
+}
diff --git a/Engine/Source/Runtime/Windows/AudioMixerWasapi/Private/AudioMixerWasapi.h b/Engine/Source/Runtime/Windows/AudioMixerWasapi/Private/AudioMixerWasapi.h
new file mode 100644
--- /dev/null
+++ b/Engine/Source/Runtime/Windows/AudioMixerWasapi/Private/AudioMixerWasapi.h
@@ -0,0 +1,97 @@
+// Copyright Epic Games, Inc. All Rights Reserved.
+
+#pragma once
+
+#include "AudioMixer.h"	
+#include "AudioMixerWasapiDeviceThread.h"
+#include "AudioMixerWasapiRenderStream.h"
+#include "Misc/ScopeRWLock.h"
+
+#include "Windows/AllowWindowsPlatformTypes.h"
+#include "Windows/AllowWindowsPlatformAtomics.h"
+
+#include "Microsoft/COMPointer.h"
+
+THIRD_PARTY_INCLUDES_START
+#include <mmdeviceapi.h>			// IMMNotificationClient
+#include <audiopolicy.h>			// IAudioSessionEvents
+THIRD_PARTY_INCLUDES_END
+
+#include "Windows/HideWindowsPlatformAtomics.h"
+#include "Windows/HideWindowsPlatformTypes.h"
+
+/**
+ *  Warning: AudioMixerWasapi module is currently considered experimental and may change in the future.
+ *           We do not recommend shipping projects with Experimental features.
+ */
+
+namespace Audio
+{
+	/**
+	 * FAudioMixerWasapi - Wasapi audio backend for Windows and Xbox
+	 */
+	class FAudioMixerWasapi : public IAudioMixerPlatformInterface
+	{
+	public:
+
+		FAudioMixerWasapi();
+		virtual ~FAudioMixerWasapi();
+
+		//~ Begin IAudioMixerPlatformInterface
+		virtual FString GetPlatformApi() const override { return TEXT("WASAPIMixer"); }
+		virtual bool InitializeHardware() override;
+		virtual bool TeardownHardware() override;
+		virtual bool IsInitialized() const override;
+		virtual int32 GetNumFrames(const int32 InNumReqestedFrames) override;
+		virtual bool GetNumOutputDevices(uint32& OutNumOutputDevices) override;
+		virtual bool GetOutputDeviceInfo(const uint32 InDeviceIndex, FAudioPlatformDeviceInfo& OutInfo) override;
+		virtual FString GetCurrentDeviceName() const override;
+		virtual bool GetDefaultOutputDeviceIndex(uint32& OutDefaultDeviceIndex) const override;
+		virtual bool OpenAudioStream(const FAudioMixerOpenStreamParams& Params) override;
+		virtual bool CloseAudioStream() override;
+		virtual bool StartAudioStream() override;
+		virtual bool StopAudioStream() override;
+		virtual FAudioPlatformDeviceInfo GetPlatformDeviceInfo() const override;
+		virtual void SubmitBuffer(const uint8* Buffer) override;
+		virtual bool DisablePCMAudioCaching() const override { return true; }
+		virtual FString GetDefaultDeviceName() override;
+		virtual FAudioPlatformSettings GetPlatformSettings() const override;
+		virtual IAudioPlatformDeviceInfoCache* GetDeviceInfoCache() const override;
+		//~ End IAudioMixerPlatformInterface
+
+		//~ Begin IAudioMixerDeviceChangedListener
+		virtual void RegisterDeviceChangedListener() override;
+		virtual void UnregisterDeviceChangedListener() override;
+		virtual FString GetDeviceId() const override;
+		//~ End IAudioMixerDeviceChangedListener
+
+	private:
+
+		/** Cache for holding information about MM audio devices (IMMDevice).  */
+		TUniquePtr<IAudioPlatformDeviceInfoCache> DeviceInfoCache;
+
+		/** Mutex for protecting shared resources during a device swap. */
+		FCriticalSection AudioDeviceSwapCriticalSection;
+
+		/** Used for determining if a device swap has been requested. */
+		FString OriginalAudioDeviceId;
+
+		/** The main audio device for outputting up to 8 channels. */
+		TUniquePtr<FAudioMixerWasapiRenderStream> MainRenderStreamDevice;
+
+		/** Indicates if this object has been successfully initialized. */
+		bool bIsInitialized = false;
+
+		/** The thread which provides an execution context during audio playback. */
+		TUniquePtr<FAudioMixerWasapiDeviceThread> RenderDeviceThread;
+
+		/** Fetches an IMMDevice with the given ID. */
+		TComPtr<IMMDevice> GetMMDevice(const FString& InDeviceID) const;
+
+		/** Initializes a Wasapi stream paramters struct with the give values. */
+		bool InitStreamParams(const uint32 InDeviceIndex, const int32 InNumBufferFrames, const int32 InSampleRate, FWasapiRenderStreamParams& OutParams);
+
+		/** Initializes a device with the given parameters. */
+		bool InitAudioStreamDevice(const FWasapiRenderStreamParams& InStreamParams);
+	};
+ }
diff --git a/Engine/Source/Runtime/Windows/AudioMixerWasapi/Private/AudioMixerWasapiDeviceListener.cpp b/Engine/Source/Runtime/Windows/AudioMixerWasapi/Private/AudioMixerWasapiDeviceListener.cpp
new file mode 100644
--- /dev/null
+++ b/Engine/Source/Runtime/Windows/AudioMixerWasapi/Private/AudioMixerWasapiDeviceListener.cpp
@@ -0,0 +1,98 @@
+// Copyright Epic Games, Inc. All Rights Reserved.
+
+#include "AudioMixerWasapi.h"
+#include "AudioMixer.h"
+#include "AudioDeviceNotificationSubsystem.h"
+#include "Misc/ScopeRWLock.h"
+
+#include <atomic>
+
+#if PLATFORM_WINDOWS
+
+#include "Microsoft/COMPointer.h"
+#include "ScopedCom.h"					// FScopedComString
+
+#include "Windows/AllowWindowsPlatformTypes.h"
+THIRD_PARTY_INCLUDES_START
+// Linkage for  Windows GUIDs included by Notification/DeviceInfoCache, otherwise they are extern.
+#include <initguid.h>
+THIRD_PARTY_INCLUDES_END
+#include "Windows/HideWindowsPlatformTypes.h"
+
+#include "WindowsMMNotificationClient.h"
+#include "WindowsMMDeviceInfoCache.h"
+#include "WindowsMMStringUtils.h"
+
+namespace Audio
+{	
+	TSharedPtr<FWindowsMMNotificationClient> WindowsNotificationClient;
+
+	void FAudioMixerWasapi::RegisterDeviceChangedListener()
+	{
+		if (!WindowsNotificationClient.IsValid())
+		{
+			// Shared (This is a COM object, so we don't delete it, just derecement the ref counter).
+			WindowsNotificationClient = TSharedPtr<FWindowsMMNotificationClient>(
+				new FWindowsMMNotificationClient, 
+				[](FWindowsMMNotificationClient* InPtr) { InPtr->ReleaseClient(); }
+			);
+		}
+		if (!DeviceInfoCache.IsValid())
+		{
+			// Setup device info cache.
+			DeviceInfoCache = MakeUnique<FWindowsMMDeviceCache>();
+			WindowsNotificationClient->RegisterDeviceChangedListener(static_cast<FWindowsMMDeviceCache*>(DeviceInfoCache.Get()));
+		}
+
+		WindowsNotificationClient->RegisterDeviceChangedListener(this);
+	}
+
+	void FAudioMixerWasapi::UnregisterDeviceChangedListener()
+	{
+		if (WindowsNotificationClient.IsValid())
+		{
+			if (DeviceInfoCache.IsValid())
+			{
+				// Unregister and kill cache.
+				WindowsNotificationClient->UnRegisterDeviceDeviceChangedListener(static_cast<FWindowsMMDeviceCache*>(DeviceInfoCache.Get()));
+				
+				DeviceInfoCache.Reset();
+			}
+			
+			WindowsNotificationClient->UnRegisterDeviceDeviceChangedListener(this);
+		}
+	}
+
+	FString FAudioMixerWasapi::GetDeviceId() const
+	{
+		return AudioStreamInfo.DeviceInfo.DeviceId;
+	}
+
+	TComPtr<IMMDevice> FAudioMixerWasapi::GetMMDevice(const FString& InDeviceID) const
+	{
+		if (WindowsNotificationClient)
+		{
+			return WindowsNotificationClient->GetDevice(InDeviceID);
+		}
+
+		return TComPtr<IMMDevice>();
+	}
+}
+
+#else 
+namespace Audio
+{
+	void FAudioMixerWasapi::RegisterDeviceChangedListener() {}
+	void FAudioMixerWasapi::UnregisterDeviceChangedListener() {}
+	FString FAudioMixerWasapi::GetDeviceId() const
+	{
+		return AudioStreamInfo.DeviceInfo.DeviceId;
+	}
+
+	TComPtr<IMMDevice> FAudioMixerWasapi::GetMMDevice(const FString& InDeviceID) const
+	{
+		return TComPtr<IMMDevice>();
+	}
+}
+#endif
+
diff --git a/Engine/Source/Runtime/Windows/AudioMixerWasapi/Private/AudioMixerWasapiDeviceThread.cpp b/Engine/Source/Runtime/Windows/AudioMixerWasapi/Private/AudioMixerWasapiDeviceThread.cpp
new file mode 100644
--- /dev/null
+++ b/Engine/Source/Runtime/Windows/AudioMixerWasapi/Private/AudioMixerWasapiDeviceThread.cpp
@@ -0,0 +1,99 @@
+// Copyright Epic Games, Inc. All Rights Reserved.
+
+#include "AudioMixerWasapiDeviceThread.h"
+#include "AudioMixer.h"	
+#include "AudioMixerWasapiLog.h"	
+
+
+namespace Audio
+{
+	FAudioMixerWasapiRunnable::FAudioMixerWasapiRunnable(const TFunction<void()>& InDeviceRenderCallback, HANDLE& OutEventHandle) :
+		DeviceRenderCallback(InDeviceRenderCallback)
+	{
+		// Not using FEvent/FEventWin here because we need access to the raw, platform
+		// handle (see SetEventHandler() below).
+		EventHandle = ::CreateEvent(nullptr, 0, 0, nullptr);
+		if (EventHandle == nullptr)
+		{
+			UE_LOG(LogAudioMixerWasapi, Error, TEXT("FAudioMixerWasapiRunnable CreateEvent failed"));
+		}
+
+		OutEventHandle = EventHandle;
+	}
+
+	uint32 FAudioMixerWasapiRunnable::Run()
+	{
+		bIsRunning = true;
+
+		bool bCoInitialized = FWindowsPlatformMisc::CoInitialize(ECOMModel::Multithreaded);
+
+		while (bIsRunning.load())
+		{
+			static constexpr uint32 TimeoutInMs = 1000;
+
+			int32 Result = ::WaitForSingleObject(EventHandle, TimeoutInMs);
+			if (Result == WAIT_TIMEOUT)
+			{
+				++OutputStreamTimeoutsDetected;
+			}
+			else if (Result == WAIT_OBJECT_0)
+			{
+				DeviceRenderCallback();
+			}
+		}
+
+		if (bCoInitialized)
+		{
+			FWindowsPlatformMisc::CoUninitialize();
+		}
+
+		return 0;
+	}
+
+	void FAudioMixerWasapiRunnable::Stop()
+	{
+		bIsRunning = false;
+		if (OutputStreamTimeoutsDetected > 0)
+		{
+			UE_LOG(LogAudioMixerWasapi, Error, TEXT("FAudioMixerWasapiRunnable::Stop render stream reported %d timeouts"), OutputStreamTimeoutsDetected);
+		}
+	}
+
+	FAudioMixerWasapiDeviceThread::FAudioMixerWasapiDeviceThread(const TFunction<void()>& InDeviceRenderCallback, HANDLE& OutEventHandle) :
+		DeviceRenderRunnable(MakeUnique<FAudioMixerWasapiRunnable>(InDeviceRenderCallback, OutEventHandle))
+	{
+	}
+
+	bool FAudioMixerWasapiDeviceThread::Start()
+	{
+		TRACE_CPUPROFILER_EVENT_SCOPE(Audio::FAudioMixerWasapiDeviceThread::Start);
+		check(DeviceRenderThread == nullptr);
+
+		DeviceRenderThread = TUniquePtr<FRunnableThread>(FRunnableThread::Create(DeviceRenderRunnable.Get(), TEXT("Audio Render Device Thread"), 0, TPri_TimeCritical));
+		return DeviceRenderThread.IsValid();
+	}
+
+	void FAudioMixerWasapiDeviceThread::Stop()
+	{
+		if (DeviceRenderThread.IsValid())
+		{
+			TRACE_CPUPROFILER_EVENT_SCOPE(Audio::FAudioMixerWasapiDeviceThread::Stop);
+
+			bool bShouldWait = true;
+			DeviceRenderThread->Kill(bShouldWait);
+		}
+	}
+
+	void FAudioMixerWasapiDeviceThread::Abort()
+	{
+		if (DeviceRenderThread.IsValid())
+		{
+			TRACE_CPUPROFILER_EVENT_SCOPE(Audio::FAudioMixerWasapiDeviceThread::Abort);
+
+			// Always wait for thread to complete otherwise we can crash if
+			// the stream is disposed of mid-callback.
+			bool bShouldWait = true;
+			DeviceRenderThread->Kill(bShouldWait);
+		}
+	}
+}
diff --git a/Engine/Source/Runtime/Windows/AudioMixerWasapi/Private/AudioMixerWasapiDeviceThread.h b/Engine/Source/Runtime/Windows/AudioMixerWasapi/Private/AudioMixerWasapiDeviceThread.h
new file mode 100644
--- /dev/null
+++ b/Engine/Source/Runtime/Windows/AudioMixerWasapi/Private/AudioMixerWasapiDeviceThread.h
@@ -0,0 +1,110 @@
+// Copyright Epic Games, Inc. All Rights Reserved.
+
+#pragma once
+
+#include <atomic>
+
+#include "HAL/Runnable.h"
+#include "HAL/RunnableThread.h"
+#include "Templates/Function.h"
+#include "Templates/UniquePtr.h"
+
+#include "Windows/AllowWindowsPlatformTypes.h"
+#include "Windows/AllowWindowsPlatformAtomics.h"
+
+#include "Microsoft/COMPointer.h"
+
+THIRD_PARTY_INCLUDES_START
+#include <winnt.h>
+THIRD_PARTY_INCLUDES_END
+
+#include "Windows/HideWindowsPlatformAtomics.h"
+#include "Windows/HideWindowsPlatformTypes.h"
+
+/**
+ *  Warning: AudioMixerWasapi module is currently considered experimental and may change in the future.
+ *           We do not recommend shipping projects with Experimental features.
+ */
+
+	namespace Audio
+{
+	/**
+	 * FAudioMixerWasapiRunnable - The runnable which executes the main thread loop for FWasapiCaptureThread.
+	 */
+	class FAudioMixerWasapiRunnable : public FRunnable
+	{
+	public:
+		FAudioMixerWasapiRunnable() = delete;
+
+		/**
+		 * FAudioMixerWasapiRunnable
+		 */
+		explicit FAudioMixerWasapiRunnable(const TFunction<void()>& InDeviceRenderCallback, HANDLE& OutEventHandle);
+		
+		/** Default destructor */
+		virtual ~FAudioMixerWasapiRunnable() = default;
+
+		// Begin FRunnable overrides
+		virtual uint32 Run() override;
+		virtual void Stop() override;
+		// End FRunnable overrides
+
+	private:
+		/** The main run loop for this runnable will continue iterating while this flag is true. */
+		std::atomic<bool> bIsRunning;
+		
+		/**
+		 * Event handle which our audio thread waits on prior to each callback. WASAPI signals this
+		 * object each quanta when a buffer of audio has been rendered and is ready for more data.
+		 */
+		HANDLE EventHandle = nullptr;
+
+		/**
+		 * Accumulates timeouts which occur when the thread event timeout is reached
+		 * prior to the event being signaled for new data being available.
+		 */
+		uint32 OutputStreamTimeoutsDetected = 0;
+		
+		/** Callback function to be called each time the device signals it is ready for another buffer of audio. */
+		TFunction<void()> DeviceRenderCallback;
+	};
+
+	/**
+	 * FAudioMixerWasapiDeviceThread - Manages both the FAudioMixerWasapiRunnable object and the thread whose context it runs in. 
+	 */
+	class FAudioMixerWasapiDeviceThread
+	{
+	public:
+		FAudioMixerWasapiDeviceThread() = delete;
+
+		/**
+		 * FAudioMixerWasapiDeviceThread
+		 */
+		explicit FAudioMixerWasapiDeviceThread(const TFunction<void()>& InDeviceRenderCallback, HANDLE& OutEventHandle);
+
+		/**
+		 * Start - Creates the FRunnableThread object which immediately begins running the FAudioMixerWasapiRunnable member.
+		 * 
+		 * @return - Boolean indicating of the thread was succesfully created.
+		 */
+		bool Start();
+
+		/**
+		 * Stop - Gracefully shuts down the thread.
+		 */
+		void Stop();
+
+		/**
+		 * Abort - Performs non-graceful shutdown of thread which will close the underyling thread handle 
+		 * without waiting for it to finish.
+		 */
+		void Abort();
+
+	private:
+		/** The thread which is the context that the runnable executes in. */
+		TUniquePtr<FRunnableThread> DeviceRenderThread;
+
+		/** The runnable which manages the run loop for the render stream. */
+		TUniquePtr<FAudioMixerWasapiRunnable> DeviceRenderRunnable;
+	};
+}
diff --git a/Engine/Source/Runtime/Windows/AudioMixerWasapi/Private/AudioMixerWasapiLog.h b/Engine/Source/Runtime/Windows/AudioMixerWasapi/Private/AudioMixerWasapiLog.h
new file mode 100644
--- /dev/null
+++ b/Engine/Source/Runtime/Windows/AudioMixerWasapi/Private/AudioMixerWasapiLog.h
@@ -0,0 +1,8 @@
+// Copyright Epic Games, Inc. All Rights Reserved.
+#pragma once
+
+#include "Logging/LogMacros.h"
+#include "Logging/LogVerbosity.h"
+
+
+AUDIOMIXERWASAPI_API DECLARE_LOG_CATEGORY_EXTERN(LogAudioMixerWasapi, Log, All);
diff --git a/Engine/Source/Runtime/Windows/AudioMixerWasapi/Private/AudioMixerWasapiModule.cpp b/Engine/Source/Runtime/Windows/AudioMixerWasapi/Private/AudioMixerWasapiModule.cpp
new file mode 100644
--- /dev/null
+++ b/Engine/Source/Runtime/Windows/AudioMixerWasapi/Private/AudioMixerWasapiModule.cpp
@@ -0,0 +1,32 @@
+// Copyright Epic Games, Inc. All Rights Reserved.
+
+#include "AudioMixerWasapi.h"
+#include "AudioMixerWasapiLog.h"
+#include "Features/IModularFeatures.h"
+#include "Modules/ModuleManager.h"
+
+DEFINE_LOG_CATEGORY(LogAudioMixerWasapi);
+
+class FAudioMixerModuleWasapi : public IAudioDeviceModule
+{
+public:
+	virtual void StartupModule() override
+	{
+		IAudioDeviceModule::StartupModule();
+
+		FModuleManager::Get().LoadModuleChecked(TEXT("AudioMixer"));
+		FModuleManager::Get().LoadModuleChecked(TEXT("AudioMixerCore"));
+	}
+
+	virtual bool IsAudioMixerModule() const override
+	{
+		return true;
+	}
+
+	virtual Audio::IAudioMixerPlatformInterface* CreateAudioMixerPlatformInterface() override
+	{
+		return new Audio::FAudioMixerWasapi();
+	}
+};
+
+IMPLEMENT_MODULE(FAudioMixerModuleWasapi, AudioMixerWasapi);
diff --git a/Engine/Source/Runtime/Windows/AudioMixerWasapi/Private/AudioMixerWasapiRenderStream.cpp b/Engine/Source/Runtime/Windows/AudioMixerWasapi/Private/AudioMixerWasapiRenderStream.cpp
new file mode 100644
--- /dev/null
+++ b/Engine/Source/Runtime/Windows/AudioMixerWasapi/Private/AudioMixerWasapiRenderStream.cpp
@@ -0,0 +1,292 @@
+// Copyright Epic Games, Inc. All Rights Reserved.
+
+#include "AudioMixerWasapiRenderStream.h"
+
+#include "AudioMixerWasapiLog.h"
+#include "WasapiAudioUtils.h"
+
+namespace Audio
+{
+	FAudioMixerWasapiRenderStream::FAudioMixerWasapiRenderStream()
+	{
+	}
+
+	FAudioMixerWasapiRenderStream::~FAudioMixerWasapiRenderStream()
+	{
+	}
+
+	bool FAudioMixerWasapiRenderStream::InitializeHardware(const FWasapiRenderStreamParams& InParams)
+	{
+		TComPtr<IAudioClient3> TempAudioClient;
+		HRESULT Result = InParams.MMDevice->Activate(__uuidof(IAudioClient3), CLSCTX_INPROC_SERVER, nullptr, IID_PPV_ARGS_Helper(&TempAudioClient));
+		if (FAILED(Result))
+		{
+			UE_LOG(LogAudioMixerWasapi, Error, TEXT("OpenAudioStream failed IMMDevice::Activate 0x%x"), (uint32)Result);
+			return false;
+		}
+
+		AudioClientProperties AudioProps = { 0 };
+		AudioProps.cbSize = sizeof(AudioClientProperties);
+		AudioProps.bIsOffload = false;
+		AudioProps.eCategory = AudioCategory_Media;
+
+		Result = TempAudioClient->SetClientProperties(&AudioProps);
+		if (FAILED(Result))
+		{
+			UE_LOG(LogAudioMixerWasapi, Error, TEXT("OpenAudioStream failed IAudioClient3::SetClientProperties 0x%x"), (uint32)Result);
+			return false;
+		}
+
+		WAVEFORMATEX* MixFormat = nullptr;
+		Result = TempAudioClient->GetMixFormat(&MixFormat);
+		if (FAILED(Result) || !MixFormat)
+		{
+			UE_LOG(LogAudioMixerWasapi, Error, TEXT("OpenAudioStream failed IAudioClient3::GetMixFormat MixFormat: 0x%llx Result: 0x%x"), (uint64)MixFormat, (uint32)Result);
+			return false;
+		}
+
+		FWasapiAudioFormat StreamFormat(FMath::Min<int32>(MixFormat->nChannels, AUDIO_MIXER_MAX_OUTPUT_CHANNELS), InParams.SampleRate, EWasapiAudioEncoding::FLOATING_POINT_32);
+
+		if (MixFormat)
+		{
+			::CoTaskMemFree(MixFormat);
+			MixFormat = nullptr;
+		}
+
+		REFERENCE_TIME DevicePeriodRefTime = 0;
+		// The second param to GetDevicePeriod is only valid for exclusive mode
+		// Note that GetDevicePeriod returns ref time which is sample rate agnostic
+		Result = TempAudioClient->GetDevicePeriod(&DevicePeriodRefTime, nullptr);
+		if (FAILED(Result))
+		{
+			UE_LOG(LogAudioMixerWasapi, Error, TEXT("OpenAudioStream failed IAudioClient3::GetDevicePeriod 0x%x"), (uint32)Result);
+			return false;
+		}
+
+		DefaultDevicePeriod = FWasapiAudioUtils::RefTimeToFrames(DevicePeriodRefTime, InParams.SampleRate);
+
+		REFERENCE_TIME DesiredBufferDuration = FWasapiAudioUtils::FramesToRefTime(FMath::Max<uint32>(InParams.NumFrames, DefaultDevicePeriod), InParams.SampleRate);
+		// Add two times the device period to the buffer size to ensure buffer phasing isn't an issue with larger buffer sizes.
+		// For example, with an engine period of 1024 and a device period of 480, we'll have two device periods to fetch the next engine buffer before underruning. 
+		DesiredBufferDuration += (DevicePeriodRefTime * 2);
+
+		// For shared mode, this is required to be zero
+		const REFERENCE_TIME Periodicity = 0;
+
+		// Audio events will be delivered to us rather than needing to poll
+		uint32 Flags = AUDCLNT_STREAMFLAGS_EVENTCALLBACK;
+
+		if (InParams.SampleRate != InParams.HardwareDeviceInfo.SampleRate)
+		{
+			Flags |= AUDCLNT_STREAMFLAGS_AUTOCONVERTPCM;
+			UE_LOG(LogAudioMixerWasapi, Warning, TEXT("Sample rate mismatch. Engine sample rate: %d Device sample rate: %d"), InParams.SampleRate, InParams.HardwareDeviceInfo.SampleRate);
+			UE_LOG(LogAudioMixerWasapi, Warning, TEXT("Device level sample rate conversion will be used."));
+		}
+
+		Result = TempAudioClient->Initialize(AUDCLNT_SHAREMODE_SHARED, Flags, DesiredBufferDuration, Periodicity, StreamFormat.GetWaveFormat(), nullptr);
+		if (FAILED(Result))
+		{
+			UE_LOG(LogAudioMixerWasapi, Error, TEXT("OpenAudioStream failed IAudioClient3::Initialize 0x%x"), (uint32)Result);
+			return false;
+		}
+
+		Result = TempAudioClient->GetBufferSize(&NumFramesPerDeviceBuffer);
+		if (FAILED(Result))
+		{
+			UE_LOG(LogAudioMixerWasapi, Error, TEXT("OpenAudioStream failed IAudioClient3::GetBufferSize 0x%x"), (uint32)Result);
+			return false;
+		}
+
+		AudioClient = MoveTemp(TempAudioClient);
+		AudioFormat = StreamFormat;
+		RenderStreamParams = InParams;
+
+		bIsInitialized = true;
+
+		UE_LOG(LogAudioMixerWasapi, Verbose, TEXT("FAudioMixerWasapiRenderStream::InitializeHardware succeeded with sample rate: %d, buffer period: %d"), InParams.SampleRate, InParams.NumFrames);
+
+		return true;
+	}
+
+	bool FAudioMixerWasapiRenderStream::TeardownHardware()
+	{
+		if (!bIsInitialized)
+		{
+			UE_LOG(LogAudioMixerWasapi, Warning, TEXT("FAudioMixerWasapiRenderStream::TeardownHardware failed...not initialized. "));
+			return false;
+		}
+
+		ReadNextBufferDelegate.Unbind();
+		RenderClient.Reset();
+		AudioClient.Reset();
+
+		bIsInitialized = false;
+
+		UE_LOG(LogAudioMixerWasapi, Verbose, TEXT("FAudioMixerWasapiRenderStream::TeardownHardware succeeded"));
+
+		return true;
+	}
+
+	bool FAudioMixerWasapiRenderStream::IsInitialized() const
+	{
+		return bIsInitialized;
+	}
+
+	int32 FAudioMixerWasapiRenderStream::GetNumFrames(const int32 InNumReqestedFrames) const
+	{
+		return FMath::Max<uint32>(InNumReqestedFrames, DefaultDevicePeriod);
+	}
+
+	bool FAudioMixerWasapiRenderStream::OpenAudioStream(const FWasapiRenderStreamParams& InParams, HANDLE InEventHandle)
+	{
+		if (InParams.HardwareDeviceInfo.DeviceId != RenderStreamParams.HardwareDeviceInfo.DeviceId)
+		{
+			if (!InitializeHardware(InParams))
+			{
+				UE_LOG(LogAudioMixerWasapi, Error, TEXT("OpenAudioStream failed InitAudioClient"));
+				return false;
+			}
+		}
+
+		if (InEventHandle == nullptr)
+		{
+			UE_LOG(LogAudioMixerWasapi, Error, TEXT("OpenAudioStream null EventHandle"));
+			return false;
+		}
+
+		HRESULT Result = AudioClient->SetEventHandle(InEventHandle);
+		if (FAILED(Result))
+		{
+			UE_LOG(LogAudioMixerWasapi, Error, TEXT("OpenAudioStream failed IAudioClient3::SetEventHandle 0x%x"), (uint32)Result);
+			return false;
+		}
+
+		TComPtr<IAudioRenderClient> TempRenderClient;
+		Result = AudioClient->GetService(__uuidof(IAudioRenderClient), IID_PPV_ARGS_Helper(&TempRenderClient));
+		if (FAILED(Result))
+		{
+			UE_LOG(LogAudioMixerWasapi, Error, TEXT("OpenAudioStream failed IAudioClient3::GetService IAudioRenderClient 0x%x"), (uint32)Result);
+			return false;
+		}
+
+		RenderClient = MoveTemp(TempRenderClient);
+		bIsInitialized = true;
+
+		UE_LOG(LogAudioMixerWasapi, Verbose, TEXT("FAudioMixerWasapiRenderStream::OpenAudioStream succeeded with SampeRate: %d, NumFrames: %d"), InParams.SampleRate, InParams.NumFrames);
+
+		return true;
+	}
+
+	bool FAudioMixerWasapiRenderStream::CloseAudioStream()
+	{
+		if (!bIsInitialized || StreamState == EAudioOutputStreamState::Closed)
+		{
+			UE_LOG(LogAudioMixerWasapi, Verbose, TEXT("FAudioMixerWasapiRenderStream::CloseAudioStream stream appears to be already closed"));
+
+			return false;
+		}
+
+		StreamState = EAudioOutputStreamState::Closed;
+
+		return true;
+	}
+
+	bool FAudioMixerWasapiRenderStream::StartAudioStream()
+	{
+		if (bIsInitialized)
+		{
+			StreamState = EAudioOutputStreamState::Running;
+
+			if (!AudioClient.IsValid())
+			{
+				UE_LOG(LogAudioMixerWasapi, Error, TEXT("StartAudioStream failed invalid audio client"));
+				return false;
+			}
+
+			if (!ReadNextBufferDelegate.IsBound())
+			{
+				UE_LOG(LogAudioMixerWasapi, Error, TEXT("StartAudioStream failed buffer delegate not bound"));
+				return false;
+			}
+
+			AudioClient->Start();
+		}
+
+		UE_LOG(LogAudioMixerWasapi, Verbose, TEXT("FAudioMixerWasapiRenderStream::StartAudioStream stream started"));
+
+		return true;
+	}
+
+	bool FAudioMixerWasapiRenderStream::StopAudioStream()
+	{
+		if (!bIsInitialized)
+		{
+			UE_LOG(LogAudioMixerWasapi, Error, TEXT("FAudioMixerWasapiRenderStream::StopAudioStream() not initialized"));
+			return false;
+		}
+
+		if (StreamState != EAudioOutputStreamState::Stopped && StreamState != EAudioOutputStreamState::Closed)
+		{
+			if (AudioClient.IsValid())
+			{
+				AudioClient->Stop();
+			}
+
+			StreamState = EAudioOutputStreamState::Stopped;
+		}
+
+		if (CallbackBufferErrors > 0)
+		{
+			UE_LOG(LogAudioMixerWasapi, Error, TEXT("FAudioMixerWasapiRenderStream::StopAudioStream render stream reported %d callback buffer errors"), CallbackBufferErrors);
+			CallbackBufferErrors = 0;
+		}
+
+		return true;
+	}
+
+	void FAudioMixerWasapiRenderStream::DeviceRenderCallback()
+	{
+		SCOPED_NAMED_EVENT(FAudioMixerWasapiRenderStream_DeviceRenderCallback, FColor::Blue);
+
+		if (bIsInitialized)
+		{
+			uint32 NumFramesPadding = 0;
+			AudioClient->GetCurrentPadding(&NumFramesPadding);
+
+			// NumFramesPerDeviceBuffer is the buffer size WASAPI allocated. It is gauranteed to 
+			// be at least the amount requested. For example, if we request a 1024 frame buffer, WASAPI
+			// might allocate a 1056 frame buffer. The padding is subtracted from the allocated amount
+			// to determine how much space is available currently in the buffer.
+			const int32 NumFramesAvailable = NumFramesPerDeviceBuffer - NumFramesPadding;
+			if (NumFramesAvailable >= static_cast<int32>(RenderStreamParams.NumFrames))
+			{
+				check(!RenderBufferPtr);
+
+				if (SUCCEEDED(RenderClient->GetBuffer(RenderStreamParams.NumFrames, &RenderBufferPtr)))
+				{
+					ReadNextBufferDelegate.Execute();
+
+					HRESULT Result = RenderClient->ReleaseBuffer(RenderStreamParams.NumFrames, 0 /* flags */);
+					if (FAILED(Result))
+					{
+						++CallbackBufferErrors;
+					}
+
+					RenderBufferPtr = nullptr;
+				}
+				else
+				{
+					++CallbackBufferErrors;
+				}
+			}
+		}
+	}
+
+	void FAudioMixerWasapiRenderStream::SubmitBuffer(const uint8* InBuffer, const SIZE_T InNumFrames)
+	{
+		if (RenderBufferPtr)
+		{
+			FMemory::Memcpy(RenderBufferPtr, InBuffer, InNumFrames * AudioFormat.GetFrameSizeInBytes());
+		}
+	}
+}
diff --git a/Engine/Source/Runtime/Windows/AudioMixerWasapi/Private/AudioMixerWasapiRenderStream.h b/Engine/Source/Runtime/Windows/AudioMixerWasapi/Private/AudioMixerWasapiRenderStream.h
new file mode 100644
--- /dev/null
+++ b/Engine/Source/Runtime/Windows/AudioMixerWasapi/Private/AudioMixerWasapiRenderStream.h
@@ -0,0 +1,119 @@
+// Copyright Epic Games, Inc. All Rights Reserved.
+
+#pragma once
+
+#include "AudioMixer.h"	
+#include "Misc/ScopeRWLock.h"
+#include "WasapiAudioFormat.h"
+
+#include "Windows/AllowWindowsPlatformTypes.h"
+#include "Windows/AllowWindowsPlatformAtomics.h"
+
+#include "Microsoft/COMPointer.h"
+
+THIRD_PARTY_INCLUDES_START
+#include <mmdeviceapi.h>			// IMMNotificationClient
+#include <audiopolicy.h>			// IAudioSessionEvents
+THIRD_PARTY_INCLUDES_END
+
+#include "Windows/HideWindowsPlatformAtomics.h"
+#include "Windows/HideWindowsPlatformTypes.h"
+
+DECLARE_DELEGATE(FAudioMixerReadNextBufferDelegate);
+
+/**
+ *  Warning: AudioMixerWasapi module is currently considered experimental and may change in the future.
+ *           We do not recommend shipping projects with Experimental features.
+ */
+
+namespace Audio
+{
+	/** Defines parameters needed for opening a new audio stream to device. */
+	struct FWasapiRenderStreamParams
+	{
+		/** The audio device to open. */
+		TComPtr<IMMDevice> MMDevice;
+
+		/** Hardware device configuration info. */
+		FAudioPlatformDeviceInfo HardwareDeviceInfo;
+
+		/** The number of desired audio frames in audio callback. */
+		uint32 NumFrames = 0;
+
+		/** The desired sample rate */
+		uint32 SampleRate = 0;
+
+		FWasapiRenderStreamParams() = default;
+
+		FWasapiRenderStreamParams(const TComPtr<IMMDevice>& InMMDevice, 
+			const FAudioPlatformDeviceInfo& InDeviceInfo, 
+			const uint32 InNumFrames, 
+			const uint32 InSampleRate) :
+			MMDevice(InMMDevice),
+			HardwareDeviceInfo(InDeviceInfo),
+			NumFrames(InNumFrames),
+			SampleRate(InSampleRate)
+		{
+		}
+	};
+
+	/**
+	 * FAudioMixerWasapiRenderStream
+	 */
+	class FAudioMixerWasapiRenderStream
+	{
+	public:
+
+		FAudioMixerWasapiRenderStream();
+		virtual ~FAudioMixerWasapiRenderStream();
+
+		bool InitializeHardware(const FWasapiRenderStreamParams& InParams);
+		bool TeardownHardware();
+		bool IsInitialized() const;
+		int32 GetNumFrames(const int32 InNumReqestedFrames) const;
+		bool OpenAudioStream(const FWasapiRenderStreamParams& InParams, HANDLE InEventHandle);
+		bool CloseAudioStream();
+		bool StartAudioStream();
+		bool StopAudioStream();
+		void SubmitBuffer(const uint8* Buffer, const SIZE_T InNumFrames);
+		void DeviceRenderCallback();
+
+		FAudioMixerReadNextBufferDelegate& OnReadNextBuffer() { return ReadNextBufferDelegate; }
+
+	private:
+
+		/** Delegate called each buffer callback to signal the mixer to process the next buffer. */
+		FAudioMixerReadNextBufferDelegate ReadNextBufferDelegate;
+
+		/** COM pointer to the WASAPI audio client object. */
+		TComPtr<IAudioClient3> AudioClient;
+
+		/** COM pointer to the WASAPI render client object. */
+		TComPtr<IAudioRenderClient> RenderClient;
+
+		/** Holds the audio format configuration for this stream. */
+		FWasapiAudioFormat AudioFormat;
+
+		/** Indicates if this object has been successfully initialized. */
+		bool bIsInitialized = false;
+
+		/** The state of the output audio stream. */
+		EAudioOutputStreamState::Type StreamState = EAudioOutputStreamState::Closed;
+
+		/** Render output device info. */
+		FWasapiRenderStreamParams RenderStreamParams;
+
+		/** The default callback period for this WASAPI render device. */
+		uint32 DefaultDevicePeriod = 0;
+
+		/** Number of frames of audio data which will be used for each audio callback. This value is 
+		    determined by the WASAPI audio client and can be equal or greater than the number of frames requested. */
+		uint32 NumFramesPerDeviceBuffer = 0;
+
+		/** Accumulates errors that occur in the audio callback. */
+		uint32 CallbackBufferErrors = 0;
+
+		/** Pointer to WASAPI render audio buffer filled in each callback by the mixer in SubmitBuffer(). */
+		uint8* RenderBufferPtr = nullptr;
+	};
+}
