Reaname listener to avoid name collision error.

#rnx

diff --git a/Engine/Source/Runtime/Windows/AudioMixerWasapi/Private/AudioMixerWasapiDeviceListener.cpp b/Engine/Source/Runtime/Windows/AudioMixerWasapi/Private/AudioMixerWasapiDeviceListener.cpp
--- a/Engine/Source/Runtime/Windows/AudioMixerWasapi/Private/AudioMixerWasapiDeviceListener.cpp
+++ b/Engine/Source/Runtime/Windows/AudioMixerWasapi/Private/AudioMixerWasapiDeviceListener.cpp
@@ -25,14 +25,14 @@ THIRD_PARTY_INCLUDES_END
 
 namespace Audio
 {	
-	TSharedPtr<FWindowsMMNotificationClient> WindowsNotificationClient;
+	TSharedPtr<FWindowsMMNotificationClient> WasapiWinNotificationClient;
 
 	void FAudioMixerWasapi::RegisterDeviceChangedListener()
 	{
-		if (!WindowsNotificationClient.IsValid())
+		if (!WasapiWinNotificationClient.IsValid())
 		{
 			// Shared (This is a COM object, so we don't delete it, just derecement the ref counter).
-			WindowsNotificationClient = TSharedPtr<FWindowsMMNotificationClient>(
+			WasapiWinNotificationClient = TSharedPtr<FWindowsMMNotificationClient>(
 				new FWindowsMMNotificationClient, 
 				[](FWindowsMMNotificationClient* InPtr) { InPtr->ReleaseClient(); }
 			);
@@ -41,25 +41,25 @@ namespace Audio
 		{
 			// Setup device info cache.
 			DeviceInfoCache = MakeUnique<FWindowsMMDeviceCache>();
-			WindowsNotificationClient->RegisterDeviceChangedListener(static_cast<FWindowsMMDeviceCache*>(DeviceInfoCache.Get()));
+			WasapiWinNotificationClient->RegisterDeviceChangedListener(static_cast<FWindowsMMDeviceCache*>(DeviceInfoCache.Get()));
 		}
 
-		WindowsNotificationClient->RegisterDeviceChangedListener(this);
+		WasapiWinNotificationClient->RegisterDeviceChangedListener(this);
 	}
 
 	void FAudioMixerWasapi::UnregisterDeviceChangedListener()
 	{
-		if (WindowsNotificationClient.IsValid())
+		if (WasapiWinNotificationClient.IsValid())
 		{
 			if (DeviceInfoCache.IsValid())
 			{
 				// Unregister and kill cache.
-				WindowsNotificationClient->UnRegisterDeviceDeviceChangedListener(static_cast<FWindowsMMDeviceCache*>(DeviceInfoCache.Get()));
+				WasapiWinNotificationClient->UnRegisterDeviceDeviceChangedListener(static_cast<FWindowsMMDeviceCache*>(DeviceInfoCache.Get()));
 				
 				DeviceInfoCache.Reset();
 			}
 			
-			WindowsNotificationClient->UnRegisterDeviceDeviceChangedListener(this);
+			WasapiWinNotificationClient->UnRegisterDeviceDeviceChangedListener(this);
 		}
 	}
 
@@ -70,9 +70,9 @@ namespace Audio
 
 	TComPtr<IMMDevice> FAudioMixerWasapi::GetMMDevice(const FString& InDeviceID) const
 	{
-		if (WindowsNotificationClient)
+		if (WasapiWinNotificationClient)
 		{
-			return WindowsNotificationClient->GetDevice(InDeviceID);
+			return WasapiWinNotificationClient->GetDevice(InDeviceID);
 		}
 
 		return TComPtr<IMMDevice>();
