Fix live race condition with XAudio2::Suspend/Resume + DeviceSwap

diff --git a/Engine/Source/Runtime/Windows/AudioMixerXAudio2/Private/AudioMixerPlatformXAudio2.cpp b/Engine/Source/Runtime/Windows/AudioMixerXAudio2/Private/AudioMixerPlatformXAudio2.cpp
--- a/Engine/Source/Runtime/Windows/AudioMixerXAudio2/Private/AudioMixerPlatformXAudio2.cpp
+++ b/Engine/Source/Runtime/Windows/AudioMixerXAudio2/Private/AudioMixerPlatformXAudio2.cpp
@@ -548,7 +548,14 @@ namespace Audio
 
 	void FMixerPlatformXAudio2::Suspend()
 	{
+		if (!IsInAudioThread())
+		{
+			FAudioThread::RunCommandOnAudioThread([this](){ Suspend(); });
+			return;
+		}
+		
 		SCOPED_ENTER_BACKGROUND_EVENT(STAT_FMixerPlatformXAudio2_Suspend);
+		
 		if( !bIsSuspended )
 		{				
 			if( XAudio2System )
@@ -564,6 +571,12 @@ namespace Audio
 	}
 	void FMixerPlatformXAudio2::Resume()
 	{
+		if (!IsInAudioThread())
+		{
+			FAudioThread::RunCommandOnAudioThread([this](){ Resume(); });
+			return;
+		}
+
 		if( bIsSuspended )
 		{
 			if( XAudio2System )
