Fix differences in FRHIGPUFence::Poll() behaviour across platform RHIs
 - Poll() should return false for fences that have never been enqueued.

diff --git a/Engine/Source/Runtime/OpenGLDrv/Private/OpenGLQuery.cpp b/Engine/Source/Runtime/OpenGLDrv/Private/OpenGLQuery.cpp
--- a/Engine/Source/Runtime/OpenGLDrv/Private/OpenGLQuery.cpp
+++ b/Engine/Source/Runtime/OpenGLDrv/Private/OpenGLQuery.cpp
@@ -746,7 +746,7 @@ void FOpenGLGPUFence::Clear()
 
 bool FOpenGLGPUFence::Poll() const
 {
-	return !Event || Event->IsComplete();
+	return Event && Event->IsComplete();
 }
 
 void FOpenGLGPUFence::Wait(FRHICommandListImmediate& RHICmdList, FRHIGPUMask GPUMask) const
diff --git a/Engine/Source/Runtime/Windows/D3D11RHI/Private/D3D11Commands.cpp b/Engine/Source/Runtime/Windows/D3D11RHI/Private/D3D11Commands.cpp
--- a/Engine/Source/Runtime/Windows/D3D11RHI/Private/D3D11Commands.cpp
+++ b/Engine/Source/Runtime/Windows/D3D11RHI/Private/D3D11Commands.cpp
@@ -1696,7 +1696,7 @@ void FD3D11GPUFence::Clear()
 
 bool FD3D11GPUFence::Poll() const
 {
-	return !Event || Event->IsComplete();
+	return Event && Event->IsComplete();
 }
 
 void FD3D11GPUFence::Wait(FRHICommandListImmediate& RHICmdList, FRHIGPUMask GPUMask) const
