Remove APU allocations when not using XMA2
#rnx

diff --git a/Engine/Source/Runtime/Windows/AudioMixerXAudio2/Private/AudioMixerPlatformXAudio2.cpp b/Engine/Source/Runtime/Windows/AudioMixerXAudio2/Private/AudioMixerPlatformXAudio2.cpp
--- a/Engine/Source/Runtime/Windows/AudioMixerXAudio2/Private/AudioMixerPlatformXAudio2.cpp
+++ b/Engine/Source/Runtime/Windows/AudioMixerXAudio2/Private/AudioMixerPlatformXAudio2.cpp
@@ -14,14 +14,6 @@
 #include "HAL/PlatformProcess.h"
 #include "Engine/EngineTypes.h"
 
-#ifndef WITH_XMA2
-#define WITH_XMA2 0
-#endif
-#if WITH_ENGINE
-#if WITH_XMA2
-#include "XMAAudioInfo.h"
-#endif  //#if WITH_XMA2
-#endif //WITH_ENGINE
 #include "CoreGlobals.h"
 #include "Misc/ConfigCacheIni.h"
 #include "Misc/MessageDialog.h"
@@ -156,7 +148,6 @@ namespace Audio
 		, TimeSinceNullDeviceWasLastChecked(0.0f)		
 		, bIsInitialized(false)
 		, bIsDeviceOpen(false)
-		, bIsSuspended(false)
 	{
 #if PLATFORM_WINDOWS 
 		FPlatformMisc::CoInitialize();
@@ -521,14 +512,9 @@ namespace Audio
 	{
 		SAFE_RELEASE(XAudio2System);
 
-		uint32 Flags = 0;
 		HRESULT Result;
 
-#if WITH_XMA2
-		// We need to raise this flag explicitly to prevent initializing SHAPE twice, because we are allocating SHAPE in FXMAAudioInfo
-		Flags |= XAUDIO2_DO_NOT_USE_SHAPE;
-#endif
-
+		const uint32 Flags = GetCreateFlags();
 		Result = XAudio2Create(&XAudio2System, Flags, GetXAudio2ProcessorsToUse());
 		if (FAILED(Result))
 		{
@@ -546,101 +532,6 @@ namespace Audio
 		return true;
 	}
 
-	void FMixerPlatformXAudio2::Suspend()
-	{
-		if (!IsInAudioThread())
-		{
-			FAudioThread::RunCommandOnAudioThread([this](){ Suspend(); });
-			return;
-		}
-		
-		SCOPED_ENTER_BACKGROUND_EVENT(STAT_FMixerPlatformXAudio2_Suspend);
-		
-		if( !bIsSuspended )
-		{				
-			if( XAudio2System )
-			{
-				const uint64 StartCycles = FPlatformTime::Cycles64();						
-				XAudio2System->StopEngine();
-				UE_LOG(LogAudioMixer, Display, TEXT("XAudio2System StopEngine(), taking '%2.2f' ms"),
-					FPlatformTime::ToMilliseconds64(FPlatformTime::Cycles64() - StartCycles));
-				StartRunningNullDevice();
-				bIsSuspended = true;
-			}					
-		}
-	}
-	void FMixerPlatformXAudio2::Resume()
-	{
-		if (!IsInAudioThread())
-		{
-			FAudioThread::RunCommandOnAudioThread([this](){ Resume(); });
-			return;
-		}
-
-		if( bIsSuspended )
-		{
-			if( XAudio2System )
-			{			
-				StopRunningNullDevice();			
-
-				HRESULT Result = XAudio2System->StartEngine();				
-
-				// Log any failed state as "Display" so we always appear in the logs.
-				UE_CLOG(FAILED(Result), LogAudioMixer, Display, TEXT("XAudio2->StartEngine() returned (%s) 0x%X"), *ToErrorFString(Result), (uint32)Result);
-
-				// If we see this error here, the audio device has been invalidated, and we must swap to a new one.
-				if (Result == XAUDIO2_E_DEVICE_INVALIDATED)
-				{
-					bIsSuspended = false;
-					RequestDeviceSwap(TEXT(""), /*force*/ true, *FString::Printf(TEXT("XAudio2->StartEngine() returned (%s) 0x%X"), *ToErrorFString(Result), (uint32)Result));
-					return;
-				}
-
-				// Suggestion from Microsoft for this returning 0x80070490 on a quick cycle. 
-				// Try again with delay.
-				if (Result == 0x80070490)
-				{
-					constexpr int32 NumStartAttempts = 16;
-				
-					uint64 StartCycles = FPlatformTime::Cycles64();				
-					int32 StartAttempt = 0;
-					for (; StartAttempt < NumStartAttempts && FAILED(Result); ++StartAttempt)
-					{
-						constexpr float DelayBetweenAttemptsSecs = 0.1f;
-						FPlatformProcess::Sleep(DelayBetweenAttemptsSecs);
-						Result = XAudio2System->StartEngine();	
-					}
-
-					if (SUCCEEDED(Result))
-					{		
-						UE_CLOG(StartAttempt == 1, LogAudioMixer, Display, 
-							TEXT("XAudio2System StartEngine() - Sucessfully started, taking '%2.2f' ms"), FPlatformTime::ToMilliseconds64(FPlatformTime::Cycles64() - StartCycles));
-						
-						UE_CLOG(StartAttempt > 1, LogAudioMixer, Warning, 
-							TEXT("StartEngine() took %d attempts to start, taking '%f' ms"), 
-							StartAttempt, FPlatformTime::ToMilliseconds64(FPlatformTime::Cycles64() - StartCycles));
-					}
-					else
-					{
-						UE_CLOG(FAILED(Result), LogAudioMixer, Error,
-								TEXT("Could not resume XAudio2, StartEngine() returned 0x%X, after %d attempts"),
-								(uint32)Result, StartAttempt);
-					}
-				}
-
-				// Always turn off the suspended state even in a failure.
-				bIsSuspended = false;
-
-				// If we're still in failed state, do our best here 
-				if (FAILED(Result))
-				{
-					UE_LOG(LogAudioMixer, Error, TEXT("XAudio2->StartEngine() returned (%s) 0x%X"), *ToErrorFString(Result), (uint32)Result);
-					StartRunningNullDevice();
-				}
-			}						
-		}		
-	}
-
 	bool FMixerPlatformXAudio2::InitializeHardware()
 	{
 		if (bIsInitialized)
@@ -649,16 +540,6 @@ namespace Audio
 			return false;
 		}
 
-#if PLATFORM_NEEDS_SUSPEND_ON_BACKGROUND
-
-		// Constrain/Suspend 
-		DeactiveHandle = FCoreDelegates::ApplicationWillDeactivateDelegate.AddRaw(this, &FMixerPlatformXAudio2::Suspend);
-		ReactivateHandle = FCoreDelegates::ApplicationHasReactivatedDelegate.AddRaw(this, &FMixerPlatformXAudio2::Resume);
-		EnteredBackgroundHandle = FCoreDelegates::ApplicationWillEnterBackgroundDelegate.AddRaw(this, &FMixerPlatformXAudio2::Suspend);
-		EnteredForegroundHandle = FCoreDelegates::ApplicationHasEnteredForegroundDelegate.AddRaw(this, &FMixerPlatformXAudio2::Resume);
-
-#endif //PLATFORM_NEEDS_SUSPEND_ON_BACKGROUND
-		
 #if USE_REDIST_LIB
 		// Work around the fact the x64 version of XAudio2_7.dll does not properly ref count
 		// by forcing it to be always loaded
@@ -679,13 +560,7 @@ namespace Audio
 		}
 #endif // #if PLATFORM_WINDOWS
 
-		uint32 Flags = 0;
-
-#if WITH_XMA2
-		// We need to raise this flag explicitly to prevent initializing SHAPE twice, because we are allocating SHAPE in FXMAAudioInfo
-		Flags |= XAUDIO2_DO_NOT_USE_SHAPE;
-#endif
-
+		const uint32 Flags = GetCreateFlags();
 		if (!XAudio2System && FAILED(XAudio2Create(&XAudio2System, Flags, GetXAudio2ProcessorsToUse())))
 		{
 			FMessageDialog::Open(EAppMsgType::Ok, NSLOCTEXT("Audio", "XAudio2Error", "Failed to initialize audio. This may be an issue with your installation of XAudio 2.7. XAudio2 is available in the DirectX End-User Runtime (June 2010)."));
@@ -702,12 +577,7 @@ namespace Audio
 		{
 			UE_LOG(LogAudioMixer, Error, TEXT("Failed to register for callbacks."));
 		}
-
-#if WITH_XMA2
-		//Initialize our XMA2 decoder context
-		FXMAAudioInfo::Initialize();
-#endif //#if WITH_XMA2
-
+		
 		if(IAudioMixer::ShouldRecycleThreads())
 		{
 			// Pre-create the null render device thread on XAudio2, so we can simple wake it up when we need it.
@@ -728,34 +598,12 @@ namespace Audio
 			return false;
 		}
 
-#if PLATFORM_NEEDS_SUSPEND_ON_BACKGROUND
-
-		auto UnregisterLambda = [](FDelegateHandle& InHandle, TMulticastDelegate<void()>& InDelegate)
-		{
-			if (InHandle.IsValid())
-			{
-				InDelegate.Remove(InHandle);
-				InHandle.Reset();
-			}
-		};
-	 		
-		UnregisterLambda(DeactiveHandle,FCoreDelegates::ApplicationWillDeactivateDelegate);
-		UnregisterLambda(ReactivateHandle,FCoreDelegates::ApplicationHasReactivatedDelegate);
-		UnregisterLambda(EnteredBackgroundHandle,FCoreDelegates::ApplicationWillEnterBackgroundDelegate);
-		UnregisterLambda(EnteredForegroundHandle,FCoreDelegates::ApplicationHasEnteredForegroundDelegate);
-
-#endif //PLATFORM_NEEDS_SUSPEND_ON_BACKGROUND
-
 		if (XAudio2System)
 		{
 			XAudio2System->UnregisterForCallbacks(this);
 		}
 		SAFE_RELEASE(XAudio2System);
 
-#if WITH_XMA2
-		FXMAAudioInfo::Shutdown();
-#endif
-
 #if PLATFORM_WINDOWS
 		if (XAudio2Dll != nullptr && IsEngineExitRequested())
 		{
@@ -1649,10 +1497,6 @@ namespace Audio
 
 	void FMixerPlatformXAudio2::OnHardwareUpdate()
 	{
-#if WITH_XMA2
-		FXMAAudioInfo::Tick();
-#endif //WITH_XMA2
-
 	}
 
 	Audio::IAudioPlatformDeviceInfoCache* FMixerPlatformXAudio2::GetDeviceInfoCache() const
@@ -1701,11 +1545,7 @@ namespace Audio
 
 	void FMixerPlatformXAudio2::OnCriticalError(HRESULT Error)
 	{
-		// Windows handles this via session events, so no need to handle that here..
-#if !PLATFORM_WINDOWS
-		UE_LOG(LogAudioMixer, Display, TEXT("FMixerPlatformXAudio2::OnCriticalError() returned (%s) 0x%X"), *ToErrorFString(Error), (uint32)Error);
-		RequestDeviceSwap(TEXT(""), /*force*/ true, *FString::Printf(TEXT("FMixerPlatformXAudio2::OnCriticalError() returned (%s) 0x%X"), *ToErrorFString(Error), (uint32)Error));
-#endif //!PLATFORM_WINDOWS
+		// Windows handles this via session events, so no need to handle that here...
 	}
 	
 }
diff --git a/Engine/Source/Runtime/Windows/AudioMixerXAudio2/Private/AudioMixerPlatformXAudio2.h b/Engine/Source/Runtime/Windows/AudioMixerXAudio2/Private/AudioMixerPlatformXAudio2.h
--- a/Engine/Source/Runtime/Windows/AudioMixerXAudio2/Private/AudioMixerPlatformXAudio2.h
+++ b/Engine/Source/Runtime/Windows/AudioMixerXAudio2/Private/AudioMixerPlatformXAudio2.h
@@ -38,10 +38,6 @@
     #define XAUDIO_SUPPORTS_DEVICE_DETAILS		1
 #endif	//XAUDIO_SUPPORTS_DEVICE_DETAILS
 
-#ifndef PLATFORM_NEEDS_SUSPEND_ON_BACKGROUND
-	#define PLATFORM_NEEDS_SUSPEND_ON_BACKGROUND	0
-#endif //PLATFORM_NEEDS_SUSPEND_ON_BACKGROUND
-
 // Any platform defines
 namespace Audio
 {
@@ -118,7 +114,9 @@ namespace Audio
 		virtual FString GetDeviceId() const override;		
 		//~ End IAudioMixerDeviceChangedListener
 
-	private:
+	protected:
+		virtual uint32 GetCreateFlags() const { return 0; }
+
 		//~ Begin IXAudio2EngineCallback
 		virtual void OnCriticalError(HRESULT Error) override;
 		virtual void OnProcessingPassStart() override;
@@ -143,16 +141,10 @@ namespace Audio
 		// This must be done to repopulate the playback device list in XAudio 2.7.
 		bool ResetXAudio2System();
 
-		void Suspend();
-		void Resume();
 
 		// Handle to XAudio2DLL
 		HMODULE XAudio2Dll;
 
-		FDelegateHandle DeactiveHandle;
-		FDelegateHandle ReactivateHandle;
-		FDelegateHandle EnteredBackgroundHandle;
-		FDelegateHandle EnteredForegroundHandle;
 
 		// Bool indicating that the default audio device changed
 		// And that we need to restart the audio device.
@@ -178,7 +170,6 @@ namespace Audio
 
 		uint32 bIsInitialized : 1;
 		uint32 bIsDeviceOpen : 1;
-		uint32 bIsSuspended : 1;
 	};
 }
 
diff --git a/Engine/Source/Runtime/Windows/AudioMixerXAudio2/Private/AudioMixerXAudio2.cpp b/Engine/Source/Runtime/Windows/AudioMixerXAudio2/Private/AudioMixerXAudio2.cpp
--- a/Engine/Source/Runtime/Windows/AudioMixerXAudio2/Private/AudioMixerXAudio2.cpp
+++ b/Engine/Source/Runtime/Windows/AudioMixerXAudio2/Private/AudioMixerXAudio2.cpp
@@ -1,29 +1,28 @@
 // Copyright Epic Games, Inc. All Rights Reserved.
 
+#include "AudioMixerXAudio2.h"
 #include "AudioMixerPlatformXAudio2.h"
 #include "AudioMixer.h"
 #include "Modules/ModuleManager.h"
 
-class FAudioMixerModuleXAudio2 : public IAudioDeviceModule
+void FAudioMixerModuleXAudio2::StartupModule() 
 {
-public:
-	virtual void StartupModule() override
-	{
-		IAudioDeviceModule::StartupModule();
+	IAudioDeviceModule::StartupModule();
 
-		FModuleManager::Get().LoadModuleChecked(TEXT("AudioMixer"));
-		FModuleManager::Get().LoadModuleChecked(TEXT("AudioMixerCore"));
-	}
+	FModuleManager::Get().LoadModuleChecked(TEXT("AudioMixer"));
+	FModuleManager::Get().LoadModuleChecked(TEXT("AudioMixerCore"));
+}
 
-	virtual bool IsAudioMixerModule() const override 
-	{ 
-		return true;
-	}
+bool FAudioMixerModuleXAudio2::IsAudioMixerModule() const
+{
+	return true;
+}
 
-	virtual Audio::IAudioMixerPlatformInterface* CreateAudioMixerPlatformInterface() override
-	{
-		return new Audio::FMixerPlatformXAudio2();
-	}
-};
+Audio::IAudioMixerPlatformInterface* FAudioMixerModuleXAudio2::CreateAudioMixerPlatformInterface() 
+{
+	return new Audio::FMixerPlatformXAudio2();
+}
 
+#if PLATFORM_WINDOWS
 IMPLEMENT_MODULE(FAudioMixerModuleXAudio2, AudioMixerXAudio2);
+#endif //PLATFORM_WINDOWS
diff --git a/Engine/Source/Runtime/Windows/AudioMixerXAudio2/Private/AudioMixerXAudio2.h b/Engine/Source/Runtime/Windows/AudioMixerXAudio2/Private/AudioMixerXAudio2.h
new file mode 100644
--- /dev/null
+++ b/Engine/Source/Runtime/Windows/AudioMixerXAudio2/Private/AudioMixerXAudio2.h
@@ -0,0 +1,14 @@
+﻿// Copyright Epic Games, Inc. All Rights Reserved.
+
+#pragma once
+
+#include "AudioMixer.h"
+#include "Modules/ModuleManager.h"
+
+class FAudioMixerModuleXAudio2 : public IAudioDeviceModule
+{
+public:
+	virtual void StartupModule() override;
+	virtual bool IsAudioMixerModule() const override;
+	virtual Audio::IAudioMixerPlatformInterface* CreateAudioMixerPlatformInterface() override;
+};
