Check for not null Swapchain in DX11 Present

diff --git a/Engine/Source/Runtime/Windows/D3D11RHI/Private/D3D11Viewport.cpp b/Engine/Source/Runtime/Windows/D3D11RHI/Private/D3D11Viewport.cpp
--- a/Engine/Source/Runtime/Windows/D3D11RHI/Private/D3D11Viewport.cpp
+++ b/Engine/Source/Runtime/Windows/D3D11RHI/Private/D3D11Viewport.cpp
@@ -447,23 +447,26 @@ bool FD3D11Viewport::PresentChecked(IRHICommandContext& RHICmdContext, int32 Syn
 		UE_LOG(LogRHI, Error, TEXT("Present Fail Count %i"), PresentFailCount);
 		DXGI_SWAP_CHAIN_DESC Desc;
 		UE_LOG(LogRHI, Error, TEXT("SyncInterval %i"), SyncInterval);
-		if (!FAILED(SwapChain->GetDesc(&Desc)))
+		if (SwapChain.IsValid())
 		{
-			UE_LOG(LogRHI, Error, TEXT("SwapChainDesc.BufferDesc.Width %i"), Desc.BufferDesc.Width);
-			UE_LOG(LogRHI, Error, TEXT("SwapChainDesc.BufferDesc.Height %i"), Desc.BufferDesc.Height);
-			UE_LOG(LogRHI, Error, TEXT("SwapChainDesc.BufferDesc.RefreshRate.Numerator %i"), Desc.BufferDesc.RefreshRate.Numerator);
-			UE_LOG(LogRHI, Error, TEXT("SwapChainDesc.BufferDesc.RefreshRate.Denominator %i"), Desc.BufferDesc.RefreshRate.Denominator);
-			UE_LOG(LogRHI, Error, TEXT("SwapChainDesc.BufferDesc.Format %i"), Desc.BufferDesc.Format);
-			UE_LOG(LogRHI, Error, TEXT("SwapChainDesc.BufferDesc.ScanlineOrdering %i"), Desc.BufferDesc.ScanlineOrdering);
-			UE_LOG(LogRHI, Error, TEXT("SwapChainDesc.BufferDesc.Scaling %i"), Desc.BufferDesc.Scaling);
-			UE_LOG(LogRHI, Error, TEXT("SwapChainDesc.SampleDesc.Count %i"), Desc.SampleDesc.Count);
-			UE_LOG(LogRHI, Error, TEXT("SwapChainDesc.SampleDesc.Quality %i"), Desc.SampleDesc.Quality);
-			UE_LOG(LogRHI, Error, TEXT("SwapChainDesc.BufferUsage %i"), Desc.BufferUsage);
-			UE_LOG(LogRHI, Error, TEXT("SwapChainDesc.BufferCount %i"), Desc.BufferCount);
-			UE_LOG(LogRHI, Error, TEXT("SwapChainDesc.OutputWindow %p"), Desc.OutputWindow);
-			UE_LOG(LogRHI, Error, TEXT("SwapChainDesc.Windowed %s"), Desc.Windowed ? TEXT("true") : TEXT("false"));
-			UE_LOG(LogRHI, Error, TEXT("SwapChainDesc.SwapEffect %u"), Desc.SwapEffect);
-			UE_LOG(LogRHI, Error, TEXT("SwapChainDesc.Flags %u"), Desc.Flags);
+			if (!FAILED(SwapChain->GetDesc(&Desc)))
+			{
+				UE_LOG(LogRHI, Error, TEXT("SwapChainDesc.BufferDesc.Width %i"), Desc.BufferDesc.Width);
+				UE_LOG(LogRHI, Error, TEXT("SwapChainDesc.BufferDesc.Height %i"), Desc.BufferDesc.Height);
+				UE_LOG(LogRHI, Error, TEXT("SwapChainDesc.BufferDesc.RefreshRate.Numerator %i"), Desc.BufferDesc.RefreshRate.Numerator);
+				UE_LOG(LogRHI, Error, TEXT("SwapChainDesc.BufferDesc.RefreshRate.Denominator %i"), Desc.BufferDesc.RefreshRate.Denominator);
+				UE_LOG(LogRHI, Error, TEXT("SwapChainDesc.BufferDesc.Format %i"), Desc.BufferDesc.Format);
+				UE_LOG(LogRHI, Error, TEXT("SwapChainDesc.BufferDesc.ScanlineOrdering %i"), Desc.BufferDesc.ScanlineOrdering);
+				UE_LOG(LogRHI, Error, TEXT("SwapChainDesc.BufferDesc.Scaling %i"), Desc.BufferDesc.Scaling);
+				UE_LOG(LogRHI, Error, TEXT("SwapChainDesc.SampleDesc.Count %i"), Desc.SampleDesc.Count);
+				UE_LOG(LogRHI, Error, TEXT("SwapChainDesc.SampleDesc.Quality %i"), Desc.SampleDesc.Quality);
+				UE_LOG(LogRHI, Error, TEXT("SwapChainDesc.BufferUsage %i"), Desc.BufferUsage);
+				UE_LOG(LogRHI, Error, TEXT("SwapChainDesc.BufferCount %i"), Desc.BufferCount);
+				UE_LOG(LogRHI, Error, TEXT("SwapChainDesc.OutputWindow %p"), Desc.OutputWindow);
+				UE_LOG(LogRHI, Error, TEXT("SwapChainDesc.Windowed %s"), Desc.Windowed ? TEXT("true") : TEXT("false"));
+				UE_LOG(LogRHI, Error, TEXT("SwapChainDesc.SwapEffect %u"), Desc.SwapEffect);
+				UE_LOG(LogRHI, Error, TEXT("SwapChainDesc.Flags %u"), Desc.Flags);
+			}
 		}
 		if (PresentFailCount > 10)
 		{
@@ -484,7 +487,7 @@ bool FD3D11Viewport::PresentChecked(IRHICommandContext& RHICmdContext, int32 Syn
 	D3DRHI->GetDeviceContext()->OMSetRenderTargets(0,0,0);
 
 	UINT PresentID;
-	if (SUCCEEDED(SwapChain->GetLastPresentCount(&PresentID)))
+	if ( (SwapChain.IsValid()) && (SUCCEEDED(SwapChain->GetLastPresentCount(&PresentID))) )
 	{
 		GRHIPresentCounter = PresentID;
 	}
