Refactor how we detect integrated GPU in DX11/DX12, using the UMA (Unified Memory Architecture) caps from the device rather than checking if the adapter doesn't have non local memory.
Use VkPhysicalDeviceType on Vulkan to properly initialize GRHIDeviceIsIntegrated.

diff --git a/Engine/Source/Runtime/D3D12RHI/Private/D3D12Adapter.cpp b/Engine/Source/Runtime/D3D12RHI/Private/D3D12Adapter.cpp
--- a/Engine/Source/Runtime/D3D12RHI/Private/D3D12Adapter.cpp
+++ b/Engine/Source/Runtime/D3D12RHI/Private/D3D12Adapter.cpp
@@ -276,6 +276,7 @@ FD3D12AdapterDesc::FD3D12AdapterDesc(const DXGI_ADAPTER_DESC& InDesc, int32 InAd
 	, ResourceBindingTier(DeviceInfo.ResourceBindingTier)
 	, ResourceHeapTier(DeviceInfo.ResourceHeapTier)
 	, MaxRHIFeatureLevel(DeviceInfo.MaxRHIFeatureLevel)
+	, bUMA(DeviceInfo.bUMA)
 	, bSupportsWaveOps(DeviceInfo.bSupportsWaveOps)
 	, bSupportsAtomic64(DeviceInfo.bSupportsAtomic64)
 {
diff --git a/Engine/Source/Runtime/D3D12RHI/Private/D3D12Adapter.h b/Engine/Source/Runtime/D3D12RHI/Private/D3D12Adapter.h
--- a/Engine/Source/Runtime/D3D12RHI/Private/D3D12Adapter.h
+++ b/Engine/Source/Runtime/D3D12RHI/Private/D3D12Adapter.h
@@ -50,7 +50,7 @@ struct FD3D12DeviceBasicInfo
 	uint32                      NumDeviceNodes;
 	bool                        bSupportsWaveOps;
 	bool                        bSupportsAtomic64;
-
+	bool						bUMA;
 	ERHIFeatureLevel::Type      MaxRHIFeatureLevel;
 };
 
@@ -86,8 +86,8 @@ struct FD3D12AdapterDesc
 	/** Number of device nodes (read: GPUs) */
 	uint32 NumDeviceNodes = 1;
 
-	/** Whether the GPU is integrated or discrete. */
-	bool bIsIntegrated = false;
+	/** Whether the GPU has unified memory. This can be used to identify integrated GPU as well. */
+	bool bUMA = false;
 
 	/** Whether SM6.0 wave ops are supported */
 	bool bSupportsWaveOps = false;
diff --git a/Engine/Source/Runtime/D3D12RHI/Private/Windows/WindowsD3D12Device.cpp b/Engine/Source/Runtime/D3D12RHI/Private/Windows/WindowsD3D12Device.cpp
--- a/Engine/Source/Runtime/D3D12RHI/Private/Windows/WindowsD3D12Device.cpp
+++ b/Engine/Source/Runtime/D3D12RHI/Private/Windows/WindowsD3D12Device.cpp
@@ -472,6 +472,16 @@ inline bool GetSupportsAtomic64(IDXGIAdapter* Adapter, ID3D12Device* Device)
 	return (D3D12Caps9.AtomicInt64OnTypedResourceSupported || CheckDeviceForEmulatedAtomic64Support(Adapter, Device));
 }
 
+inline bool IsDeviceUMA(ID3D12Device* device)
+{
+	D3D12_FEATURE_DATA_ARCHITECTURE data = { 0 };
+	if (SUCCEEDED(device->CheckFeatureSupport(D3D12_FEATURE_ARCHITECTURE, &data, sizeof(D3D12_FEATURE_DATA_ARCHITECTURE))))
+	{
+		return (bool)(data.UMA);
+	}
+	return false;
+}
+
 /**
  * Attempts to create a D3D12 device for the adapter using at minimum MinFeatureLevel.
  * If creation is successful, true is returned and the max supported feature level is set in OutMaxFeatureLevel.
@@ -494,6 +504,8 @@ static bool SafeTestD3D12CreateDevice(IDXGIAdapter* Adapter, D3D_FEATURE_LEVEL M
 
 			OutInfo.MaxRHIFeatureLevel = FindMaxRHIFeatureLevel(OutInfo.MaxFeatureLevel, OutInfo.MaxShaderModel, OutInfo.ResourceBindingTier, OutInfo.bSupportsWaveOps, OutInfo.bSupportsAtomic64);
 
+			OutInfo.bUMA = IsDeviceUMA(Device);
+
 			Device->Release();
 			return true;
 		}
@@ -769,23 +781,6 @@ static uint32 CountAdapterOutputs(TRefCountPtr<IDXGIAdapter>& Adapter)
 	return OutputCount;
 }
 
-inline bool IsAdapterIntegrated(IDXGIAdapter* InAdapter)
-{
-	TRefCountPtr<IDXGIAdapter3> Adapter3;
-	InAdapter->QueryInterface(IID_PPV_ARGS(Adapter3.GetInitReference()));
-
-	// Simple heuristic but without profiling it's hard to do better
-	DXGI_QUERY_VIDEO_MEMORY_INFO NonLocalVideoMemoryInfo{};
-	if (Adapter3 && SUCCEEDED(Adapter3->QueryVideoMemoryInfo(0, DXGI_MEMORY_SEGMENT_GROUP_NON_LOCAL, &NonLocalVideoMemoryInfo)))
-	{
-		return NonLocalVideoMemoryInfo.Budget == 0;
-	}
-
-	// TODO: Using GPUDetect for Intel GPUs to check for integrated vs discrete status, pending GPUDetect update
-
-	return true;
-}
-
 void FD3D12DynamicRHIModule::FindAdapter()
 {
 	// Once we've chosen one we don't need to do it again.
@@ -901,11 +896,12 @@ void FD3D12DynamicRHIModule::FindAdapter()
 				);
 
 				UE_LOG(LogD3D12RHI, Log,
-					TEXT("  Adapter has %uMB of dedicated video memory, %uMB of dedicated system memory, and %uMB of shared system memory, %d output[s]"),
+					TEXT("  Adapter has %uMB of dedicated video memory, %uMB of dedicated system memory, and %uMB of shared system memory, %d output[s], UMA:%s"),
 					(uint32)(AdapterDesc.DedicatedVideoMemory / (1024 * 1024)),
 					(uint32)(AdapterDesc.DedicatedSystemMemory / (1024 * 1024)),
 					(uint32)(AdapterDesc.SharedSystemMemory / (1024 * 1024)),
-					OutputCount
+					OutputCount,
+					DeviceInfo.bUMA ? TEXT("true") : TEXT("false")
 				);
 
 				const bool bIsWARP = (RHIConvertToGpuVendorId(AdapterDesc.VendorId) == EGpuVendorId::Microsoft);
@@ -921,7 +917,6 @@ void FD3D12DynamicRHIModule::FindAdapter()
 
 				CurrentAdapter.NumDeviceNodes = DeviceInfo.NumDeviceNodes;
 				CurrentAdapter.GpuPreference = GpuPreference;
-				CurrentAdapter.bIsIntegrated = IsAdapterIntegrated(TempAdapter);
 
 				// If requested WARP, then reject all other adapters. If WARP not requested, then reject the WARP device if software rendering support is disallowed
 				const bool bSkipWARP = (bRequestedWARP && !bIsWARP) || (!bRequestedWARP && bIsWARP && !bAllowSoftwareRendering);
@@ -940,6 +935,8 @@ void FD3D12DynamicRHIModule::FindAdapter()
 
 				const bool bSkipAdapter = bSkipWARP || bSkipPerfHUDAdapter || bSkipHmdGraphicsAdapter || bSkipExplicitAdapter || FParse::Param(FCommandLine::Get(), TEXT("ForceZeroAdapters"));
 
+				const bool bIsIntegrated = DeviceInfo.bUMA;
+
 				if (!bSkipAdapter && CurrentAdapter.IsValid())
 				{
 					if (PreferredVendor != EGpuVendorId::Unknown && PreferredVendor == RHIConvertToGpuVendorId(AdapterDesc.VendorId) && !PreferredAdapter.IsValid())
@@ -947,7 +944,7 @@ void FD3D12DynamicRHIModule::FindAdapter()
 						PreferredAdapter = CurrentAdapter;
 					}
 					
-					if (!bIsWARP && !CurrentAdapter.bIsIntegrated)
+					if (!bIsWARP && !bIsIntegrated)
 					{
 						if (!FirstDiscreteAdapter.IsValid())
 						{
@@ -1079,8 +1076,10 @@ FDynamicRHI* FD3D12DynamicRHIModule::CreateRHI(ERHIFeatureLevel::Type RequestedF
 	
 	if (ChosenAdapters.Num() > 0 && ChosenAdapters[0].IsValid())
 	{
-		FGenericCrashContext::SetEngineData(TEXT("RHI.IntegratedGPU"), ChosenAdapters[0].Get()->GetDesc().bIsIntegrated ? TEXT("true") : TEXT("false"));
-		GRHIDeviceIsIntegrated = ChosenAdapters[0].Get()->GetDesc().bIsIntegrated;
+		const bool bIsIntegrated = ChosenAdapters[0].Get()->GetDesc().bUMA;
+		FGenericCrashContext::SetEngineData(TEXT("RHI.IntegratedGPU"), bIsIntegrated ? TEXT("true") : TEXT("false"));
+		GRHIDeviceIsIntegrated = bIsIntegrated;
+		UE_LOG(LogD3D12RHI, Log, TEXT("Integrated GPU (iGPU): %s"), GRHIDeviceIsIntegrated ? TEXT("true") : TEXT("false"));
 	}
 
 	const FString FeatureLevelString = LexToString(GMaxRHIFeatureLevel);
diff --git a/Engine/Source/Runtime/VulkanRHI/Private/VulkanRHI.cpp b/Engine/Source/Runtime/VulkanRHI/Private/VulkanRHI.cpp
--- a/Engine/Source/Runtime/VulkanRHI/Private/VulkanRHI.cpp
+++ b/Engine/Source/Runtime/VulkanRHI/Private/VulkanRHI.cpp
@@ -1042,6 +1042,11 @@ void FVulkanDynamicRHI::InitInstance()
 
 #endif
 
+#if PLATFORM_WINDOWS || PLATFORM_UNIX
+		GRHIDeviceIsIntegrated = (Device->GetDeviceProperties().deviceType == VK_PHYSICAL_DEVICE_TYPE_INTEGRATED_GPU);
+		UE_LOG(LogVulkanRHI, Log, TEXT("Integrated GPU (iGPU): %s"), GRHIDeviceIsIntegrated ? TEXT("true") : TEXT("false"));
+#endif
+
 		FRenderResource::InitPreRHIResources();
 		GIsRHIInitialized = true;
 	}
diff --git a/Engine/Source/Runtime/Windows/D3D11RHI/Private/D3D11RHIPrivate.h b/Engine/Source/Runtime/Windows/D3D11RHI/Private/D3D11RHIPrivate.h
--- a/Engine/Source/Runtime/Windows/D3D11RHI/Private/D3D11RHIPrivate.h
+++ b/Engine/Source/Runtime/Windows/D3D11RHI/Private/D3D11RHIPrivate.h
@@ -313,6 +313,13 @@ struct FD3D11TransitionData
 /** Forward declare the context for the AMD AGS utility library. */
 struct AGSContext;
 
+
+struct FD3D11DeviceBasicInfo
+{
+	D3D_FEATURE_LEVEL           MaxFeatureLevel;
+	bool						bUMA;
+};
+
 struct FD3D11Adapter
 {
 	/** Null if not supported or FindAdapter() wasn't called. */
@@ -326,19 +333,19 @@ struct FD3D11Adapter
 	/** Whether this is a software adapter */
 	bool bSoftwareAdapter;
 
-	/** Whether the GPU is integrated or discrete. */
-	bool bIsIntegrated;
+	/** Whether the GPU has unified memory. This can be used to identify integrated GPU as well. */
+	bool bUMA;
 
 	// constructors
 	FD3D11Adapter() 
 	{
 	}
 
-	FD3D11Adapter(TRefCountPtr<IDXGIAdapter> InDXGIAdapter, D3D_FEATURE_LEVEL InMaxSupportedFeatureLevel, bool bInSoftwareAdatper, bool InIsIntegrated)
+	FD3D11Adapter(TRefCountPtr<IDXGIAdapter> InDXGIAdapter, bool bInSoftwareAdatper, const FD3D11DeviceBasicInfo& DeviceInfo)
 		: DXGIAdapter(InDXGIAdapter)
-		, MaxSupportedFeatureLevel(InMaxSupportedFeatureLevel)
+		, MaxSupportedFeatureLevel(DeviceInfo.MaxFeatureLevel)
 		, bSoftwareAdapter(bInSoftwareAdatper)
-		, bIsIntegrated(InIsIntegrated)
+		, bUMA(DeviceInfo.bUMA)
 	{
 		if (DXGIAdapter.IsValid())
 		{
diff --git a/Engine/Source/Runtime/Windows/D3D11RHI/Private/Windows/WindowsD3D11Device.cpp b/Engine/Source/Runtime/Windows/D3D11RHI/Private/Windows/WindowsD3D11Device.cpp
--- a/Engine/Source/Runtime/Windows/D3D11RHI/Private/Windows/WindowsD3D11Device.cpp
+++ b/Engine/Source/Runtime/Windows/D3D11RHI/Private/Windows/WindowsD3D11Device.cpp
@@ -333,11 +333,27 @@ static D3D_FEATURE_LEVEL GetMaxAllowedD3DFeatureLevel()
 	return D3D_FEATURE_LEVEL_11_1;
 }
 
+inline bool IsDeviceUMA(ID3D11Device* device)
+{
+	ID3D11Device3* pDevice3 = nullptr;
+	if (SUCCEEDED(device->QueryInterface(__uuidof(ID3D11Device3), (void**)&pDevice3)))
+	{
+		D3D11_FEATURE_DATA_D3D11_OPTIONS2 FeatureData = {};
+		if (SUCCEEDED(pDevice3->CheckFeatureSupport(D3D11_FEATURE_D3D11_OPTIONS2, &FeatureData, sizeof(FeatureData))))
+		{
+			pDevice3->Release();
+			return (bool)(FeatureData.UnifiedMemoryArchitecture);
+		}
+		pDevice3->Release();
+	}
+	return false;
+}
+
 /**
  * Attempts to create a D3D11 device for the adapter using at most MaxFeatureLevel.
  * If creation is successful, true is returned and the supported feature level is set in OutFeatureLevel.
  */
-static bool SafeTestD3D11CreateDevice(IDXGIAdapter* Adapter,D3D_FEATURE_LEVEL MinFeatureLevel,D3D_FEATURE_LEVEL MaxFeatureLevel,D3D_FEATURE_LEVEL* OutFeatureLevel)
+static bool SafeTestD3D11CreateDevice(IDXGIAdapter* Adapter, D3D_FEATURE_LEVEL MinFeatureLevel,D3D_FEATURE_LEVEL MaxFeatureLevel, FD3D11DeviceBasicInfo& OutInfo)
 {
 	ID3D11Device* D3DDevice = nullptr;
 	ID3D11DeviceContext* D3DDeviceContext = nullptr;
@@ -392,6 +408,7 @@ static bool SafeTestD3D11CreateDevice(IDXGIAdapter* Adapter,D3D_FEATURE_LEVEL Mi
 		// We don't want software renderer. Ideally we specify D3D_DRIVER_TYPE_HARDWARE on creation but
 		// when we specify an adapter we need to specify D3D_DRIVER_TYPE_UNKNOWN (otherwise the call fails).
 		// We cannot check the device type later (seems this is missing functionality in D3D).
+		D3D_FEATURE_LEVEL DeviceFeatureLevel = (D3D_FEATURE_LEVEL)0;
 		HRESULT Result = D3D11CreateDevice(
 			Adapter,
 			D3D_DRIVER_TYPE_UNKNOWN,
@@ -401,11 +418,14 @@ static bool SafeTestD3D11CreateDevice(IDXGIAdapter* Adapter,D3D_FEATURE_LEVEL Mi
 			NumAllowedFeatureLevels,
 			D3D11_SDK_VERSION,
 			&D3DDevice,
-			OutFeatureLevel,
+			&DeviceFeatureLevel,
 			&D3DDeviceContext);
 
 		if (SUCCEEDED(Result))
 		{
+			OutInfo.MaxFeatureLevel = DeviceFeatureLevel;
+			OutInfo.bUMA = IsDeviceUMA(D3DDevice);
+
 			D3DDevice->Release();
 			D3DDeviceContext->Release();
 			return true;
@@ -777,8 +797,8 @@ void FD3D11DynamicRHIModule::FindAdapter()
 				LogDXGIAdapterDesc(AdapterDesc);
 			}
 
-			D3D_FEATURE_LEVEL ActualFeatureLevel = (D3D_FEATURE_LEVEL)0;
-			if(SafeTestD3D11CreateDevice(TempAdapter,MinAllowedFeatureLevel,MaxAllowedFeatureLevel,&ActualFeatureLevel))
+			FD3D11DeviceBasicInfo DeviceInfo;
+			if(SafeTestD3D11CreateDevice(TempAdapter, MinAllowedFeatureLevel, MaxAllowedFeatureLevel, DeviceInfo))
 			{
 				// Log some information about the available D3D11 adapters.
 				
@@ -788,36 +808,24 @@ void FD3D11DynamicRHIModule::FindAdapter()
 					TEXT("  %2u. '%s' (Feature Level %s)"),
 					AdapterIndex,
 					AdapterDesc.Description,
-					GetFeatureLevelString(ActualFeatureLevel)
+					GetFeatureLevelString(DeviceInfo.MaxFeatureLevel)
 					);
 				UE_LOG(LogD3D11RHI, Log,
-					TEXT("      %u/%u/%u MB DedicatedVideo/DedicatedSystem/SharedSystem, Outputs:%d, VendorId:0x%x"),
+					TEXT("      %u/%u/%u MB DedicatedVideo/DedicatedSystem/SharedSystem, Outputs:%d, VendorId:0x%x UMA:%s"),
 					(uint32)(AdapterDesc.DedicatedVideoMemory / (1024*1024)),
 					(uint32)(AdapterDesc.DedicatedSystemMemory / (1024*1024)),
 					(uint32)(AdapterDesc.SharedSystemMemory / (1024*1024)),
 					OutputCount,
-					AdapterDesc.VendorId
+					AdapterDesc.VendorId,
+					DeviceInfo.bUMA ? TEXT("true") : TEXT("false")
 					);
 
 				const bool bIsWARP = RHIConvertToGpuVendorId(AdapterDesc.VendorId) == EGpuVendorId::Microsoft;
 
-				// Simple heuristic but without profiling it's hard to do better
-				bool bIsNonLocalMemoryPresent = false;
-				TRefCountPtr<IDXGIAdapter3> TempDxgiAdapter3;
-				DXGI_QUERY_VIDEO_MEMORY_INFO NonLocalVideoMemoryInfo;
-				if (SUCCEEDED(TempAdapter->QueryInterface(_uuidof(IDXGIAdapter3), (void**)TempDxgiAdapter3.GetInitReference())) &&
-					TempDxgiAdapter3.IsValid() && SUCCEEDED(TempDxgiAdapter3->QueryVideoMemoryInfo(0, DXGI_MEMORY_SEGMENT_GROUP_NON_LOCAL, &NonLocalVideoMemoryInfo)))
-				{
-					bIsNonLocalMemoryPresent = NonLocalVideoMemoryInfo.Budget != 0;
-				}
-
-				// TODO: Using GPUDetect for Intel GPUs to check for integrated vs discrete status, pending GPUDetect update
-
-				const bool bIsIntegrated = !bIsNonLocalMemoryPresent;
 				// PerfHUD is for performance profiling
 				const bool bIsPerfHUD = !FCString::Stricmp(AdapterDesc.Description,TEXT("NVIDIA PerfHUD"));
 
-				FD3D11Adapter CurrentAdapter(TempAdapter, ActualFeatureLevel, bIsWARP, bIsIntegrated);
+				FD3D11Adapter CurrentAdapter(TempAdapter, bIsWARP, DeviceInfo);
 
 				// Add special check to support HMDs, which do not have associated outputs.
 				// To reject the software emulation, unless the cvar wants it.
@@ -836,6 +844,8 @@ void FD3D11DynamicRHIModule::FindAdapter()
 				
 				const bool bSkipAdapter = bSkipSoftwareAdapter || bSkipPerfHUDAdapter || bSkipHmdGraphicsAdapter || bSkipExplicitAdapter;
 
+				const bool bIsIntegrated = DeviceInfo.bUMA;
+
 				if (!bSkipAdapter && CurrentAdapter.IsValid())
 				{
 					if (PreferredVendor != EGpuVendorId::Unknown && PreferredVendor == RHIConvertToGpuVendorId(AdapterDesc.VendorId) && !PreferredAdapter.IsValid())
@@ -843,7 +853,7 @@ void FD3D11DynamicRHIModule::FindAdapter()
 						PreferredAdapter = CurrentAdapter;
 					}
 					
-					if (!bIsWARP && !CurrentAdapter.bIsIntegrated)
+					if (!bIsWARP && !bIsIntegrated)
 					{
 						if (!FirstDiscreteAdapter.IsValid())
 						{
@@ -917,7 +927,8 @@ void FD3D11DynamicRHIModule::FindAdapter()
 	GRHIVendorId = ChosenAdapter.DXGIAdapterDesc.VendorId;
 	GRHIDeviceId = ChosenAdapter.DXGIAdapterDesc.DeviceId;
 	GRHIDeviceRevision = ChosenAdapter.DXGIAdapterDesc.Revision;
-	GRHIDeviceIsIntegrated = ChosenAdapter.bIsIntegrated;
+	GRHIDeviceIsIntegrated = ChosenAdapter.bUMA;
+	UE_LOG(LogD3D11RHI, Log, TEXT("Integrated GPU (iGPU): %s"), GRHIDeviceIsIntegrated ? TEXT("true") : TEXT("false"));
 }
 
 FDynamicRHI* FD3D11DynamicRHIModule::CreateRHI(ERHIFeatureLevel::Type RequestedFeatureLevel)
@@ -960,7 +971,7 @@ FDynamicRHI* FD3D11DynamicRHIModule::CreateRHI(ERHIFeatureLevel::Type RequestedF
 	}
 #endif
 
-	FGenericCrashContext::SetEngineData(TEXT("RHI.IntegratedGPU"), ChosenAdapter.bIsIntegrated ? TEXT("true") : TEXT("false"));
+	FGenericCrashContext::SetEngineData(TEXT("RHI.IntegratedGPU"), ChosenAdapter.bUMA ? TEXT("true") : TEXT("false"));
 
 	return FinalRHI;
 }
@@ -1225,7 +1236,7 @@ void FD3D11DynamicRHI::InitD3DDevice()
 		TRefCountPtr<IDXGIAdapter3> DxgiAdapter3;
 		DXGI_QUERY_VIDEO_MEMORY_INFO LocalVideoMemoryInfo;
 		FD3D11GlobalStats::GTotalGraphicsMemory = 0;
-		if (Adapter.bIsIntegrated)
+		if (Adapter.bUMA)
 		{
 			// It's all system memory.
 			FD3D11GlobalStats::GTotalGraphicsMemory = FD3D11GlobalStats::GDedicatedVideoMemory;
@@ -1635,7 +1646,7 @@ void FD3D11DynamicRHI::InitD3DDevice()
 		}
 
 		// Disable the RHI thread by default for devices that will likely suffer in performance
-		if (Adapter.bIsIntegrated || FPlatformMisc::NumberOfCores() < 4)
+		if (FPlatformMisc::NumberOfCores() < 4)
 		{
 			GRHISupportsRHIThread = false;
 		}
diff --git a/Engine/Source/Runtime/Windows/D3D11RHI/Public/Windows/D3D11ThirdParty.h b/Engine/Source/Runtime/Windows/D3D11RHI/Public/Windows/D3D11ThirdParty.h
--- a/Engine/Source/Runtime/Windows/D3D11RHI/Public/Windows/D3D11ThirdParty.h
+++ b/Engine/Source/Runtime/Windows/D3D11RHI/Public/Windows/D3D11ThirdParty.h
@@ -12,7 +12,7 @@
 #endif
 
 THIRD_PARTY_INCLUDES_START
-	#include <d3d11_2.h>
+	#include <d3d11_4.h>
 	#include <dxgi1_6.h>
 	#include <dxgidebug.h>
 THIRD_PARTY_INCLUDES_END
