Aggregate audio device initial implementation

diff --git a/Engine/Source/Runtime/AudioDeviceEnumeration/Windows/WindowsMMDeviceEnumeration/Private/WindowsMMDeviceInfoCache.cpp b/Engine/Source/Runtime/AudioDeviceEnumeration/Windows/WindowsMMDeviceEnumeration/Private/WindowsMMDeviceInfoCache.cpp
--- a/Engine/Source/Runtime/AudioDeviceEnumeration/Windows/WindowsMMDeviceEnumeration/Private/WindowsMMDeviceInfoCache.cpp
+++ b/Engine/Source/Runtime/AudioDeviceEnumeration/Windows/WindowsMMDeviceEnumeration/Private/WindowsMMDeviceInfoCache.cpp
@@ -8,6 +8,7 @@
 #include "Windows/AllowWindowsPlatformAtomics.h"
 
 THIRD_PARTY_INCLUDES_START
+#include <devicetopology.h>
 #include <functiondiscoverykeys_devpkey.h>
 THIRD_PARTY_INCLUDES_END
 
@@ -33,6 +34,8 @@ FWindowsMMDeviceCache::FCacheEntry& FWindowsMMDeviceCache::FCacheEntry::operator
 	Type = MoveTemp(InOther.Type);
 	ChannelBitmask = MoveTemp(InOther.ChannelBitmask);
 	OutputChannels = MoveTemp(InOther.OutputChannels);
+	HardwareId = MoveTemp(InOther.HardwareId);
+	FilterId = MoveTemp(InOther.FilterId);
 	return *this;
 }
 
@@ -48,6 +51,8 @@ FWindowsMMDeviceCache::FCacheEntry& FWindowsMMDeviceCache::FCacheEntry::operator
 	Type = InOther.Type;
 	ChannelBitmask = InOther.ChannelBitmask;
 	OutputChannels = InOther.OutputChannels;
+	HardwareId = InOther.HardwareId;
+	FilterId = InOther.FilterId;
 	return *this;
 }
 
@@ -284,6 +289,73 @@ bool FWindowsMMDeviceCache::EnumerateDeviceProps(const TComPtr<IMMDevice>& InDev
 	return true;
 }
 
+bool FWindowsMMDeviceCache::EnumerateHardwareTopology(const TComPtr<IMMDevice>& InDevice, FCacheEntry& OutInfo)
+{
+	if (!DeviceEnumerator)
+	{
+		return false;
+	}
+
+	TComPtr<IDeviceTopology> RenderEndpointTopology;
+	HRESULT Result = InDevice->Activate(__uuidof(IDeviceTopology), CLSCTX_ALL, NULL, (void**)&RenderEndpointTopology);
+	if (FAILED(Result))
+	{
+		UE_LOG(LogAudioEnumeration, Verbose, TEXT("FWindowsMMDeviceCache::EnumerateHardwareTopology failed to fetch IDeviceTopology: 0x%x"), Result);
+		return false;
+	}
+
+	TComPtr<IConnector> Connector;
+	Result = RenderEndpointTopology->GetConnector(0, &Connector);
+	if (FAILED(Result))
+	{
+		UE_LOG(LogAudioEnumeration, Verbose, TEXT("FWindowsMMDeviceCache::EnumerateHardwareTopology failed to fetch connector: 0x%x"), Result);
+		return false;
+	}
+
+	LPWSTR RenderFilterId = nullptr;
+	Result = Connector->GetDeviceIdConnectedTo(&RenderFilterId);
+	if (FAILED(Result) || !RenderFilterId)
+	{
+		UE_LOG(LogAudioEnumeration, Verbose, TEXT("FWindowsMMDeviceCache::EnumerateHardwareTopology failed to fetch render filter Id: 0x%x"), Result);
+		return false;
+	}
+
+	FString TempFilterId(RenderFilterId);
+	CoTaskMemFree(RenderFilterId);
+
+	TComPtr<IMMDevice> RenderDevnode;
+	Result = DeviceEnumerator->GetDevice(*TempFilterId, &RenderDevnode);
+	if (FAILED(Result))
+	{
+		UE_LOG(LogAudioEnumeration, Verbose, TEXT("FWindowsMMDeviceCache::EnumerateHardwareTopology failed to fetch render node: 0x%x"), Result);
+		return false;
+	}
+
+	TComPtr<IPropertyStore> PropertyStore;
+	Result = RenderDevnode->OpenPropertyStore(STGM_READ, &PropertyStore);
+	if (FAILED(Result))
+	{
+		UE_LOG(LogAudioEnumeration, Verbose, TEXT("FWindowsMMDeviceCache::EnumerateHardwareTopology failed to open property store: 0x%x"), Result);
+		return false;
+	}
+
+	PROPVARIANT HardwareId;
+	PropVariantInit(&HardwareId);
+	Result = PropertyStore->GetValue(PKEY_Device_InstanceId, &HardwareId);
+	if (FAILED(Result) || !HardwareId.pwszVal)
+	{
+		UE_LOG(LogAudioEnumeration, Verbose, TEXT("FWindowsMMDeviceCache::EnumerateHardwareTopology failed to fetch hardward Id: 0x%x"), Result);
+		return false;
+	}
+
+	OutInfo.HardwareId = FString(HardwareId.pwszVal);
+	PropVariantClear(&HardwareId);
+
+	OutInfo.FilterId = TempFilterId;
+
+	return true;
+}
+
 void FWindowsMMDeviceCache::EnumerateEndpoints()
 {
 	// Build a new cache from scratch.
@@ -312,6 +384,9 @@ void FWindowsMMDeviceCache::EnumerateEndpoints()
 						// Enumerate props into our info object.
 						EnumerateDeviceProps(Device, Info);
 
+						// Enumerate hardware topology to fetch hardware Id.
+						EnumerateHardwareTopology(Device, Info);
+
 						UE_LOG(LogAudioEnumeration, Verbose, TEXT("FWindowsMMDeviceCache: %s Device '%s' ID='%s'"),
 							Info.Type == FCacheEntry::EEndpointType::Capture ? TEXT("Capture") :
 							Info.Type == FCacheEntry::EEndpointType::Render ? TEXT("Render") :
@@ -551,6 +626,45 @@ TArray<Audio::FAudioPlatformDeviceInfo> FWindowsMMDeviceCache::GetAllActiveOutpu
 	return ActiveDevices;
 }
 
+TArray<Audio::FAudioPlatformDeviceInfo> FWindowsMMDeviceCache::GetAggregateDeviceInfo(FName InDeviceId) const
+{
+	TArray<FAudioPlatformDeviceInfo> AggregateDevice;
+	FReadScopeLock ReadLock(CacheMutationLock);
+
+	if (const FCacheEntry* DeviceEntry = Cache.Find(InDeviceId))
+	{
+		if (!DeviceEntry->HardwareId.IsEmpty())
+		{
+			TArray<FCacheEntry> ActiveDevices;
+
+			for (const auto& Entry : Cache)
+			{
+				if (Entry.Value.HardwareId == DeviceEntry->HardwareId &&
+					Entry.Value.State == EAudioDeviceState::Active &&
+					Entry.Value.Type == DeviceEntry->Type)
+				{
+					ActiveDevices.Emplace(Entry.Value);
+				}
+			}
+
+			// Sort the devices according to their filter id which usually ends with a vendor specific
+			// string which is sortable
+			Algo::Sort(ActiveDevices, [](const FCacheEntry& InEntryA, const FCacheEntry& InEntryB)
+			{
+				return InEntryA.FilterId.Compare(InEntryB.FilterId) < 0;
+			});
+
+			for (const FCacheEntry& Device : ActiveDevices)
+			{
+				FAudioPlatformDeviceInfo& Info = AggregateDevice.Emplace_GetRef();
+				MakeDeviceInfo(Device, GetDefaultOutputDevice_NoLock(), Info);
+			}
+		}
+	}
+
+	return AggregateDevice;
+}
+
 FName FWindowsMMDeviceCache::GetDefaultOutputDevice_NoLock() const
 {
 	if (!DefaultRenderId[(int32)EAudioDeviceRole::Console].IsNone())
diff --git a/Engine/Source/Runtime/AudioDeviceEnumeration/Windows/WindowsMMDeviceEnumeration/Public/WindowsMMDeviceInfoCache.h b/Engine/Source/Runtime/AudioDeviceEnumeration/Windows/WindowsMMDeviceEnumeration/Public/WindowsMMDeviceInfoCache.h
--- a/Engine/Source/Runtime/AudioDeviceEnumeration/Windows/WindowsMMDeviceEnumeration/Public/WindowsMMDeviceInfoCache.h
+++ b/Engine/Source/Runtime/AudioDeviceEnumeration/Windows/WindowsMMDeviceEnumeration/Public/WindowsMMDeviceInfoCache.h
@@ -40,6 +40,8 @@ struct WINDOWSMMDEVICEENUMERATION_API FWindowsMMDeviceCache : IAudioMixerDeviceC
 		int32 SampleRate = 0;
 		EEndpointType Type = EEndpointType::Unknown;
 		uint32 ChannelBitmask = 0;				// Bitfield used to build output channels, for easy comparison.
+		FString HardwareId;						// Unique string of the physical hardware device this MMDevice belongs to
+		FString FilterId;						// Unique identifier for this device containing product id (pid), vendor id (vid), etc.
 
 		TArray<EAudioMixerChannel::Type> OutputChannels;	// TODO. Generate this from the ChannelNum and bitmask when we are asked for it.
 		mutable FRWLock MutationLock;
@@ -73,6 +75,8 @@ struct WINDOWSMMDEVICEENUMERATION_API FWindowsMMDeviceCache : IAudioMixerDeviceC
 
 	bool EnumerateDeviceProps(const TComPtr<IMMDevice>& InDevice, FCacheEntry& OutInfo);
 
+	bool EnumerateHardwareTopology(const TComPtr<IMMDevice>& InDevice, FCacheEntry& OutInfo);
+
 	void EnumerateEndpoints();
 
 	void EnumerateDefaults();
@@ -97,6 +101,8 @@ struct WINDOWSMMDEVICEENUMERATION_API FWindowsMMDeviceCache : IAudioMixerDeviceC
 
 	virtual TArray<FAudioPlatformDeviceInfo> GetAllActiveOutputDevices() const override;
 
+	virtual TArray<FAudioPlatformDeviceInfo> GetAggregateDeviceInfo(FName InDeviceId) const override;
+
 	FName GetDefaultOutputDevice_NoLock() const;
 
 	TOptional<FAudioPlatformDeviceInfo> FindDefaultOutputDevice() const override;
diff --git a/Engine/Source/Runtime/AudioMixerCore/Public/AudioMixer.h b/Engine/Source/Runtime/AudioMixerCore/Public/AudioMixer.h
--- a/Engine/Source/Runtime/AudioMixerCore/Public/AudioMixer.h
+++ b/Engine/Source/Runtime/AudioMixerCore/Public/AudioMixer.h
@@ -231,6 +231,7 @@ namespace Audio
 		virtual TOptional<FAudioPlatformDeviceInfo> FindActiveOutputDevice(FName InDeviceID) const = 0;
 		virtual TArray<FAudioPlatformDeviceInfo> GetAllActiveOutputDevices() const = 0;
 		virtual TOptional<FAudioPlatformDeviceInfo> FindDefaultOutputDevice() const = 0;
+		virtual TArray<FAudioPlatformDeviceInfo> GetAggregateDeviceInfo(FName InDeviceID) const = 0;
 	};
 
 	/** Defines parameters needed for opening a new audio stream to device. */
diff --git a/Engine/Source/Runtime/Windows/AudioMixerWasapi/AudioMixerWasapi.Build.cs b/Engine/Source/Runtime/Windows/AudioMixerWasapi/AudioMixerWasapi.Build.cs
--- a/Engine/Source/Runtime/Windows/AudioMixerWasapi/AudioMixerWasapi.Build.cs
+++ b/Engine/Source/Runtime/Windows/AudioMixerWasapi/AudioMixerWasapi.Build.cs
@@ -15,6 +15,7 @@ public class AudioMixerWasapi : ModuleRules
 				{
 					"Engine",
 					"AudioPlatformSupportWasapi",
+					"SignalProcessing",
 					"WindowsMMDeviceEnumeration"
 				}
 			);
diff --git a/Engine/Source/Runtime/Windows/AudioMixerWasapi/Private/AudioMixerWasapi.cpp b/Engine/Source/Runtime/Windows/AudioMixerWasapi/Private/AudioMixerWasapi.cpp
--- a/Engine/Source/Runtime/Windows/AudioMixerWasapi/Private/AudioMixerWasapi.cpp
+++ b/Engine/Source/Runtime/Windows/AudioMixerWasapi/Private/AudioMixerWasapi.cpp
@@ -2,11 +2,22 @@
 
 #include "AudioMixerWasapi.h"
 #include "WasapiAudioUtils.h"
+#include "WasapiDefaultDeviceMgr.h"
+#include "WasapiAggregateDeviceMgr.h"
 
 namespace Audio
 {
-	FAudioMixerWasapi::FAudioMixerWasapi()
+	FAudioMixerWasapi::FAudioMixerWasapi(const bool bInUseAggregateDevice) :
+		bUseWasapiAggregateDevice(bInUseAggregateDevice)
 	{
+		if (bUseWasapiAggregateDevice)
+		{
+			ensure(DeviceManager = MakeUnique<FWasapiAggregateDeviceMgr>());
+		}
+		else
+		{
+			ensure(DeviceManager = MakeUnique<FWasapiDefaultDeviceMgr>());
+		}
 	}
 
 	FAudioMixerWasapi::~FAudioMixerWasapi()
@@ -21,24 +32,17 @@ namespace Audio
 
 		FAudioPlatformSettings EngineSettings = GetPlatformSettings();
 
-		MainRenderStreamDevice = MakeUnique<FAudioMixerWasapiRenderStream>();
-		if (!MainRenderStreamDevice.IsValid())
-		{
-			UE_LOG(LogAudioMixer, Error, TEXT("FAudioMixerWasapi::InitializeHardware failed to create MainRenderStreamDevice"));
-			return false;
-		}
-
-		MainRenderStreamDevice->OnReadNextBuffer().BindRaw(this, &FAudioMixerWasapi::ReadNextBuffer);
-
-		FWasapiRenderStreamParams StreamParams;
-		if (!InitStreamParams(AUDIO_MIXER_DEFAULT_DEVICE_INDEX, EngineSettings.CallbackBufferFrameSize, EngineSettings.SampleRate, StreamParams))
+		TArray<FWasapiRenderStreamParams> StreamParams;
+		if (!InitStreamParams(AUDIO_MIXER_DEFAULT_DEVICE_INDEX, EngineSettings.CallbackBufferFrameSize, EngineSettings.NumBuffers, EngineSettings.SampleRate, StreamParams))
 		{
+			AUDIO_PLATFORM_LOG_ONCE(TEXT("FAudioMixerWasapi::InitializeHardware InitStreamParams() failed"), Warning);
 			return false;
 		}
 
-		if (!InitAudioStreamDevice(StreamParams))
+		TFunction<void()> ReadNextBufferCallback = [this]() { ReadNextBuffer(); };
+		if (!ensure(DeviceManager && DeviceManager->InitializeHardware(StreamParams, ReadNextBufferCallback)))
 		{
-			UE_LOG(LogAudioMixer, Error, TEXT("FAudioMixerWasapi::InitializeHardware failed InitAudioStreamDevice"));
+			AUDIO_PLATFORM_LOG_ONCE(TEXT("FAudioMixerWasapi::InitializeHardware InitializeHardware() failed"), Warning);
 			return false;
 		}
 
@@ -53,11 +57,10 @@ namespace Audio
 			return false;
 		}
 
-		if (MainRenderStreamDevice.IsValid())
+		if (DeviceManager.IsValid() && !DeviceManager->TeardownHardware())
 		{
-			// Teardown the main device which will also unbind our delegate
-			MainRenderStreamDevice->TeardownHardware();
-			MainRenderStreamDevice.Reset();
+			AUDIO_PLATFORM_LOG_ONCE(TEXT("FAudioMixerWasapi::TeardownHardware DeviceManager->TeardownHardware() failed."), Warning);
+			return false;
 		}
 
 		bIsInitialized = false;
@@ -67,14 +70,14 @@ namespace Audio
 
 	bool FAudioMixerWasapi::IsInitialized() const
 	{
-		return bIsInitialized && MainRenderStreamDevice.IsValid() && MainRenderStreamDevice->IsInitialized();
+		return bIsInitialized && DeviceManager.IsValid() && DeviceManager->IsInitialized();
 	}
 
 	int32 FAudioMixerWasapi::GetNumFrames(const int32 InNumReqestedFrames)
 	{
-		if (MainRenderStreamDevice.IsValid())
+		if (DeviceManager.IsValid())
 		{
-			return MainRenderStreamDevice->GetNumFrames(InNumReqestedFrames);
+			return DeviceManager->GetNumFrames(InNumReqestedFrames);
 		}
 
 		return InNumReqestedFrames;
@@ -137,7 +140,7 @@ namespace Audio
 		return true;
 	}
 
-	bool FAudioMixerWasapi::InitStreamParams(const uint32 InDeviceIndex, const int32 InNumBufferFrames, const int32 InSampleRate, FWasapiRenderStreamParams& OutParams)
+	bool FAudioMixerWasapi::InitStreamParams(const uint32 InDeviceIndex, const int32 InNumBufferFrames, const int32 InNumBuffers, const int32 InSampleRate, TArray<FWasapiRenderStreamParams>& OutParams)
 	{
 		FAudioPlatformDeviceInfo DeviceInfo;
 		if (!GetOutputDeviceInfo(InDeviceIndex, DeviceInfo))
@@ -146,26 +149,39 @@ namespace Audio
 			return false;
 		}
 
-		TComPtr<IMMDevice> MMDevice = GetMMDevice(DeviceInfo.DeviceId);
-		if (!MMDevice)
+		if (bUseWasapiAggregateDevice)
 		{
-			UE_LOG(LogAudioMixer, Error, TEXT("FAudioMixerWasapi::InitStreamParams null MMDevice"));
-			return false;
-		}
-
-		OutParams = FWasapiRenderStreamParams(MMDevice, DeviceInfo, InNumBufferFrames, InSampleRate);
-
-		return true;
-	}
+			if (IAudioPlatformDeviceInfoCache* Cache = GetDeviceInfoCache())
+			{
+				FName DeviceId = *DeviceInfo.DeviceId;
+				TArray<FAudioPlatformDeviceInfo> AggregateDevice = Cache->GetAggregateDeviceInfo(DeviceId);
 
-	bool FAudioMixerWasapi::InitAudioStreamDevice(const FWasapiRenderStreamParams& InStreamParams)
-	{
-		if (ensure(MainRenderStreamDevice.IsValid()))
+				for (const FAudioPlatformDeviceInfo& AggregateDeviceInfo : AggregateDevice)
+				{
+					TComPtr<IMMDevice> MMDevice = GetMMDevice(AggregateDeviceInfo.DeviceId);
+					if (!MMDevice)
+					{
+						UE_LOG(LogAudioMixer, Error, TEXT("FAudioMixerWasapi::InitStreamParams null MMDevice"));
+						return false;
+					}
+
+					OutParams.Emplace(FWasapiRenderStreamParams(MMDevice, AggregateDeviceInfo, InNumBufferFrames, InNumBuffers, InSampleRate));
+				}
+			}
+		}
+		else
 		{
-			return MainRenderStreamDevice->InitializeHardware(InStreamParams);
+			TComPtr<IMMDevice> MMDevice = GetMMDevice(DeviceInfo.DeviceId);
+			if (!MMDevice)
+			{
+				UE_LOG(LogAudioMixer, Error, TEXT("FAudioMixerWasapi::InitStreamParams null MMDevice"));
+				return false;
+			}
+
+			OutParams.Emplace(FWasapiRenderStreamParams(MMDevice, DeviceInfo, InNumBufferFrames, InNumBuffers, InSampleRate));
 		}
 
-		return false;
+		return true;
 	}
 
 	bool FAudioMixerWasapi::OpenAudioStream(const FAudioMixerOpenStreamParams& Params)
@@ -179,33 +195,23 @@ namespace Audio
 		AudioStreamInfo.NumBuffers = OpenStreamParams.NumBuffers;
 		AudioStreamInfo.AudioMixer = OpenStreamParams.AudioMixer;
 
-		FWasapiRenderStreamParams StreamParams;
-		if (!InitStreamParams(OpenStreamParams.OutputDeviceIndex, OpenStreamParams.NumFrames, OpenStreamParams.SampleRate, StreamParams))
+		TArray<FWasapiRenderStreamParams> StreamParams;
+		if (!InitStreamParams(OpenStreamParams.OutputDeviceIndex, OpenStreamParams.NumFrames, OpenStreamParams.NumBuffers, OpenStreamParams.SampleRate, StreamParams))
 		{
 			return false;
 		}
 
-		AudioStreamInfo.DeviceInfo = StreamParams.HardwareDeviceInfo;
-		
-		HANDLE EventHandle = nullptr;
-		const TFunction<void()> RenderCallback = [this]() { MainRenderStreamDevice->DeviceRenderCallback(); };
-		RenderDeviceThread = MakeUnique<FAudioMixerWasapiDeviceThread>(RenderCallback, EventHandle);
-		
-		if (!RenderDeviceThread.IsValid())
-		{
-			UE_LOG(LogAudioMixer, Error, TEXT("Unable to create RenderDeviceThread"));
-			return false;
-		}
+		// Adopt the first device info. In the case of an aggregate device, all of the sub-devices will
+		// be identical because they belong to the same physical device.
+		AudioStreamInfo.DeviceInfo = StreamParams[0].HardwareDeviceInfo;
 
-		if (EventHandle == nullptr)
-		{
-			UE_LOG(LogAudioMixer, Error, TEXT("OpenAudioStream null EventHandle"));
-			return false;
-		}
-
-		if (!MainRenderStreamDevice->OpenAudioStream(StreamParams, EventHandle))
+		// Assign the total number of direct out channels based on the device manager. For WasapiDefaultDeviceMgr this will be 0 
+		// and for WasapiAggregateDeviceMgr this will be the total of all the channels less the main outs (first 8 channels).
+		AudioStreamInfo.DeviceInfo.NumDirectOutChannels = DeviceManager->GetNumDirectOutChannels();
+		
+		if (!ensure(DeviceManager.IsValid() && DeviceManager->OpenAudioStream(StreamParams)))
 		{
-			UE_LOG(LogAudioMixer, Error, TEXT("OpenAudioStream failed to open main audio device"));
+			UE_LOG(LogAudioMixer, Error, TEXT("FAudioMixerWasapi::OpenAudioStream DeviceManager->OpenAudioStream() failed"));
 			return false;
 		}
 
@@ -223,7 +229,10 @@ namespace Audio
 			return false;
 		}
 
-		MainRenderStreamDevice->CloseAudioStream();
+		if (!DeviceManager->CloseAudioStream())
+		{
+			return false;
+		}
 
 		AudioStreamInfo.StreamState = EAudioOutputStreamState::Closed;
 
@@ -232,15 +241,12 @@ namespace Audio
 
 	bool FAudioMixerWasapi::StartAudioStream()
 	{
-		if (IsInitialized() && RenderDeviceThread.IsValid())
+		if (IsInitialized() && DeviceManager->StartAudioStream())
 		{
 			BeginGeneratingAudio();
 
 			AudioStreamInfo.StreamState = EAudioOutputStreamState::Running;
 
-			MainRenderStreamDevice->StartAudioStream();
-			RenderDeviceThread->Start();
-
 			return true;
 		}
 
@@ -262,14 +268,9 @@ namespace Audio
 				StopGeneratingAudio();
 			}
 
-			if (RenderDeviceThread.IsValid())
+			if (DeviceManager.IsValid())
 			{
-				RenderDeviceThread->Stop();
-			}
-
-			if (MainRenderStreamDevice.IsValid())
-			{
-				MainRenderStreamDevice->StopAudioStream();
+				DeviceManager->StopAudioStream();
 			}
 
 			check(AudioStreamInfo.StreamState == EAudioOutputStreamState::Stopped);
@@ -285,9 +286,21 @@ namespace Audio
 
 	void FAudioMixerWasapi::SubmitBuffer(const uint8* InBuffer)
 	{
-		if (MainRenderStreamDevice.IsValid())
+		SCOPED_NAMED_EVENT(FAudioMixerWasapi_SubmitBuffer, FColor::Blue);
+
+		if (DeviceManager.IsValid())
+		{
+			DeviceManager->SubmitBuffer(InBuffer, OpenStreamParams.NumFrames);
+		}
+	}
+
+	void FAudioMixerWasapi::SubmitDirectOutBuffer(const int32 InDirectOutIndex, const FAlignedFloatBuffer& InBuffer)
+	{
+		SCOPED_NAMED_EVENT(FAudioMixerWasapi_SubmitDirectOutBuffer, FColor::Green);
+		
+		if (DeviceManager.IsValid())
 		{
-			MainRenderStreamDevice->SubmitBuffer(InBuffer, OpenStreamParams.NumFrames);
+			DeviceManager->SubmitDirectOutBuffer(InDirectOutIndex, InBuffer);
 		}
 	}
 
diff --git a/Engine/Source/Runtime/Windows/AudioMixerWasapi/Private/AudioMixerWasapi.h b/Engine/Source/Runtime/Windows/AudioMixerWasapi/Private/AudioMixerWasapi.h
--- a/Engine/Source/Runtime/Windows/AudioMixerWasapi/Private/AudioMixerWasapi.h
+++ b/Engine/Source/Runtime/Windows/AudioMixerWasapi/Private/AudioMixerWasapi.h
@@ -3,8 +3,8 @@
 #pragma once
 
 #include "AudioMixer.h"	
-#include "AudioMixerWasapiDeviceThread.h"
 #include "AudioMixerWasapiRenderStream.h"
+#include "IAudioMixerWasapiDeviceManager.h"
 #include "Misc/ScopeRWLock.h"
 
 #include "Windows/AllowWindowsPlatformTypes.h"
@@ -20,11 +20,6 @@ THIRD_PARTY_INCLUDES_END
 #include "Windows/HideWindowsPlatformAtomics.h"
 #include "Windows/HideWindowsPlatformTypes.h"
 
-/**
- *  Warning: AudioMixerWasapi module is currently considered experimental and may change in the future.
- *           We do not recommend shipping projects with Experimental features.
- */
-
 namespace Audio
 {
 	/**
@@ -34,7 +29,8 @@ namespace Audio
 	{
 	public:
 
-		FAudioMixerWasapi();
+		FAudioMixerWasapi() = default;
+		explicit FAudioMixerWasapi(const bool bInUseAggregateDevice);
 		virtual ~FAudioMixerWasapi();
 
 		//~ Begin IAudioMixerPlatformInterface
@@ -53,6 +49,7 @@ namespace Audio
 		virtual bool StopAudioStream() override;
 		virtual FAudioPlatformDeviceInfo GetPlatformDeviceInfo() const override;
 		virtual void SubmitBuffer(const uint8* Buffer) override;
+		virtual void SubmitDirectOutBuffer(const int32 InDirectOutIndex, const FAlignedFloatBuffer& InBuffer) override;
 		virtual bool DisablePCMAudioCaching() const override { return true; }
 		virtual FString GetDefaultDeviceName() override;
 		virtual FAudioPlatformSettings GetPlatformSettings() const override;
@@ -76,22 +73,19 @@ namespace Audio
 		/** Used for determining if a device swap has been requested. */
 		FString OriginalAudioDeviceId;
 
-		/** The main audio device for outputting up to 8 channels. */
-		TUniquePtr<FAudioMixerWasapiRenderStream> MainRenderStreamDevice;
+		/** Manages either a single, default device or an aggregate of several devices belonging to the same hardware */
+		TUniquePtr<IAudioMixerWasapiDeviceManager> DeviceManager;
+
+		/** Indicates if this object is aggregating devices or not. */
+		bool bUseWasapiAggregateDevice = false;
 
 		/** Indicates if this object has been successfully initialized. */
 		bool bIsInitialized = false;
 
-		/** The thread which provides an execution context during audio playback. */
-		TUniquePtr<FAudioMixerWasapiDeviceThread> RenderDeviceThread;
-
 		/** Fetches an IMMDevice with the given ID. */
 		TComPtr<IMMDevice> GetMMDevice(const FString& InDeviceID) const;
 
 		/** Initializes a Wasapi stream paramters struct with the give values. */
-		bool InitStreamParams(const uint32 InDeviceIndex, const int32 InNumBufferFrames, const int32 InSampleRate, FWasapiRenderStreamParams& OutParams);
-
-		/** Initializes a device with the given parameters. */
-		bool InitAudioStreamDevice(const FWasapiRenderStreamParams& InStreamParams);
+		bool InitStreamParams(const uint32 InDeviceIndex, const int32 InNumBufferFrames, const int32 InNumBuffers, const int32 InSampleRate, TArray<FWasapiRenderStreamParams>& OutParams);
 	};
  }
diff --git a/Engine/Source/Runtime/Windows/AudioMixerWasapi/Private/AudioMixerWasapiDeviceThread.cpp b/Engine/Source/Runtime/Windows/AudioMixerWasapi/Private/AudioMixerWasapiDeviceThread.cpp
--- a/Engine/Source/Runtime/Windows/AudioMixerWasapi/Private/AudioMixerWasapiDeviceThread.cpp
+++ b/Engine/Source/Runtime/Windows/AudioMixerWasapi/Private/AudioMixerWasapiDeviceThread.cpp
@@ -7,18 +7,52 @@
 
 namespace Audio
 {
+	std::atomic<uint32> FAudioMixerWasapiDeviceThread::AudioDeviceThreadCounter = 0;
+
 	FAudioMixerWasapiRunnable::FAudioMixerWasapiRunnable(const TFunction<void()>& InDeviceRenderCallback, HANDLE& OutEventHandle) :
 		DeviceRenderCallback(InDeviceRenderCallback)
 	{
-		// Not using FEvent/FEventWin here because we need access to the raw, platform
-		// handle (see SetEventHandler() below).
-		EventHandle = ::CreateEvent(nullptr, 0, 0, nullptr);
-		if (EventHandle == nullptr)
+		if (CreateEventHandles())
 		{
-			UE_LOG(LogAudioMixerWasapi, Error, TEXT("FAudioMixerWasapiRunnable CreateEvent failed"));
+			OutEventHandle = EventHandles[0];
 		}
+	}
 
-		OutEventHandle = EventHandle;
+	FAudioMixerWasapiRunnable::FAudioMixerWasapiRunnable(const TFunction<void()>& InDeviceRenderCallback, TArray<HANDLE>& OutEventHandles, const int32 InNumRequestedHandles) :
+		DeviceRenderCallback(InDeviceRenderCallback)
+	{
+		if (CreateEventHandles(InNumRequestedHandles))
+		{
+			OutEventHandles = EventHandles;
+		}
+	}
+
+	bool FAudioMixerWasapiRunnable::CreateEventHandles(const uint32 InNumRequestedHandles)
+	{
+		uint32 NumHandles = InNumRequestedHandles;
+		if (NumHandles > MAXIMUM_WAIT_OBJECTS)
+		{
+			NumHandles = MAXIMUM_WAIT_OBJECTS;
+			UE_LOG(LogAudioMixerWasapi, Warning, TEXT("FAudioMixerWasapiRunnable::CreateEventHandles requested %d handles which exceeds max supported. Clamping to %d. "), InNumRequestedHandles, MAXIMUM_WAIT_OBJECTS);
+		}
+
+		EventHandles.SetNumZeroed(NumHandles);
+
+		for (uint32 Index = 0; Index < NumHandles; ++Index)
+		{
+			// Not using FEvent/FEventWin here because we need access to the raw, platform
+			// handle (see SetEventHandler() below).
+			HANDLE EventHandle = ::CreateEvent(nullptr, 0, 0, nullptr);
+			if (EventHandle == nullptr)
+			{
+				EventHandles.Reset();
+				return false;
+			}
+
+			EventHandles[Index] = EventHandle;
+		}
+
+		return true;
 	}
 
 	uint32 FAudioMixerWasapiRunnable::Run()
@@ -27,16 +61,21 @@ namespace Audio
 
 		bool bCoInitialized = FWindowsPlatformMisc::CoInitialize(ECOMModel::Multithreaded);
 
+		const uint32 NumHandles = EventHandles.Num();
+		const HANDLE* Handles = EventHandles.GetData();
+
 		while (bIsRunning.load())
 		{
 			static constexpr uint32 TimeoutInMs = 1000;
 
-			int32 Result = ::WaitForSingleObject(EventHandle, TimeoutInMs);
+			// WASAPI events for multiple audio devices will all be signaled at the
+			// same time when beloning to the same physical device.
+			uint32 Result = ::WaitForMultipleObjects(NumHandles, Handles, true, TimeoutInMs);
 			if (Result == WAIT_TIMEOUT)
 			{
 				++OutputStreamTimeoutsDetected;
 			}
-			else if (Result == WAIT_OBJECT_0)
+			else if (Result >= WAIT_OBJECT_0 && Result < (WAIT_OBJECT_0 + NumHandles))
 			{
 				DeviceRenderCallback();
 			}
@@ -64,12 +103,17 @@ namespace Audio
 	{
 	}
 
+	FAudioMixerWasapiDeviceThread::FAudioMixerWasapiDeviceThread(const TFunction<void()>& InDeviceRenderCallback, TArray<HANDLE>& OutEventHandles, const int32 InNumRequestedHandles) :
+		DeviceRenderRunnable(MakeUnique<FAudioMixerWasapiRunnable>(InDeviceRenderCallback, OutEventHandles, InNumRequestedHandles))
+	{
+	}
+
 	bool FAudioMixerWasapiDeviceThread::Start()
 	{
 		TRACE_CPUPROFILER_EVENT_SCOPE(Audio::FAudioMixerWasapiDeviceThread::Start);
 		check(DeviceRenderThread == nullptr);
 
-		DeviceRenderThread = TUniquePtr<FRunnableThread>(FRunnableThread::Create(DeviceRenderRunnable.Get(), TEXT("Audio Render Device Thread"), 0, TPri_TimeCritical));
+		DeviceRenderThread = TUniquePtr<FRunnableThread>(FRunnableThread::Create(DeviceRenderRunnable.Get(), *FString::Printf(TEXT("AudioDeviceThread(%d)"), ++AudioDeviceThreadCounter), 0, TPri_TimeCritical));
 		return DeviceRenderThread.IsValid();
 	}
 
diff --git a/Engine/Source/Runtime/Windows/AudioMixerWasapi/Private/AudioMixerWasapiDeviceThread.h b/Engine/Source/Runtime/Windows/AudioMixerWasapi/Private/AudioMixerWasapiDeviceThread.h
--- a/Engine/Source/Runtime/Windows/AudioMixerWasapi/Private/AudioMixerWasapiDeviceThread.h
+++ b/Engine/Source/Runtime/Windows/AudioMixerWasapi/Private/AudioMixerWasapiDeviceThread.h
@@ -21,12 +21,7 @@ THIRD_PARTY_INCLUDES_END
 #include "Windows/HideWindowsPlatformAtomics.h"
 #include "Windows/HideWindowsPlatformTypes.h"
 
-/**
- *  Warning: AudioMixerWasapi module is currently considered experimental and may change in the future.
- *           We do not recommend shipping projects with Experimental features.
- */
-
-	namespace Audio
+namespace Audio
 {
 	/**
 	 * FAudioMixerWasapiRunnable - The runnable which executes the main thread loop for FWasapiCaptureThread.
@@ -39,8 +34,13 @@ THIRD_PARTY_INCLUDES_END
 		/**
 		 * FAudioMixerWasapiRunnable
 		 */
-		explicit FAudioMixerWasapiRunnable(const TFunction<void()>& InDeviceRenderCallback, HANDLE& OutEventHandle);
-		
+		FAudioMixerWasapiRunnable(const TFunction<void()>& InDeviceRenderCallback, HANDLE& OutEventHandle);
+
+		/**
+		 * FAudioMixerWasapiRunnable
+		 */
+		FAudioMixerWasapiRunnable(const TFunction<void()>& InDeviceRenderCallback, TArray<HANDLE>& OutEventHandles, const int32 InNumRequestedHandles);
+
 		/** Default destructor */
 		virtual ~FAudioMixerWasapiRunnable() = default;
 
@@ -57,7 +57,7 @@ THIRD_PARTY_INCLUDES_END
 		 * Event handle which our audio thread waits on prior to each callback. WASAPI signals this
 		 * object each quanta when a buffer of audio has been rendered and is ready for more data.
 		 */
-		HANDLE EventHandle = nullptr;
+		TArray<HANDLE> EventHandles;
 
 		/**
 		 * Accumulates timeouts which occur when the thread event timeout is reached
@@ -67,6 +67,8 @@ THIRD_PARTY_INCLUDES_END
 		
 		/** Callback function to be called each time the device signals it is ready for another buffer of audio. */
 		TFunction<void()> DeviceRenderCallback;
+
+		bool CreateEventHandles(const uint32 InNumRequestedHandles = 1);
 	};
 
 	/**
@@ -80,7 +82,12 @@ THIRD_PARTY_INCLUDES_END
 		/**
 		 * FAudioMixerWasapiDeviceThread
 		 */
-		explicit FAudioMixerWasapiDeviceThread(const TFunction<void()>& InDeviceRenderCallback, HANDLE& OutEventHandle);
+		FAudioMixerWasapiDeviceThread(const TFunction<void()>& InDeviceRenderCallback, HANDLE& OutEventHandle);
+
+		/**
+		 * FAudioMixerWasapiDeviceThread
+		 */
+		FAudioMixerWasapiDeviceThread(const TFunction<void()>& InDeviceRenderCallback, TArray<HANDLE>& OutEventHandles, const int32 InNumRequestedHandles);
 
 		/**
 		 * Start - Creates the FRunnableThread object which immediately begins running the FAudioMixerWasapiRunnable member.
@@ -101,6 +108,9 @@ THIRD_PARTY_INCLUDES_END
 		void Abort();
 
 	private:
+		/** Tracks number of device threads created. */
+		static std::atomic<uint32> AudioDeviceThreadCounter;
+
 		/** The thread which is the context that the runnable executes in. */
 		TUniquePtr<FRunnableThread> DeviceRenderThread;
 
diff --git a/Engine/Source/Runtime/Windows/AudioMixerWasapi/Private/AudioMixerWasapiModule.cpp b/Engine/Source/Runtime/Windows/AudioMixerWasapi/Private/AudioMixerWasapiModule.cpp
--- a/Engine/Source/Runtime/Windows/AudioMixerWasapi/Private/AudioMixerWasapiModule.cpp
+++ b/Engine/Source/Runtime/Windows/AudioMixerWasapi/Private/AudioMixerWasapiModule.cpp
@@ -3,6 +3,7 @@
 #include "AudioMixerWasapi.h"
 #include "AudioMixerWasapiLog.h"
 #include "Features/IModularFeatures.h"
+#include "Misc/ConfigCacheIni.h"
 #include "Modules/ModuleManager.h"
 
 DEFINE_LOG_CATEGORY(LogAudioMixerWasapi);
@@ -25,7 +26,11 @@ public:
 
 	virtual Audio::IAudioMixerPlatformInterface* CreateAudioMixerPlatformInterface() override
 	{
-		return new Audio::FAudioMixerWasapi();
+		bool bUseWasapiAggregateDevice = false;
+
+		GConfig->GetBool(TEXT("Audio"), TEXT("UseWasapiAggregateDevice"), bUseWasapiAggregateDevice, GEngineIni);
+
+		return new Audio::FAudioMixerWasapi(bUseWasapiAggregateDevice);
 	}
 };
 
diff --git a/Engine/Source/Runtime/Windows/AudioMixerWasapi/Private/AudioMixerWasapiRenderStream.cpp b/Engine/Source/Runtime/Windows/AudioMixerWasapi/Private/AudioMixerWasapiRenderStream.cpp
--- a/Engine/Source/Runtime/Windows/AudioMixerWasapi/Private/AudioMixerWasapiRenderStream.cpp
+++ b/Engine/Source/Runtime/Windows/AudioMixerWasapi/Private/AudioMixerWasapiRenderStream.cpp
@@ -17,8 +17,16 @@ namespace Audio
 
 	bool FAudioMixerWasapiRenderStream::InitializeHardware(const FWasapiRenderStreamParams& InParams)
 	{
+		TComPtr<IMMDevice> MMDevice = InParams.MMDevice;
+
+		if (!MMDevice)
+		{
+			UE_LOG(LogAudioMixerWasapi, Error, TEXT("InitializeHardware null MMDevice"));
+			return false;
+		}
+
 		TComPtr<IAudioClient3> TempAudioClient;
-		HRESULT Result = InParams.MMDevice->Activate(__uuidof(IAudioClient3), CLSCTX_INPROC_SERVER, nullptr, IID_PPV_ARGS_Helper(&TempAudioClient));
+		HRESULT Result = MMDevice->Activate(__uuidof(IAudioClient3), CLSCTX_INPROC_SERVER, nullptr, IID_PPV_ARGS_Helper(&TempAudioClient));
 		if (FAILED(Result))
 		{
 			UE_LOG(LogAudioMixerWasapi, Error, TEXT("OpenAudioStream failed IMMDevice::Activate 0x%x"), (uint32)Result);
@@ -116,7 +124,6 @@ namespace Audio
 			return false;
 		}
 
-		ReadNextBufferDelegate.Unbind();
 		RenderClient.Reset();
 		AudioClient.Reset();
 
@@ -186,6 +193,13 @@ namespace Audio
 			return false;
 		}
 
+		if (StreamState == EAudioOutputStreamState::Running)
+		{
+			UE_LOG(LogAudioMixerWasapi, Verbose, TEXT("FAudioMixerWasapiRenderStream::CloseAudioStream stream appears to be running. StopAudioStream() must be called prior to closing."));
+
+			return false;
+		}
+
 		StreamState = EAudioOutputStreamState::Closed;
 
 		return true;
@@ -203,12 +217,6 @@ namespace Audio
 				return false;
 			}
 
-			if (!ReadNextBufferDelegate.IsBound())
-			{
-				UE_LOG(LogAudioMixerWasapi, Error, TEXT("StartAudioStream failed buffer delegate not bound"));
-				return false;
-			}
-
 			AudioClient->Start();
 		}
 
@@ -243,50 +251,4 @@ namespace Audio
 
 		return true;
 	}
-
-	void FAudioMixerWasapiRenderStream::DeviceRenderCallback()
-	{
-		SCOPED_NAMED_EVENT(FAudioMixerWasapiRenderStream_DeviceRenderCallback, FColor::Blue);
-
-		if (bIsInitialized)
-		{
-			uint32 NumFramesPadding = 0;
-			AudioClient->GetCurrentPadding(&NumFramesPadding);
-
-			// NumFramesPerDeviceBuffer is the buffer size WASAPI allocated. It is gauranteed to 
-			// be at least the amount requested. For example, if we request a 1024 frame buffer, WASAPI
-			// might allocate a 1056 frame buffer. The padding is subtracted from the allocated amount
-			// to determine how much space is available currently in the buffer.
-			const int32 NumFramesAvailable = NumFramesPerDeviceBuffer - NumFramesPadding;
-			if (NumFramesAvailable >= static_cast<int32>(RenderStreamParams.NumFrames))
-			{
-				check(!RenderBufferPtr);
-
-				if (SUCCEEDED(RenderClient->GetBuffer(RenderStreamParams.NumFrames, &RenderBufferPtr)))
-				{
-					ReadNextBufferDelegate.Execute();
-
-					HRESULT Result = RenderClient->ReleaseBuffer(RenderStreamParams.NumFrames, 0 /* flags */);
-					if (FAILED(Result))
-					{
-						++CallbackBufferErrors;
-					}
-
-					RenderBufferPtr = nullptr;
-				}
-				else
-				{
-					++CallbackBufferErrors;
-				}
-			}
-		}
-	}
-
-	void FAudioMixerWasapiRenderStream::SubmitBuffer(const uint8* InBuffer, const SIZE_T InNumFrames)
-	{
-		if (RenderBufferPtr)
-		{
-			FMemory::Memcpy(RenderBufferPtr, InBuffer, InNumFrames * AudioFormat.GetFrameSizeInBytes());
-		}
-	}
 }
diff --git a/Engine/Source/Runtime/Windows/AudioMixerWasapi/Private/AudioMixerWasapiRenderStream.h b/Engine/Source/Runtime/Windows/AudioMixerWasapi/Private/AudioMixerWasapiRenderStream.h
--- a/Engine/Source/Runtime/Windows/AudioMixerWasapi/Private/AudioMixerWasapiRenderStream.h
+++ b/Engine/Source/Runtime/Windows/AudioMixerWasapi/Private/AudioMixerWasapiRenderStream.h
@@ -3,6 +3,7 @@
 #pragma once
 
 #include "AudioMixer.h"	
+#include "IAudioMixerWasapiDeviceManager.h"
 #include "Misc/ScopeRWLock.h"
 #include "WasapiAudioFormat.h"
 
@@ -21,42 +22,8 @@ THIRD_PARTY_INCLUDES_END
 
 DECLARE_DELEGATE(FAudioMixerReadNextBufferDelegate);
 
-/**
- *  Warning: AudioMixerWasapi module is currently considered experimental and may change in the future.
- *           We do not recommend shipping projects with Experimental features.
- */
-
 namespace Audio
 {
-	/** Defines parameters needed for opening a new audio stream to device. */
-	struct FWasapiRenderStreamParams
-	{
-		/** The audio device to open. */
-		TComPtr<IMMDevice> MMDevice;
-
-		/** Hardware device configuration info. */
-		FAudioPlatformDeviceInfo HardwareDeviceInfo;
-
-		/** The number of desired audio frames in audio callback. */
-		uint32 NumFrames = 0;
-
-		/** The desired sample rate */
-		uint32 SampleRate = 0;
-
-		FWasapiRenderStreamParams() = default;
-
-		FWasapiRenderStreamParams(const TComPtr<IMMDevice>& InMMDevice, 
-			const FAudioPlatformDeviceInfo& InDeviceInfo, 
-			const uint32 InNumFrames, 
-			const uint32 InSampleRate) :
-			MMDevice(InMMDevice),
-			HardwareDeviceInfo(InDeviceInfo),
-			NumFrames(InNumFrames),
-			SampleRate(InSampleRate)
-		{
-		}
-	};
-
 	/**
 	 * FAudioMixerWasapiRenderStream
 	 */
@@ -67,23 +34,19 @@ namespace Audio
 		FAudioMixerWasapiRenderStream();
 		virtual ~FAudioMixerWasapiRenderStream();
 
-		bool InitializeHardware(const FWasapiRenderStreamParams& InParams);
-		bool TeardownHardware();
-		bool IsInitialized() const;
-		int32 GetNumFrames(const int32 InNumReqestedFrames) const;
-		bool OpenAudioStream(const FWasapiRenderStreamParams& InParams, HANDLE InEventHandle);
-		bool CloseAudioStream();
-		bool StartAudioStream();
-		bool StopAudioStream();
-		void SubmitBuffer(const uint8* Buffer, const SIZE_T InNumFrames);
-		void DeviceRenderCallback();
+		virtual bool InitializeHardware(const FWasapiRenderStreamParams& InParams);
+		virtual bool TeardownHardware();
+		virtual bool IsInitialized() const;
+		virtual int32 GetNumFrames(const int32 InNumReqestedFrames) const;
+		virtual bool OpenAudioStream(const FWasapiRenderStreamParams& InParams, HANDLE InEventHandle);
+		virtual bool CloseAudioStream();
+		virtual bool StartAudioStream();
+		virtual bool StopAudioStream();
+		virtual void SubmitBuffer(const uint8* Buffer, const SIZE_T InNumFrames) { }
+		virtual void SubmitDirectOutBuffer(const int32 InChannelIndex, const FAlignedFloatBuffer& InBuffer) { }
+		virtual void DeviceRenderCallback() = 0;
 
-		FAudioMixerReadNextBufferDelegate& OnReadNextBuffer() { return ReadNextBufferDelegate; }
-
-	private:
-
-		/** Delegate called each buffer callback to signal the mixer to process the next buffer. */
-		FAudioMixerReadNextBufferDelegate ReadNextBufferDelegate;
+	protected:
 
 		/** COM pointer to the WASAPI audio client object. */
 		TComPtr<IAudioClient3> AudioClient;
@@ -112,8 +75,5 @@ namespace Audio
 
 		/** Accumulates errors that occur in the audio callback. */
 		uint32 CallbackBufferErrors = 0;
-
-		/** Pointer to WASAPI render audio buffer filled in each callback by the mixer in SubmitBuffer(). */
-		uint8* RenderBufferPtr = nullptr;
 	};
 }
diff --git a/Engine/Source/Runtime/Windows/AudioMixerWasapi/Private/IAudioMixerWasapiDeviceManager.h b/Engine/Source/Runtime/Windows/AudioMixerWasapi/Private/IAudioMixerWasapiDeviceManager.h
new file mode 100644
--- /dev/null
+++ b/Engine/Source/Runtime/Windows/AudioMixerWasapi/Private/IAudioMixerWasapiDeviceManager.h
@@ -0,0 +1,82 @@
+// Copyright Epic Games, Inc. All Rights Reserved.
+
+#pragma once
+
+#include "AudioMixer.h"	
+#include "Misc/ScopeRWLock.h"
+#include "WasapiAudioFormat.h"
+
+#include "Windows/AllowWindowsPlatformTypes.h"
+#include "Windows/AllowWindowsPlatformAtomics.h"
+
+#include "Microsoft/COMPointer.h"
+
+THIRD_PARTY_INCLUDES_START
+#include <mmdeviceapi.h>			// IMMNotificationClient
+#include <audiopolicy.h>			// IAudioSessionEvents
+THIRD_PARTY_INCLUDES_END
+
+#include "Windows/HideWindowsPlatformAtomics.h"
+#include "Windows/HideWindowsPlatformTypes.h"
+
+DECLARE_DELEGATE(FAudioMixerReadNextBufferDelegate);
+
+namespace Audio
+{
+	/** Defines parameters needed for opening a new audio stream to device. */
+	struct FWasapiRenderStreamParams
+	{
+		/** The audio device to open. */
+		TComPtr<IMMDevice> MMDevice;
+
+		/** Hardware device configuration info. */
+		FAudioPlatformDeviceInfo HardwareDeviceInfo;
+
+		/** The number of desired audio frames in audio callback. */
+		uint32 NumFrames = 0;
+
+		/** The number of buffers to keep enqueued. More buffers increases latency, but can compensate for variable compute availability in audio callbacks on some platforms. */
+		int32 NumBuffers = 0;
+
+		/** The desired sample rate */
+		uint32 SampleRate = 0;
+
+		FWasapiRenderStreamParams() = default;
+
+		FWasapiRenderStreamParams(const TComPtr<IMMDevice>& InMMDevice,
+			const FAudioPlatformDeviceInfo& InDeviceInfo,
+			const uint32 InNumFrames,
+			const uint32 InNumBuffers,
+			const uint32 InSampleRate) :
+			MMDevice(InMMDevice),
+			HardwareDeviceInfo(InDeviceInfo),
+			NumFrames(InNumFrames),
+			NumBuffers(InNumBuffers),
+			SampleRate(InSampleRate)
+		{
+		}
+	};
+
+	/**
+	 * IAudioMixerWasapiDeviceManager - classes implementing this interface manage underlying render device streams
+	 */
+	class IAudioMixerWasapiDeviceManager
+	{
+	public:
+
+		IAudioMixerWasapiDeviceManager() = default;
+		virtual ~IAudioMixerWasapiDeviceManager() = default;
+
+		virtual bool InitializeHardware(const TArray<FWasapiRenderStreamParams>& InParams, const TFunction<void()>& InCallback) = 0;
+		virtual bool TeardownHardware() = 0;
+		virtual bool IsInitialized() const = 0;
+		virtual int32 GetNumDirectOutChannels() const = 0;
+		virtual int32 GetNumFrames(const int32 InNumReqestedFrames) const = 0;
+		virtual bool OpenAudioStream(const  TArray<FWasapiRenderStreamParams>& InParams) = 0;
+		virtual bool CloseAudioStream() = 0;
+		virtual bool StartAudioStream() = 0;
+		virtual bool StopAudioStream() = 0;
+		virtual void SubmitBuffer(const uint8* InBuffer, const SIZE_T InNumFrames) = 0;
+		virtual void SubmitDirectOutBuffer(const int32 InDirectOutIndex, const FAlignedFloatBuffer& InBuffer) = 0;
+	};
+}
diff --git a/Engine/Source/Runtime/Windows/AudioMixerWasapi/Private/WasapiAggregateDeviceMgr.cpp b/Engine/Source/Runtime/Windows/AudioMixerWasapi/Private/WasapiAggregateDeviceMgr.cpp
new file mode 100644
--- /dev/null
+++ b/Engine/Source/Runtime/Windows/AudioMixerWasapi/Private/WasapiAggregateDeviceMgr.cpp
@@ -0,0 +1,246 @@
+// Copyright Epic Games, Inc. All Rights Reserved.
+
+#include "WasapiAggregateDeviceMgr.h"
+
+#include "WasapiDefaultRenderStream.h"
+#include "WasapiAggregateRenderStream.h"
+
+namespace Audio
+{
+	constexpr static int32 MaxDeviceCount = 64;
+
+	FWasapiAggregateDeviceMgr::FWasapiAggregateDeviceMgr()
+	{
+		RenderStreamDevices.Reserve(MaxDeviceCount);
+	}
+
+	bool FWasapiAggregateDeviceMgr::InitializeHardware(const TArray<FWasapiRenderStreamParams>& InParams, const TFunction<void()>& InCallback)
+	{
+		RenderStreamDevices.Reset();
+
+		for (int32 DeviceIndex = 0; DeviceIndex < InParams.Num(); ++DeviceIndex)
+		{
+			// The first device is our main device. This is used for rendering the bed channels.
+			// The other devices (if any) are used for direct outputs (see ADM plugin).
+			if (DeviceIndex == 0)
+			{
+				TUniquePtr<FWasapiDefaultRenderStream> DefaultRenderStream = MakeUnique<FWasapiDefaultRenderStream>();
+
+				if (DefaultRenderStream.IsValid())
+				{
+					// The first device is always the main device which also drives the callback
+					DefaultRenderStream->OnReadNextBuffer().BindLambda(InCallback);
+
+					RenderStreamDevices.Emplace(MoveTemp(DefaultRenderStream));
+				}
+			}
+			else
+			{
+				RenderStreamDevices.Emplace(MakeUnique<FWasapiAggregateRenderStream>());
+			}
+
+			if (RenderStreamDevices.IsEmpty() || !RenderStreamDevices[DeviceIndex].IsValid())
+			{
+				UE_LOG(LogAudioMixer, Error, TEXT("FWasapiAggregateDeviceMgr::InitializeHardware failed to create RenderStreamDevice: %d"), DeviceIndex);
+				return false;
+			}
+
+			if (!RenderStreamDevices[DeviceIndex]->InitializeHardware(InParams[DeviceIndex]))
+			{
+				UE_LOG(LogAudioMixer, Error, TEXT("FWasapiAggregateDeviceMgr::InitializeHardware InitializeHardware failed for RenderStreamDevice: %d"), DeviceIndex);
+				return false;
+			}
+		}
+
+		NumChannelsPerDevice = InParams[0].HardwareDeviceInfo.NumChannels;
+		// Total channel count including main device and all direct out devices
+		const int32 NumDevices = RenderStreamDevices.Num();
+		NumDirectOutChannels = (NumDevices > 0) ? (NumChannelsPerDevice * (NumDevices - 1)) : 0;
+
+		bIsInitialized = true;
+
+		return true;
+	}
+
+	bool FWasapiAggregateDeviceMgr::TeardownHardware()
+	{
+		if (!RenderStreamDevices.IsEmpty())
+		{
+			for (TUniquePtr<FAudioMixerWasapiRenderStream>& Device : RenderStreamDevices)
+			{
+				if (Device.IsValid())
+				{
+					// Teardown the main device which will also unbind our delegate
+					Device->TeardownHardware();
+					Device.Reset();
+				}
+			}
+
+			RenderStreamDevices.Reset();
+
+			bIsInitialized = false;
+
+			return true;
+		}
+
+		return false;
+	}
+
+	bool FWasapiAggregateDeviceMgr::IsInitialized() const
+	{
+		return bIsInitialized;
+	}
+
+	int32 FWasapiAggregateDeviceMgr::GetNumFrames(const int32 InNumReqestedFrames) const
+	{
+		if (!RenderStreamDevices.IsEmpty())
+		{
+			return RenderStreamDevices[0]->GetNumFrames(InNumReqestedFrames);
+		}
+
+		return InNumReqestedFrames;
+	}
+
+	bool FWasapiAggregateDeviceMgr::OpenAudioStream(const TArray<FWasapiRenderStreamParams>& InParams)
+	{
+		TArray<HANDLE> EventHandles;
+		const TFunction<void()> RenderCallback = [this]() { DeviceRenderCallback(); };
+
+		const uint32 NumRenderDevices = RenderStreamDevices.Num();
+		if (NumRenderDevices == 0)
+		{
+			UE_LOG(LogAudioMixer, Error, TEXT("FWasapiAggregateDeviceMgr::OpenAudioStream no render devices found"));
+			return false;
+		}
+
+		RenderDeviceThread = MakeUnique<FAudioMixerWasapiDeviceThread>(RenderCallback, EventHandles, NumRenderDevices);
+
+		if (!RenderDeviceThread.IsValid())
+		{
+			UE_LOG(LogAudioMixer, Error, TEXT("Unable to create RenderDeviceThread"));
+			return false;
+		}
+
+		if (EventHandles.Num() != NumRenderDevices)
+		{
+			UE_LOG(LogAudioMixer, Error, TEXT("OpenAudioStream error creating EventHandles"));
+			return false;
+		}
+
+		for (uint32 Index = 0; Index < NumRenderDevices; ++Index)
+		{
+			if (!RenderStreamDevices[Index]->OpenAudioStream(InParams[Index], EventHandles[Index]))
+			{
+				UE_LOG(LogAudioMixer, Error, TEXT("OpenAudioStream failed to open main audio device"));
+				return false;
+			}
+		}
+
+		return true;
+	}
+
+	bool FWasapiAggregateDeviceMgr::CloseAudioStream()
+	{
+		if (!RenderStreamDevices.IsEmpty())
+		{
+			bool bDidAllClose = true;
+
+			for (TUniquePtr<FAudioMixerWasapiRenderStream>& Device : RenderStreamDevices)
+			{
+				if (Device.IsValid())
+				{
+					bDidAllClose |= Device->CloseAudioStream();
+				}
+			}
+
+			return bDidAllClose;
+		}
+
+		return false;
+	}
+
+	bool FWasapiAggregateDeviceMgr::StartAudioStream()
+	{
+		if (RenderStreamDevices.IsEmpty())
+		{
+			UE_LOG(LogAudioMixer, Error, TEXT("FWasapiAggregateDeviceMgr::StartAudioStream no devices available to start"));
+			return false;
+		}
+
+		for (TUniquePtr<FAudioMixerWasapiRenderStream>& Device : RenderStreamDevices)
+		{
+			if (Device.IsValid())
+			{
+				if (!Device->StartAudioStream())
+				{
+					UE_LOG(LogAudioMixer, Error, TEXT("FWasapiAggregateDeviceMgr::StartAudioStream unable to start render device"));
+					return false;
+				}
+			}
+		}
+
+		if (!RenderDeviceThread->Start())
+		{
+			UE_LOG(LogAudioMixer, Error, TEXT("FWasapiAggregateDeviceMgr::StartAudioStream failed to start device thread"));
+			return false;
+		}
+
+		return true;
+	}
+
+	bool FWasapiAggregateDeviceMgr::StopAudioStream()
+	{
+		if (RenderDeviceThread.IsValid())
+		{
+			RenderDeviceThread->Stop();
+		}
+
+		for (TUniquePtr<FAudioMixerWasapiRenderStream>& Device : RenderStreamDevices)
+		{
+			if (Device.IsValid())
+			{
+				Device->StopAudioStream();
+			}
+		}
+
+		return true;
+	}
+
+	void FWasapiAggregateDeviceMgr::SubmitBuffer(const uint8* InBuffer, const SIZE_T InNumFrames)
+	{
+		if (!RenderStreamDevices.IsEmpty())
+		{
+			RenderStreamDevices[0]->SubmitBuffer(InBuffer, InNumFrames);
+		}
+	}
+
+	void FWasapiAggregateDeviceMgr::SubmitDirectOutBuffer(const int32 InDirectOutIndex, const FAlignedFloatBuffer& InBuffer)
+	{
+		const int32 NumDevices = RenderStreamDevices.Num();
+
+		if (NumChannelsPerDevice > 0 && NumDevices > 1 &&
+			InDirectOutIndex >= 0 && InDirectOutIndex < NumDirectOutChannels)
+		{
+			// The first device is reserved for the main audio out (the bed channels)
+			// so we add one here.
+			const int32 RenderDeviceIndex = (InDirectOutIndex / NumChannelsPerDevice) + 1;
+			const int32 ChannelIndex = InDirectOutIndex % NumChannelsPerDevice;
+
+			if (RenderDeviceIndex > 0 && RenderDeviceIndex < NumDevices)
+			{
+				RenderStreamDevices[RenderDeviceIndex]->SubmitDirectOutBuffer(ChannelIndex, InBuffer);
+			}
+		}
+	}
+
+	void FWasapiAggregateDeviceMgr::DeviceRenderCallback()
+	{
+		for (TUniquePtr<FAudioMixerWasapiRenderStream>& Device : RenderStreamDevices)
+		{
+			if (Device.IsValid())
+			{
+				Device->DeviceRenderCallback();
+			}
+		}
+	}
+}
diff --git a/Engine/Source/Runtime/Windows/AudioMixerWasapi/Private/WasapiAggregateDeviceMgr.h b/Engine/Source/Runtime/Windows/AudioMixerWasapi/Private/WasapiAggregateDeviceMgr.h
new file mode 100644
--- /dev/null
+++ b/Engine/Source/Runtime/Windows/AudioMixerWasapi/Private/WasapiAggregateDeviceMgr.h
@@ -0,0 +1,61 @@
+// Copyright Epic Games, Inc. All Rights Reserved.
+
+#pragma once
+
+#include "IAudioMixerWasapiDeviceManager.h"
+#include "AudioMixerWasapiDeviceThread.h"
+#include "AudioMixerWasapiRenderStream.h"
+
+namespace Audio
+{
+	/**
+	 * FWasapiAggregateDeviceMgr
+	 * 
+	 * Manages the software device streams associated with a single, physical hardware device. 
+	 * This provides the ability to address more than 8 channels of a physical device using the
+	 * Windows Audio Streaming API (WASAPI).
+	 */
+	class FWasapiAggregateDeviceMgr : public IAudioMixerWasapiDeviceManager
+	{
+	public:
+
+		FWasapiAggregateDeviceMgr();
+		virtual ~FWasapiAggregateDeviceMgr() = default;
+
+		//~ Begin IAudioMixerWasapiDeviceManager
+		virtual bool InitializeHardware(const TArray<FWasapiRenderStreamParams>& InParams, const TFunction<void()>& InCallback) override;
+		virtual bool TeardownHardware() override;
+		virtual bool IsInitialized() const override;
+		virtual int32 GetNumDirectOutChannels() const override { return NumDirectOutChannels; }
+		virtual int32 GetNumFrames(const int32 InNumReqestedFrames) const override;
+		virtual bool OpenAudioStream(const TArray<FWasapiRenderStreamParams>& InParams) override;
+		virtual bool CloseAudioStream() override;
+		virtual bool StartAudioStream() override;
+		virtual bool StopAudioStream() override;
+		virtual void SubmitBuffer(const uint8* InBuffer, const SIZE_T InNumFrames) override;
+		virtual void SubmitDirectOutBuffer(const int32 InDirectOutIndex, const FAlignedFloatBuffer& InBuffer) override;
+		//~ End IAudioMixerWasapiDeviceManager
+
+	private:
+
+		bool bIsInitialized = false;
+
+		/** The number of channels per WASAPI deivce (max 8 channels) */
+		int32 NumChannelsPerDevice = 0;
+
+		/** The number of direct output channels supported by the physical device. The sum 
+		    total channels of all WASAPI audio devices belonging to the physical device, minus the 
+			number of main out channels (typically the first 8 channels)
+		*/
+		int32 NumDirectOutChannels = 0;
+
+		/** The render streams associated with this aggregate device */
+		TArray<TUniquePtr<FAudioMixerWasapiRenderStream>> RenderStreamDevices;
+
+		/** The thread which provides an execution context during audio playback. */
+		TUniquePtr<FAudioMixerWasapiDeviceThread> RenderDeviceThread;
+
+		void DeviceRenderCallback();
+
+	};
+ }
diff --git a/Engine/Source/Runtime/Windows/AudioMixerWasapi/Private/WasapiAggregateRenderStream.cpp b/Engine/Source/Runtime/Windows/AudioMixerWasapi/Private/WasapiAggregateRenderStream.cpp
new file mode 100644
--- /dev/null
+++ b/Engine/Source/Runtime/Windows/AudioMixerWasapi/Private/WasapiAggregateRenderStream.cpp
@@ -0,0 +1,124 @@
+// Copyright Epic Games, Inc. All Rights Reserved.
+
+#include "WasapiAggregateRenderStream.h"
+
+#include "AudioMixerWasapiLog.h"
+#include "DSP/FloatArrayMath.h"
+#include "WasapiAudioUtils.h"
+
+namespace Audio
+{
+	bool FWasapiAggregateRenderStream::InitializeHardware(const FWasapiRenderStreamParams& InParams)
+	{
+		DirectOutBuffers.Reset();
+		DirectOutBuffers.SetNum(InParams.HardwareDeviceInfo.NumChannels);
+
+		for (TCircularAudioBuffer<float>& Buffer : DirectOutBuffers)
+		{
+			int32 NumOutputBuffers = FMath::Max(InParams.NumBuffers, 2);
+			int32 BufferCapacity = InParams.NumFrames * NumOutputBuffers;
+
+			Buffer.SetCapacity(BufferCapacity);
+		}
+
+		InterleaveBuffers.Reset();
+		InterleaveBuffers.SetNum(InParams.HardwareDeviceInfo.NumChannels);
+
+		for (FAlignedFloatBuffer& Buffer : InterleaveBuffers)
+		{
+			Buffer.SetNumZeroed(InParams.NumFrames);
+		}
+
+		return FAudioMixerWasapiRenderStream::InitializeHardware(InParams);
+	}
+
+	bool FWasapiAggregateRenderStream::StartAudioStream()
+	{
+		FAudioMixerWasapiRenderStream::StartAudioStream();
+
+		return true;
+	}
+
+	void FWasapiAggregateRenderStream::SubmitDirectOutBuffer(const int32 InChannelIndex, const FAlignedFloatBuffer& InBuffer)
+	{
+		if (InChannelIndex >= 0 && InChannelIndex < RenderStreamParams.HardwareDeviceInfo.NumChannels)
+		{
+			TCircularAudioBuffer<float>& DirectOutBuffer = DirectOutBuffers[InChannelIndex];
+
+			if (InBuffer.Num() == RenderStreamParams.NumFrames)
+			{
+				DirectOutBuffer.Push(InBuffer.GetData(), InBuffer.Num());
+			}
+		}
+	}
+
+	void FWasapiAggregateRenderStream::DeviceRenderCallback()
+	{
+		SCOPED_NAMED_EVENT(FWasapiAggregateRenderStream_DeviceRenderCallback, FColor::Blue);
+
+		if (bIsInitialized)
+		{
+			uint32 NumFramesPadding = 0;
+			AudioClient->GetCurrentPadding(&NumFramesPadding);
+
+			// NumFramesPerDeviceBuffer is the buffer size WASAPI allocated. It is gauranteed to 
+			// be at least the amount requested. For example, if we request a 1024 frame buffer, WASAPI
+			// might allocate a 1056 frame buffer. The padding is subtracted from the allocated amount
+			// to determine how much space is available currently in the buffer.
+			const int32 NumFramesAvailable = NumFramesPerDeviceBuffer - NumFramesPadding;
+			if (NumFramesAvailable >= static_cast<int32>(RenderStreamParams.NumFrames))
+			{
+				uint8* RenderBufferPtr = nullptr;
+
+				if (SUCCEEDED(RenderClient->GetBuffer(RenderStreamParams.NumFrames, &RenderBufferPtr)))
+				{
+					InterleaveOutput((float*)RenderBufferPtr, RenderStreamParams.NumFrames);
+					
+					HRESULT Result = RenderClient->ReleaseBuffer(RenderStreamParams.NumFrames, 0 /* flags */);
+					if (FAILED(Result))
+					{
+						++CallbackBufferErrors;
+					}
+				}
+				else
+				{
+					++CallbackBufferErrors;
+				}
+			}
+		}
+	}
+
+	void FWasapiAggregateRenderStream::InterleaveOutput(float* OutRenderBufferPtr, const uint32 InNumFrames)
+	{
+		const int32 NumChannels = InterleaveBuffers.Num();
+		const int32 NumFrames = InterleaveBuffers[0].Num();
+
+		if (NumFrames == InNumFrames)
+		{
+			for (int32 Index = 0; Index < DirectOutBuffers.Num(); ++Index)
+			{
+				// Clear out the interleave buffer
+				InterleaveBuffers[Index].Reset();
+				InterleaveBuffers[Index].SetNumZeroed(NumFrames);
+
+				if (DirectOutBuffers[Index].Num() >= InNumFrames)
+				{
+					const int32 NumFramesPopped = DirectOutBuffers[Index].Pop(InterleaveBuffers[Index].GetData(), RenderStreamParams.NumFrames);
+				}
+			}
+
+			TArray<const float*> BufferPtrArray;
+			BufferPtrArray.Reset(NumChannels);
+
+			for (const FAlignedFloatBuffer& Buffer : InterleaveBuffers)
+			{
+				const float* BufferPtr = Buffer.GetData();
+				BufferPtrArray.Add(BufferPtr);
+			}
+
+			const float** InterleaveBufferPtr = BufferPtrArray.GetData();
+
+			ArrayInterleave(InterleaveBufferPtr, OutRenderBufferPtr, NumFrames, NumChannels);
+		}
+	}
+}
diff --git a/Engine/Source/Runtime/Windows/AudioMixerWasapi/Private/WasapiAggregateRenderStream.h b/Engine/Source/Runtime/Windows/AudioMixerWasapi/Private/WasapiAggregateRenderStream.h
new file mode 100644
--- /dev/null
+++ b/Engine/Source/Runtime/Windows/AudioMixerWasapi/Private/WasapiAggregateRenderStream.h
@@ -0,0 +1,35 @@
+// Copyright Epic Games, Inc. All Rights Reserved.
+
+#pragma once
+
+#include "AudioMixerWasapiRenderStream.h"
+
+namespace Audio
+{
+	/**
+	 * FWasapiAggregateRenderStream
+	 * A render stream which can be sent to direct output channels which 
+	 * are channels of an aggregate device above the first 8 channels which are
+	 * reserved for the main output of the editor/game.
+	 */
+	class FWasapiAggregateRenderStream : public FAudioMixerWasapiRenderStream
+	{
+	public:
+
+		virtual ~FWasapiAggregateRenderStream() = default;
+
+		//~ Begin FAudioMixerWasapiRenderStream
+		virtual bool InitializeHardware(const FWasapiRenderStreamParams& InParams) override;
+		virtual bool StartAudioStream() override;
+		virtual void SubmitDirectOutBuffer(const int32 InChannelIndex, const FAlignedFloatBuffer& InBuffer) override;
+		virtual void DeviceRenderCallback() override;
+		//~ End FAudioMixerWasapiRenderStream
+
+	private:
+
+		TArray<FAlignedFloatBuffer> InterleaveBuffers;
+		TArray<TCircularAudioBuffer<float>> DirectOutBuffers;
+
+		void InterleaveOutput(float* InRenderBufferPtr, const uint32 InNumFrames);
+	};
+}
diff --git a/Engine/Source/Runtime/Windows/AudioMixerWasapi/Private/WasapiDefaultDeviceMgr.cpp b/Engine/Source/Runtime/Windows/AudioMixerWasapi/Private/WasapiDefaultDeviceMgr.cpp
new file mode 100644
--- /dev/null
+++ b/Engine/Source/Runtime/Windows/AudioMixerWasapi/Private/WasapiDefaultDeviceMgr.cpp
@@ -0,0 +1,125 @@
+// Copyright Epic Games, Inc. All Rights Reserved.
+
+#include "WasapiDefaultDeviceMgr.h"
+
+namespace Audio
+{
+	bool FWasapiDefaultDeviceMgr::InitializeHardware(const TArray<FWasapiRenderStreamParams>& InParams, const TFunction<void()>& InCallback)
+	{
+		MainRenderStreamDevice = MakeUnique<FWasapiDefaultRenderStream>();
+		if (ensure(MainRenderStreamDevice.IsValid()) && InParams.Num() > 0)
+		{
+			MainRenderStreamDevice->OnReadNextBuffer().BindLambda(InCallback);
+
+			return MainRenderStreamDevice->InitializeHardware(InParams[0]);
+		}
+
+		return false;
+	}
+
+	bool FWasapiDefaultDeviceMgr::TeardownHardware()
+	{ 
+		if (MainRenderStreamDevice.IsValid())
+		{
+			// Teardown the main device which will also unbind our delegate
+			MainRenderStreamDevice->TeardownHardware();
+			MainRenderStreamDevice.Reset();
+
+			return true;
+		}
+
+		return false;
+	}
+
+	bool FWasapiDefaultDeviceMgr::IsInitialized() const
+	{
+		return MainRenderStreamDevice.IsValid() && MainRenderStreamDevice->IsInitialized();
+	}
+
+	int32 FWasapiDefaultDeviceMgr::GetNumFrames(const int32 InNumReqestedFrames) const
+	{ 
+		if (MainRenderStreamDevice.IsValid())
+		{
+			return MainRenderStreamDevice->GetNumFrames(InNumReqestedFrames);
+		}
+
+		return InNumReqestedFrames;
+	}
+
+	bool FWasapiDefaultDeviceMgr::OpenAudioStream(const TArray<FWasapiRenderStreamParams>& InParams)
+	{ 
+		HANDLE EventHandle = nullptr;
+		const TFunction<void()> RenderCallback = [this]() { MainRenderStreamDevice->DeviceRenderCallback(); };
+		RenderDeviceThread = MakeUnique<FAudioMixerWasapiDeviceThread>(RenderCallback, EventHandle);
+
+		if (!RenderDeviceThread.IsValid())
+		{
+			UE_LOG(LogAudioMixer, Error, TEXT("Unable to create RenderDeviceThread"));
+			return false;
+		}
+
+		if (EventHandle == nullptr)
+		{
+			UE_LOG(LogAudioMixer, Error, TEXT("OpenAudioStream null EventHandle"));
+			return false;
+		}
+
+		if (!MainRenderStreamDevice->OpenAudioStream(InParams[0], EventHandle))
+		{
+			UE_LOG(LogAudioMixer, Error, TEXT("OpenAudioStream failed to open main audio device"));
+			return false;
+		}
+
+		return true;
+	}
+
+	bool FWasapiDefaultDeviceMgr::CloseAudioStream()
+	{
+		if (RenderDeviceThread.IsValid())
+		{
+			return MainRenderStreamDevice->CloseAudioStream();
+		}
+
+		return false;
+	}
+
+	bool FWasapiDefaultDeviceMgr::StartAudioStream()
+	{ 
+		if (!MainRenderStreamDevice->StartAudioStream())
+		{
+			UE_LOG(LogAudioMixer, Error, TEXT("FWasapiDefaultDeviceMgr::StartAudioStream failed to start main audio device"));
+			return false;
+		}
+
+		if (!RenderDeviceThread->Start())
+		{
+			UE_LOG(LogAudioMixer, Error, TEXT("FWasapiDefaultDeviceMgr::StartAudioStream failed to start device thread"));
+			return false;
+		}
+		
+		return true;
+	}
+
+	bool FWasapiDefaultDeviceMgr::StopAudioStream()
+	{
+		if (RenderDeviceThread.IsValid())
+		{
+			RenderDeviceThread->Stop();
+		}
+
+		if (MainRenderStreamDevice.IsValid())
+		{
+			MainRenderStreamDevice->StopAudioStream();
+		}
+
+		return true;
+	}
+
+	void FWasapiDefaultDeviceMgr::SubmitBuffer(const uint8* InBuffer, const SIZE_T InNumFrames)
+	{
+		if (MainRenderStreamDevice.IsValid())
+		{
+			MainRenderStreamDevice->SubmitBuffer(InBuffer, InNumFrames);
+		}
+	}
+}
diff --git a/Engine/Source/Runtime/Windows/AudioMixerWasapi/Private/WasapiDefaultDeviceMgr.h b/Engine/Source/Runtime/Windows/AudioMixerWasapi/Private/WasapiDefaultDeviceMgr.h
new file mode 100644
--- /dev/null
+++ b/Engine/Source/Runtime/Windows/AudioMixerWasapi/Private/WasapiDefaultDeviceMgr.h
@@ -0,0 +1,44 @@
+// Copyright Epic Games, Inc. All Rights Reserved.
+
+#pragma once
+
+#include "IAudioMixerWasapiDeviceManager.h"
+#include "AudioMixerWasapiDeviceThread.h"
+#include "WasapiDefaultRenderStream.h"
+
+namespace Audio
+{
+	/**
+	 * FAudioMixerWasapi 
+	 * Manages a single audio device which is used for the main (first eight) output channels.
+	 */
+	class FWasapiDefaultDeviceMgr : public IAudioMixerWasapiDeviceManager
+	{
+	public:
+
+		FWasapiDefaultDeviceMgr() = default;
+		virtual ~FWasapiDefaultDeviceMgr() = default;
+
+		//~ Begin IAudioMixerWasapiDeviceManager
+		virtual bool InitializeHardware(const TArray<FWasapiRenderStreamParams>& InParams, const TFunction<void()>& InCallback) override;
+		virtual bool TeardownHardware() override;
+		virtual bool IsInitialized() const override;
+		virtual int32 GetNumDirectOutChannels() const override { return 0; }
+		virtual int32 GetNumFrames(const int32 InNumReqestedFrames) const override;
+		virtual bool OpenAudioStream(const TArray<FWasapiRenderStreamParams>& InParams) override;
+		virtual bool CloseAudioStream() override;
+		virtual bool StartAudioStream() override;
+		virtual bool StopAudioStream() override;
+		virtual void SubmitBuffer(const uint8* InBuffer, const SIZE_T InNumFrames) override;
+		virtual void SubmitDirectOutBuffer(const int32 InDirectOutIndex, const FAlignedFloatBuffer& InBuffer) { }
+		//~ End IAudioMixerWasapiDeviceManager
+
+	private:
+	
+		/** The main audio device for outputting up to 8 channels. */
+		TUniquePtr<FWasapiDefaultRenderStream> MainRenderStreamDevice;
+
+		/** The thread which provides an execution context during audio playback. */
+		TUniquePtr<FAudioMixerWasapiDeviceThread> RenderDeviceThread;
+	};
+ }
diff --git a/Engine/Source/Runtime/Windows/AudioMixerWasapi/Private/WasapiDefaultRenderStream.cpp b/Engine/Source/Runtime/Windows/AudioMixerWasapi/Private/WasapiDefaultRenderStream.cpp
new file mode 100644
--- /dev/null
+++ b/Engine/Source/Runtime/Windows/AudioMixerWasapi/Private/WasapiDefaultRenderStream.cpp
@@ -0,0 +1,62 @@
+// Copyright Epic Games, Inc. All Rights Reserved.
+
+#include "WasapiDefaultRenderStream.h"
+
+#include "AudioMixerWasapiLog.h"
+#include "WasapiAudioUtils.h"
+
+namespace Audio
+{
+	bool FWasapiDefaultRenderStream::TeardownHardware()
+	{
+		ReadNextBufferDelegate.Unbind();
+
+		return FAudioMixerWasapiRenderStream::TeardownHardware();
+	}
+
+	void FWasapiDefaultRenderStream::DeviceRenderCallback()
+	{
+		SCOPED_NAMED_EVENT(FWasapiDefaultRenderStream_DeviceRenderCallback, FColor::Blue);
+
+		if (bIsInitialized)
+		{
+			uint32 NumFramesPadding = 0;
+			AudioClient->GetCurrentPadding(&NumFramesPadding);
+
+			// NumFramesPerDeviceBuffer is the buffer size WASAPI allocated. It is gauranteed to 
+			// be at least the amount requested. For example, if we request a 1024 frame buffer, WASAPI
+			// might allocate a 1056 frame buffer. The padding is subtracted from the allocated amount
+			// to determine how much space is available currently in the buffer.
+			const int32 NumFramesAvailable = NumFramesPerDeviceBuffer - NumFramesPadding;
+			if (NumFramesAvailable >= static_cast<int32>(RenderStreamParams.NumFrames))
+			{
+				check(!RenderBufferPtr);
+
+				if (SUCCEEDED(RenderClient->GetBuffer(RenderStreamParams.NumFrames, &RenderBufferPtr)))
+				{
+					ReadNextBufferDelegate.ExecuteIfBound();
+
+					HRESULT Result = RenderClient->ReleaseBuffer(RenderStreamParams.NumFrames, 0 /* flags */);
+					if (FAILED(Result))
+					{
+						++CallbackBufferErrors;
+					}
+
+					RenderBufferPtr = nullptr;
+				}
+				else
+				{
+					++CallbackBufferErrors;
+				}
+			}
+		}
+	}
+
+	void FWasapiDefaultRenderStream::SubmitBuffer(const uint8* InBuffer, const SIZE_T InNumFrames)
+	{
+		if (RenderBufferPtr)
+		{
+			FMemory::Memcpy(RenderBufferPtr, InBuffer, InNumFrames * AudioFormat.GetFrameSizeInBytes());
+		}
+	}
+}
diff --git a/Engine/Source/Runtime/Windows/AudioMixerWasapi/Private/WasapiDefaultRenderStream.h b/Engine/Source/Runtime/Windows/AudioMixerWasapi/Private/WasapiDefaultRenderStream.h
new file mode 100644
--- /dev/null
+++ b/Engine/Source/Runtime/Windows/AudioMixerWasapi/Private/WasapiDefaultRenderStream.h
@@ -0,0 +1,38 @@
+// Copyright Epic Games, Inc. All Rights Reserved.
+
+#pragma once
+
+#include "AudioMixerWasapiRenderStream.h"
+
+DECLARE_DELEGATE(FAudioMixerReadNextBufferDelegate);
+
+namespace Audio
+{
+	/**
+	 * FWasapiDefaultRenderStream
+	 * The default channel output stream. This is the first eight channels
+	 * which are reserved for the main outputs of the editor/game.
+	 */
+	class FWasapiDefaultRenderStream : public FAudioMixerWasapiRenderStream
+	{
+	public:
+
+		virtual ~FWasapiDefaultRenderStream() = default;
+
+		//~ Begin FAudioMixerWasapiRenderStream
+		virtual bool TeardownHardware() override;
+		virtual void SubmitBuffer(const uint8* Buffer, const SIZE_T InNumFrames) override;
+		virtual void DeviceRenderCallback() override;
+		//~ End FAudioMixerWasapiRenderStream
+
+		FAudioMixerReadNextBufferDelegate& OnReadNextBuffer() { return ReadNextBufferDelegate; }
+
+	private:
+
+		/** Delegate called each buffer callback to signal the mixer to process the next buffer. */
+		FAudioMixerReadNextBufferDelegate ReadNextBufferDelegate;
+
+		/** Pointer to WASAPI render audio buffer filled in each callback by the mixer in SubmitBuffer(). */
+		uint8* RenderBufferPtr = nullptr;
+	};
+}
