Add num dispatches / num vertices stats to new GPU profiler.
- Removed some more dead codepaths / wrapped old code in "#if (RHI_NEW_GPU_PROFILER == 0)"

Fixed ProfileGPU printing a table with NaNs when there's no work recorded for the given queue.

diff --git a/Engine/Source/Runtime/Apple/MetalRHI/Private/MetalCommands.cpp b/Engine/Source/Runtime/Apple/MetalRHI/Private/MetalCommands.cpp
--- a/Engine/Source/Runtime/Apple/MetalRHI/Private/MetalCommands.cpp
+++ b/Engine/Source/Runtime/Apple/MetalRHI/Private/MetalCommands.cpp
@@ -269,6 +269,8 @@ static IRRuntimeDrawInfo IRRuntimeCalculateDrawInfoForGSEmulation(IRRuntimePrimi
 
 void FMetalRHICommandContext::RHIDispatchComputeShader(uint32 ThreadGroupCountX, uint32 ThreadGroupCountY, uint32 ThreadGroupCountZ)
 {
+	RHI_DISPATCH_CALL_INC();
+
     MTL_SCOPED_AUTORELEASE_POOL;
     
 	ThreadGroupCountX = FMath::Max(ThreadGroupCountX, 1u);
@@ -311,6 +313,8 @@ void FMetalRHICommandContext::RHIDispatchComputeShader(uint32 ThreadGroupCountX,
 
 void FMetalRHICommandContext::RHIDispatchIndirectComputeShader(FRHIBuffer* ArgumentBufferRHI, uint32 ArgumentOffset)
 {
+	RHI_DISPATCH_CALL_INC();
+
     MTL_SCOPED_AUTORELEASE_POOL;
     
 	if (Device.SupportsFeature(EMetalFeaturesIndirectBuffer))
@@ -901,16 +905,13 @@ void FMetalRHICommandContext::RHIDrawPrimitive(uint32 BaseVertexIndex, uint32 Nu
 	SCOPE_CYCLE_COUNTER(STAT_MetalDrawCallTime);
    
 	EPrimitiveType PrimitiveType = StateCache.GetPrimitiveType();
-	
-	NumInstances = FMath::Max(NumInstances,1u);
-	
-	RHI_DRAW_CALL_STATS(PrimitiveType,NumInstances*NumPrimitives);
 
 	// how many verts to render
 	uint32 NumVertices = GetVertexCountForPrimitiveCount(NumPrimitives, PrimitiveType);
-	uint32 VertexCount = GetVertexCountForPrimitiveCount(NumPrimitives, PrimitiveType);
 	
-	NumInstances = FMath::Max(NumInstances,1u);
+	RHI_DRAW_CALL_STATS(PrimitiveType, NumVertices, NumPrimitives, NumInstances);
+
+	NumInstances = FMath::Max(NumInstances, 1u);
 	
 #if PLATFORM_SUPPORTS_GEOMETRY_SHADERS
 	if (IsValidRef(StateCache.GetGraphicsPSO()->GeometryShader))
@@ -994,7 +995,7 @@ void FMetalRHICommandContext::RHIDrawPrimitiveIndirect(FRHIBuffer* ArgumentBuffe
         SCOPE_CYCLE_COUNTER(STAT_MetalDrawCallTime);
         EPrimitiveType PrimitiveType = StateCache.GetPrimitiveType();
         
-        RHI_DRAW_CALL_STATS(PrimitiveType,1);
+        RHI_DRAW_CALL_INC();
         FMetalRHIBuffer* ArgumentBuffer = ResourceCast(ArgumentBufferRHI);
 		
 		FMetalBufferPtr TheBackingBuffer = ArgumentBuffer->GetCurrentBuffer();
@@ -1031,6 +1032,8 @@ void FMetalRHICommandContext::RHIDispatchMeshShader(uint32 ThreadGroupCountX, ui
 {
     MTL_SCOPED_AUTORELEASE_POOL;
 
+	RHI_DRAW_CALL_INC();
+
 #if METAL_USE_METAL_SHADER_CONVERTER
 	checkNoEntry();
 #else
@@ -1042,6 +1045,8 @@ void FMetalRHICommandContext::RHIDispatchIndirectMeshShader(FRHIBuffer* Argument
 {
     MTL_SCOPED_AUTORELEASE_POOL;
 
+	RHI_DRAW_CALL_INC();
+
 #if METAL_USE_METAL_SHADER_CONVERTER
 	uint32 PrimitiveType = StateCache.GetPrimitiveType();
 	PrepareToRender(PrimitiveType);
@@ -1083,7 +1088,7 @@ void FMetalRHICommandContext::RHIDrawIndexedPrimitive(FRHIBuffer* IndexBufferRHI
 	
 	EPrimitiveType PrimitiveType = StateCache.GetPrimitiveType();
 	
-	RHI_DRAW_CALL_STATS(PrimitiveType,FMath::Max(NumInstances,1u)*NumPrimitives);
+	RHI_DRAW_CALL_STATS(PrimitiveType, NumVertices, NumPrimitives, NumInstances);
 
 	FMetalRHIBuffer* IndexBuffer = ResourceCast(IndexBufferRHI);
 	
@@ -1189,7 +1194,7 @@ void FMetalRHICommandContext::RHIDrawIndexedIndirect(FRHIBuffer* IndexBufferRHI,
 		
 		EPrimitiveType PrimitiveType = StateCache.GetPrimitiveType();
 		
-		RHI_DRAW_CALL_STATS(PrimitiveType,1);
+		RHI_DRAW_CALL_INC();
 		FMetalRHIBuffer* IndexBuffer = ResourceCast(IndexBufferRHI);
 		FMetalRHIBuffer* ArgumentsBuffer = ResourceCast(ArgumentsBufferRHI);
 		
@@ -1262,7 +1267,7 @@ void FMetalRHICommandContext::RHIDrawIndexedPrimitiveIndirect(FRHIBuffer* IndexB
 			return;
 		}
 		
-		RHI_DRAW_CALL_STATS(PrimitiveType,1);
+		RHI_DRAW_CALL_INC();
 		FMetalRHIBuffer* IndexBuffer = ResourceCast(IndexBufferRHI);
 		FMetalRHIBuffer* ArgumentsBuffer = ResourceCast(ArgumentBufferRHI);
 		
diff --git a/Engine/Source/Runtime/D3D12RHI/Private/D3D12Buffer.cpp b/Engine/Source/Runtime/D3D12RHI/Private/D3D12Buffer.cpp
--- a/Engine/Source/Runtime/D3D12RHI/Private/D3D12Buffer.cpp
+++ b/Engine/Source/Runtime/D3D12RHI/Private/D3D12Buffer.cpp
@@ -798,5 +798,7 @@ void FD3D12CommandContext::RHICopyBufferRegion(FRHIBuffer* DestBufferRHI, uint64
 
 	ConditionalSplitCommandList();
 
+#if (RHI_NEW_GPU_PROFILER == 0)
 	BufferDevice->RegisterGPUWork(1);
+#endif
 }
diff --git a/Engine/Source/Runtime/D3D12RHI/Private/D3D12Commands.cpp b/Engine/Source/Runtime/D3D12RHI/Private/D3D12Commands.cpp
--- a/Engine/Source/Runtime/D3D12RHI/Private/D3D12Commands.cpp
+++ b/Engine/Source/Runtime/D3D12RHI/Private/D3D12Commands.cpp
@@ -106,10 +106,12 @@ void FD3D12CommandContext::RHISetStreamSource(uint32 StreamIndex, FRHIBuffer* Ve
 
 void FD3D12CommandContext::SetupDispatch(uint32 ThreadGroupCountX, uint32 ThreadGroupCountY, uint32 ThreadGroupCountZ)
 {
+#if (RHI_NEW_GPU_PROFILER == 0)
 	if (IsDefaultContext())
 	{
 		GetParentDevice()->RegisterGPUDispatch(FIntVector(ThreadGroupCountX, ThreadGroupCountY, ThreadGroupCountZ));
 	}
+#endif
 
 	const bool bNewBindlessHeapsSet = FlushPendingDescriptorUpdates();
 
@@ -142,6 +144,8 @@ void FD3D12CommandContext::PostGpuEvent()
 
 void FD3D12CommandContext::RHIDispatchComputeShader(uint32 ThreadGroupCountX, uint32 ThreadGroupCountY, uint32 ThreadGroupCountZ)
 {
+	RHI_DISPATCH_CALL_INC();
+
 	SetupDispatch(ThreadGroupCountX, ThreadGroupCountY, ThreadGroupCountZ);
 
 	GraphicsCommandList()->Dispatch(ThreadGroupCountX, ThreadGroupCountY, ThreadGroupCountZ);
@@ -151,6 +155,8 @@ void FD3D12CommandContext::RHIDispatchComputeShader(uint32 ThreadGroupCountX, ui
 
 void FD3D12CommandContext::RHIDispatchIndirectComputeShader(FRHIBuffer* ArgumentBufferRHI, uint32 ArgumentOffset)
 {
+	RHI_DISPATCH_CALL_INC();
+
 	SetupDispatch(1, 1, 1);
 	
 	FD3D12ResourceLocation& ArgumentBufferLocation = SetupIndirectArgument(ArgumentBufferRHI, D3D12_RESOURCE_STATE_NON_PIXEL_SHADER_RESOURCE);
@@ -1633,11 +1639,11 @@ void FD3D12CommandContext::SetupDispatchDraw(uint32 ThreadGroupCountX, uint32 Th
 
 void FD3D12CommandContext::RHIDrawPrimitive(uint32 BaseVertexIndex, uint32 NumPrimitives, uint32 NumInstances)
 {
-	RHI_DRAW_CALL_STATS(StateCache.GetGraphicsPipelinePrimitiveType(), FMath::Max(NumInstances, 1U) * NumPrimitives);
-
 	uint32 VertexCount = StateCache.GetVertexCount(NumPrimitives);
 	NumInstances = FMath::Max<uint32>(1, NumInstances);
 
+	RHI_DRAW_CALL_STATS(StateCache.GetGraphicsPipelinePrimitiveType(), VertexCount, NumPrimitives, NumInstances);
+
 	SetupDraw(nullptr, NumPrimitives * NumInstances, VertexCount * NumInstances);
 
 	GraphicsCommandList()->DrawInstanced(VertexCount, NumInstances, BaseVertexIndex, 0);
@@ -1699,7 +1705,7 @@ void FD3D12CommandContext::RHIDrawIndexedPrimitive(FRHIBuffer* IndexBufferRHI, i
 		return;
 	}
 
-	RHI_DRAW_CALL_STATS(StateCache.GetGraphicsPipelinePrimitiveType(), FMath::Max(NumInstances, 1U) * NumPrimitives);
+	RHI_DRAW_CALL_STATS(StateCache.GetGraphicsPipelinePrimitiveType(), NumVertices, NumPrimitives, NumInstances);
 
 	NumInstances = FMath::Max<uint32>(1, NumInstances);
 
@@ -1766,6 +1772,8 @@ void FD3D12CommandContext::RHIDrawIndexedPrimitiveIndirect(FRHIBuffer* IndexBuff
 #if PLATFORM_SUPPORTS_MESH_SHADERS
 void FD3D12CommandContext::RHIDispatchMeshShader(uint32 ThreadGroupCountX, uint32 ThreadGroupCountY, uint32 ThreadGroupCountZ)
 {
+	RHI_DRAW_CALL_INC();
+
 	SetupDispatchDraw(ThreadGroupCountX, ThreadGroupCountY, ThreadGroupCountZ);
 
 	GraphicsCommandList6()->DispatchMesh(ThreadGroupCountX, ThreadGroupCountY, ThreadGroupCountZ);
@@ -1906,10 +1914,12 @@ void FD3D12CommandContext::RHIClearMRTImpl(bool* bClearColorArray, int32 NumClea
 		ConditionalSplitCommandList();
 	}
 
+#if (RHI_NEW_GPU_PROFILER == 0)
 	if (IsDefaultContext())
 	{
 		GetParentDevice()->RegisterGPUWork(0);
 	}
+#endif
 
 	DEBUG_EXECUTE_COMMAND_LIST(this);
 }
diff --git a/Engine/Source/Runtime/D3D12RHI/Private/D3D12Device.h b/Engine/Source/Runtime/D3D12RHI/Private/D3D12Device.h
--- a/Engine/Source/Runtime/D3D12RHI/Private/D3D12Device.h
+++ b/Engine/Source/Runtime/D3D12RHI/Private/D3D12Device.h
@@ -183,10 +183,7 @@ public:
 
 	ID3D12Device* GetDevice();
 
-#if RHI_NEW_GPU_PROFILER
-	void RegisterGPUWork(uint32 NumPrimitives = 0, uint32 NumVertices = 0)	{ /*checkNoEntry();*/ } // @todo - new gpu profiler
-	void RegisterGPUDispatch(FIntVector GroupCount)	                        { /*checkNoEntry();*/ } // @todo - new gpu profiler
-#else
+#if (RHI_NEW_GPU_PROFILER == 0)
 	void RegisterGPUWork(uint32 NumPrimitives = 0, uint32 NumVertices = 0)	{ GPUProfilingData.RegisterGPUWork(NumPrimitives, NumVertices); }
 	void RegisterGPUDispatch(FIntVector GroupCount)	                        { GPUProfilingData.RegisterGPUDispatch(GroupCount); }
 
diff --git a/Engine/Source/Runtime/D3D12RHI/Private/D3D12RayTracing.cpp b/Engine/Source/Runtime/D3D12RHI/Private/D3D12RayTracing.cpp
--- a/Engine/Source/Runtime/D3D12RHI/Private/D3D12RayTracing.cpp
+++ b/Engine/Source/Runtime/D3D12RHI/Private/D3D12RayTracing.cpp
@@ -5601,10 +5601,12 @@ static void DispatchRays(FD3D12CommandContext& CommandContext,
 			CommandContext.RayTracingCommandList()->DispatchRays(&DispatchDesc);
 		}
 
+	#if (RHI_NEW_GPU_PROFILER == 0)
 		if (CommandContext.IsDefaultContext())
 		{
 			CommandContext.GetParentDevice()->RegisterGPUWork(1);
 		}
+	#endif
 	}
 
 	// Restore old global descriptor heaps
diff --git a/Engine/Source/Runtime/D3D12RHI/Private/D3D12RenderTarget.cpp b/Engine/Source/Runtime/D3D12RHI/Private/D3D12RenderTarget.cpp
--- a/Engine/Source/Runtime/D3D12RHI/Private/D3D12RenderTarget.cpp
+++ b/Engine/Source/Runtime/D3D12RHI/Private/D3D12RenderTarget.cpp
@@ -184,10 +184,12 @@ void FD3D12CommandContext::ResolveTextureUsingShader(
 		// Clear the destination texture.
 		if (bClearDestTexture)
 		{
+		#if (RHI_NEW_GPU_PROFILER == 0)
 			if (IsDefaultContext())
 			{
 				GetParentDevice()->RegisterGPUWork(0);
 			}
+		#endif
 
 			TransitionResource(DestTextureDSV, D3D12_RESOURCE_STATE_DEPTH_WRITE);
 
@@ -215,10 +217,12 @@ void FD3D12CommandContext::ResolveTextureUsingShader(
 		// Clear the destination texture.
 		if (bClearDestTexture)
 		{
+		#if (RHI_NEW_GPU_PROFILER == 0)
 			if (IsDefaultContext())
 			{
 				GetParentDevice()->RegisterGPUWork(0);
 			}
+		#endif
 
 			TransitionResource(DestTextureRTV, D3D12_RESOURCE_STATE_RENDER_TARGET);
 
@@ -310,10 +314,12 @@ void FD3D12CommandContext::ResolveTexture(UE::RHICore::FResolveTextureInfo Info)
 {
 	uint32 GPUIndex = GetGPUIndex();
 
+#if (RHI_NEW_GPU_PROFILER == 0)
 	if (IsDefaultContext())
 	{
 		GetParentDevice()->RegisterGPUWork();
 	}
+#endif
 
 	FD3D12Texture* SourceTexture         = GetD3D12TextureFromRHITexture(Info.SourceTexture, GPUIndex);
 	FD3D12Resource* SourceResource       = SourceTexture->GetResource();
diff --git a/Engine/Source/Runtime/D3D12RHI/Private/D3D12UAV.cpp b/Engine/Source/Runtime/D3D12RHI/Private/D3D12UAV.cpp
--- a/Engine/Source/Runtime/D3D12RHI/Private/D3D12UAV.cpp
+++ b/Engine/Source/Runtime/D3D12RHI/Private/D3D12UAV.cpp
@@ -396,10 +396,12 @@ void FD3D12CommandContext::ClearUAV(TRHICommandList_RecursiveHazardous<FD3D12Com
 				Context.UpdateResidency(UAV.GetResource());
 				Context.ConditionalSplitCommandList();
 
+			#if (RHI_NEW_GPU_PROFILER == 0)
 				if (Context.IsDefaultContext())
 				{
 					ParentDevice->RegisterGPUWork(1);
 				}
+			#endif
 			});
 		}
 		else
diff --git a/Engine/Source/Runtime/OpenGLDrv/Private/OpenGLCommands.cpp b/Engine/Source/Runtime/OpenGLDrv/Private/OpenGLCommands.cpp
--- a/Engine/Source/Runtime/OpenGLDrv/Private/OpenGLCommands.cpp
+++ b/Engine/Source/Runtime/OpenGLDrv/Private/OpenGLCommands.cpp
@@ -2199,16 +2199,20 @@ void FOpenGLDynamicRHI::RHIDrawPrimitive(uint32 BaseVertexIndex,uint32 NumPrimit
 {
 	VERIFY_GL_SCOPE();
 	SCOPE_CYCLE_COUNTER_DETAILED(STAT_OpenGLDrawPrimitiveTime);
-	RHI_DRAW_CALL_STATS(PrimitiveType,NumPrimitives*NumInstances);
 
 	uint32 VertexCount = GetVertexCountForPrimitiveCount(NumPrimitives, PrimitiveType);
+
+	RHI_DRAW_CALL_STATS(PrimitiveType, VertexCount, NumPrimitives, NumInstances);
+
 	SetupDraw(nullptr, BaseVertexIndex, VertexCount);
 
 	GLenum DrawMode = GL_TRIANGLES;
 	GLsizei NumElements = 0;
 	FindPrimitiveType(PrimitiveType, NumPrimitives, DrawMode, NumElements);
 
+#if (RHI_NEW_GPU_PROFILER == 0)
 	RegisterGPUWork(NumPrimitives * NumInstances, VertexCount * NumInstances);
+#endif
 	if (NumInstances == 1)
 	{
 		SCOPE_CYCLE_COUNTER_DETAILED(STAT_OpenGLDrawPrimitiveDriverTime);
@@ -2238,7 +2242,9 @@ void FOpenGLDynamicRHI::RHIDrawPrimitiveIndirect(FRHIBuffer* ArgumentBufferRHI,
 	VERIFY_GL_SCOPE();
 
 	check(ArgumentBufferRHI);
+#if (RHI_NEW_GPU_PROFILER == 0)
 	RegisterGPUWork(0);
+#endif
 
 	// Zero-stride buffer emulation won't work here, need to use VAB with proper zero strides
 	SetupDraw(nullptr, 0, 1);
@@ -2269,7 +2275,9 @@ void FOpenGLDynamicRHI::RHIDrawIndexedIndirect(FRHIBuffer* IndexBufferRHI, FRHIB
 	}
 
 	VERIFY_GL_SCOPE();
+#if (RHI_NEW_GPU_PROFILER == 0)
 	RegisterGPUWork(1);
+#endif
 
 	FOpenGLBuffer* IndexBuffer = ResourceCast(IndexBufferRHI);
 	SetupDraw(IndexBuffer, 0, 1);
@@ -2303,7 +2311,7 @@ void FOpenGLDynamicRHI::RHIDrawIndexedPrimitive(FRHIBuffer* IndexBufferRHI, int3
 
 	FOpenGLBuffer* IndexBuffer = ResourceCast(IndexBufferRHI);
 
-	RHI_DRAW_CALL_STATS(PrimitiveType,NumPrimitives*NumInstances);
+	RHI_DRAW_CALL_STATS(PrimitiveType, NumVertices, NumPrimitives, NumInstances);
 
 	SetupDraw(IndexBuffer, BaseVertexIndex, NumVertices + StartIndex);
 
@@ -2314,7 +2322,10 @@ void FOpenGLDynamicRHI::RHIDrawIndexedPrimitive(FRHIBuffer* IndexBufferRHI, int3
 	GLenum IndexType = IndexBuffer->GetStride() == sizeof(uint32) ? GL_UNSIGNED_INT : GL_UNSIGNED_SHORT;
 	StartIndex *= IndexBuffer->GetStride() == sizeof(uint32) ? sizeof(uint32) : sizeof(uint16);
 
+#if (RHI_NEW_GPU_PROFILER == 0)
 	RegisterGPUWork(NumPrimitives * NumInstances, NumElements * NumInstances);
+#endif
+
 	if (NumInstances > 1)
 	{
 		SCOPE_CYCLE_COUNTER_DETAILED(STAT_OpenGLDrawPrimitiveDriverTime);
@@ -2349,7 +2360,9 @@ void FOpenGLDynamicRHI::RHIDrawIndexedPrimitiveIndirect(FRHIBuffer* IndexBufferR
 	VERIFY_GL_SCOPE();
 
 	FOpenGLBuffer* IndexBuffer = ResourceCast(IndexBufferRHI);
+#if (RHI_NEW_GPU_PROFILER == 0)
 	RegisterGPUWork(1);
+#endif
 
 	check(ArgumentBufferRHI);
 
@@ -2511,7 +2524,9 @@ void FOpenGLDynamicRHI::RHIClearMRT(const bool* bClearColorArray,int32 NumClearC
 	bool bPrevScissorEnabled = PendingState.bScissorEnabled;
 
 	bool bScissorChanged = false;
+#if (RHI_NEW_GPU_PROFILER == 0)
 	RegisterGPUWork(0);
+#endif
 
 	BindPendingFramebuffer();
 
@@ -2697,8 +2712,13 @@ void FOpenGLDynamicRHI::RHIDispatchComputeShader(uint32 ThreadGroupCountX, uint3
 		return;
 	}
 
+	RHI_DISPATCH_CALL_INC();
+
 	VERIFY_GL_SCOPE();
+
+#if (RHI_NEW_GPU_PROFILER == 0)
 	RegisterGPUDispatch(FIntVector(ThreadGroupCountX, ThreadGroupCountY, ThreadGroupCountZ));
+#endif
 	
 	SetupDispatch();
 
@@ -2709,8 +2729,13 @@ void FOpenGLDynamicRHI::RHIDispatchComputeShader(uint32 ThreadGroupCountX, uint3
 
 void FOpenGLDynamicRHI::RHIDispatchIndirectComputeShader(FRHIBuffer* ArgumentBufferRHI, uint32 ArgumentOffset)
 {
+	RHI_DISPATCH_CALL_INC();
+
 	VERIFY_GL_SCOPE();
+
+#if (RHI_NEW_GPU_PROFILER == 0)
 	RegisterGPUDispatch(FIntVector(1, 1, 1));
+#endif
 
 	SetupDispatch();
 	
diff --git a/Engine/Source/Runtime/OpenGLDrv/Private/OpenGLDrv.h b/Engine/Source/Runtime/OpenGLDrv/Private/OpenGLDrv.h
--- a/Engine/Source/Runtime/OpenGLDrv/Private/OpenGLDrv.h
+++ b/Engine/Source/Runtime/OpenGLDrv/Private/OpenGLDrv.h
@@ -697,19 +697,7 @@ private:
 	/** Underlying platform-specific data */
 	struct FPlatformOpenGLDevice* PlatformDevice = nullptr;
 
-#if RHI_NEW_GPU_PROFILER
-
-	void RegisterGPUWork(uint32 NumPrimitives = 0, uint32 NumVertices = 0)
-	{
-		checkNoEntry(); // @todo - new gpu profiler
-	}
-
-	void RegisterGPUDispatch(FIntVector GroupCount)
-	{
-		checkNoEntry(); // @todo - new gpu profiler
-	}
-
-#else
+#if (RHI_NEW_GPU_PROFILER == 0)
 
 	TOptional<FOpenGLGPUProfiler> GPUProfilingData;
 	friend FOpenGLGPUProfiler;
diff --git a/Engine/Source/Runtime/OpenGLDrv/Private/OpenGLTexture.cpp b/Engine/Source/Runtime/OpenGLDrv/Private/OpenGLTexture.cpp
--- a/Engine/Source/Runtime/OpenGLDrv/Private/OpenGLTexture.cpp
+++ b/Engine/Source/Runtime/OpenGLDrv/Private/OpenGLTexture.cpp
@@ -1387,7 +1387,9 @@ FTextureRHIRef FOpenGLDynamicRHI::RHIAsyncCreateTexture2D(uint32 SizeX, uint32 S
 void FOpenGLDynamicRHI::RHIGenerateMips(FRHITexture* TextureRHI)
 {
 	VERIFY_GL_SCOPE();
+#if (RHI_NEW_GPU_PROFILER == 0)
 	RegisterGPUWork(0);
+#endif
 
 	FOpenGLTexture* Texture = ResourceCast(TextureRHI);
 	// Setup the texture on a disused unit
diff --git a/Engine/Source/Runtime/RHI/Private/GPUProfiler.cpp b/Engine/Source/Runtime/RHI/Private/GPUProfiler.cpp
--- a/Engine/Source/Runtime/RHI/Private/GPUProfiler.cpp
+++ b/Engine/Source/Runtime/RHI/Private/GPUProfiler.cpp
@@ -1463,11 +1463,9 @@ namespace UE::RHI::GPUProfiler
 				FStats& operator += (FEvent::FStats const& Stats)
 				{
 					NumDraws      += Stats.NumDraws;
+					NumDispatches += Stats.NumDispatches;
 					NumPrimitives += Stats.NumPrimitives;
-
-					// @todo add NumDispatches / NumVertices to FEvent::FStats
-					//NumDispatches += Stats.NumDispatches;
-					//NumVertices   += Stats.NumVertices;
+					NumVertices   += Stats.NumVertices;
 
 					return *this;
 				}
@@ -2281,42 +2279,43 @@ namespace UE::RHI::GPUProfiler
 			return bDisplayEvent;
 		};
 
-		Recurse(Recurse, nullptr, State.First, false, 0);
+		// Skip building the table if there was no useful work
+		if (State.First && State.First->Inclusive.BusyCycles > 0)
+		{
+			Recurse(Recurse, nullptr, State.First, false, 0);
+		}
 
-		if (Table.HasRows())
+		FString Final = FString::Printf(
+			TEXT("\n")
+			TEXT("GPU Profile for Frame %d - GPU %d - %s %d\n")
+			TEXT("\n")
+			TEXT("    - %-30s: %.2fms\n")
+			TEXT("    - %-30s: \"%s\"\n")
+			TEXT("    - %-30s: %.2f%%\n")
+			TEXT("    - %-30s: %s\n")
+			TEXT("\n")
+			TEXT("%s")
+			, FrameNumber
+			, Queue.GPU
+			, Queue.GetTypeString()
+			, Queue.Index
+			, TEXT("Frame Time")
+			, State.First->Inclusive.GetBusyMilliseconds()
+			, *IConsoleManager::Get().FindConsoleObjectName(GCVarProfileGPU_Root.AsVariable())
+			, *RootWildcard
+			, *IConsoleManager::Get().FindConsoleObjectName(GCVarProfileGPU_ThresholdPercent.AsVariable())
+			, PercentThreshold
+			, *IConsoleManager::Get().FindConsoleObjectName(GCVarProfileGPU_ShowLeafEvents.AsVariable())
+			, bShowEmptyNodes ? TEXT("true") : TEXT("false")
+			, Table.HasRows() ? *Table.ToString() : TEXT("    No recorded work for this queue.\n")
+		);
+
+		TArray<FString> Lines;
+		Final.ParseIntoArrayLines(Lines, false);
+
+		for (FString const& Line : Lines)
 		{
-			FString Final = FString::Printf(
-				TEXT("\n")
-				TEXT("GPU Profile for Frame %d - GPU %d - %s %d\n")
-				TEXT("\n")
-				TEXT("    - %-30s: %.2fms\n")
-				TEXT("    - %-30s: \"%s\"\n")
-				TEXT("    - %-30s: %.2f%%\n")
-				TEXT("    - %-30s: %s\n")
-				TEXT("\n")
-				TEXT("%s")
-				, FrameNumber
-				, Queue.GPU
-				, Queue.GetTypeString()
-				, Queue.Index
-				, TEXT("Frame Time")
-				, State.First->Inclusive.GetBusyMilliseconds()
-				, *IConsoleManager::Get().FindConsoleObjectName(GCVarProfileGPU_Root.AsVariable())
-				, *RootWildcard
-				, *IConsoleManager::Get().FindConsoleObjectName(GCVarProfileGPU_ThresholdPercent.AsVariable())
-				, PercentThreshold
-				, *IConsoleManager::Get().FindConsoleObjectName(GCVarProfileGPU_ShowLeafEvents.AsVariable())
-				, bShowEmptyNodes ? TEXT("true") : TEXT("false")
-				, *Table.ToString()
-			);
-
-			TArray<FString> Lines;
-			Final.ParseIntoArrayLines(Lines, false);
-
-			for (FString const& Line : Lines)
-			{
-				UE_LOG(LogRHI, Display, TEXT("%s"), *Line);
-			}
+			UE_LOG(LogRHI, Display, TEXT("%s"), *Line);
 		}
 	}
 
diff --git a/Engine/Source/Runtime/RHI/Private/GpuProfilerTrace.cpp b/Engine/Source/Runtime/RHI/Private/GpuProfilerTrace.cpp
--- a/Engine/Source/Runtime/RHI/Private/GpuProfilerTrace.cpp
+++ b/Engine/Source/Runtime/RHI/Private/GpuProfilerTrace.cpp
@@ -320,6 +320,9 @@ UE_TRACE_EVENT_BEGIN(GpuProfiler, EventStats)
 	UE_TRACE_EVENT_FIELD(uint32, QueueId)
 	UE_TRACE_EVENT_FIELD(uint32, NumDraws)
 	UE_TRACE_EVENT_FIELD(uint32, NumPrimitives)
+	// @todo add num dispatches / num vertices stats
+	//UE_TRACE_EVENT_FIELD(uint32, NumDispatches)
+	//UE_TRACE_EVENT_FIELD(uint32, NumVertices)
 UE_TRACE_EVENT_END()
 
 struct FGPUProfilerSink_Trace final : public FEventSink
@@ -448,7 +451,11 @@ struct FGPUProfilerSink_Trace final : public FEventSink
 					UE_TRACE_LOG(GpuProfiler, EventStats, GpuChannel)
 						<< EventStats.QueueId(Queue.Value)
 						<< EventStats.NumDraws(Stats.NumDraws)
-						<< EventStats.NumPrimitives(Stats.NumPrimitives);
+						<< EventStats.NumPrimitives(Stats.NumPrimitives)
+						// @todo add num dispatches / num vertices stats
+						//<< EventStats.NumDispatches(Stats.NumDispatches)
+						//<< EventStats.NumVertices(Stats.NumVertices)
+						;
 					break;
 				}
 
diff --git a/Engine/Source/Runtime/RHI/Public/GPUProfiler.h b/Engine/Source/Runtime/RHI/Public/GPUProfiler.h
--- a/Engine/Source/Runtime/RHI/Public/GPUProfiler.h
+++ b/Engine/Source/Runtime/RHI/Public/GPUProfiler.h
@@ -163,11 +163,16 @@ namespace UE::RHI::GPUProfiler
 		struct FStats
 		{
 			uint32 NumDraws;
+			uint32 NumDispatches;
 			uint32 NumPrimitives;
+			uint32 NumVertices;
 
 			operator bool() const
 			{
-				return NumDraws > 0 || NumPrimitives > 0;
+				return NumDraws > 0
+					|| NumDispatches > 0
+					|| NumPrimitives > 0
+					|| NumVertices > 0;
 			}
 		};
 
diff --git a/Engine/Source/Runtime/RHI/Public/RHIStats.h b/Engine/Source/Runtime/RHI/Public/RHIStats.h
--- a/Engine/Source/Runtime/RHI/Public/RHIStats.h
+++ b/Engine/Source/Runtime/RHI/Public/RHIStats.h
@@ -101,14 +101,24 @@ struct FRHIDrawStatsCategory
 
 #endif
 
-// Macros for use inside RHI context Draw functions.
+// Macros for use inside RHI context Draw/Dispatch functions.
 // Updates the Stats structure on the executing RHI command list
 #if RHI_NEW_GPU_PROFILER
-	#define RHI_DRAW_CALL_INC()             do { StatEvent.NumDraws++; GetExecutingCommandList().Stats_AddDraw(); } while (false)
-	#define RHI_DRAW_CALL_STATS(Type,Prims) do { StatEvent.NumDraws++; StatEvent.NumPrimitives += Prims; GetExecutingCommandList().Stats_AddDrawAndPrimitives(Type, Prims); } while (false)
+	#define RHI_DRAW_CALL_STATS(Type,Verts,Prims,Instances)                \
+	do                                                                     \
+	{                                                                      \
+		StatEvent.NumDraws++;                                              \
+		StatEvent.NumPrimitives += Prims * FMath::Max(1u, Instances);      \
+		StatEvent.NumVertices   += Verts * FMath::Max(1u, Instances);      \
+		GetExecutingCommandList().Stats_AddDrawAndPrimitives(Type, Prims); \
+	} while (false)
+
+	#define RHI_DRAW_CALL_INC()                             do { StatEvent.NumDraws++; GetExecutingCommandList().Stats_AddDraw(); } while (false)
+	#define RHI_DISPATCH_CALL_INC()			                do { StatEvent.NumDispatches++; } while (false)
 #else
-	#define RHI_DRAW_CALL_INC()             do { GetExecutingCommandList().Stats_AddDraw(); } while (false)
-	#define RHI_DRAW_CALL_STATS(Type,Prims) do { GetExecutingCommandList().Stats_AddDrawAndPrimitives(Type, Prims); } while (false)
+	#define RHI_DRAW_CALL_INC()                             do { GetExecutingCommandList().Stats_AddDraw(); } while (false)
+	#define RHI_DRAW_CALL_STATS(Type,Verts,Prims,Instances) do { GetExecutingCommandList().Stats_AddDrawAndPrimitives(Type, Prims); } while (false)
+	#define RHI_DISPATCH_CALL_INC()			                do { } while (false)
 #endif
 
 struct FRHIDrawStats
diff --git a/Engine/Source/Runtime/VulkanRHI/Private/VulkanCommands.cpp b/Engine/Source/Runtime/VulkanRHI/Private/VulkanCommands.cpp
--- a/Engine/Source/Runtime/VulkanRHI/Private/VulkanCommands.cpp
+++ b/Engine/Source/Runtime/VulkanRHI/Private/VulkanCommands.cpp
@@ -276,6 +276,8 @@ void FVulkanCommandListContext::CommitComputeResourceTables()
 
 void FVulkanCommandListContext::RHIDispatchComputeShader(uint32 ThreadGroupCountX, uint32 ThreadGroupCountY, uint32 ThreadGroupCountZ)
 {
+	RHI_DISPATCH_CALL_INC();
+
 #if VULKAN_ENABLE_AGGRESSIVE_STATS
 	SCOPE_CYCLE_COUNTER(STAT_VulkanDispatchCallTime);
 #endif
@@ -299,6 +301,8 @@ void FVulkanCommandListContext::RHIDispatchComputeShader(uint32 ThreadGroupCount
 
 void FVulkanCommandListContext::RHIDispatchIndirectComputeShader(FRHIBuffer* ArgumentBufferRHI, uint32 ArgumentOffset)
 {
+	RHI_DISPATCH_CALL_INC();
+
 	static_assert(sizeof(FRHIDispatchIndirectParameters) == sizeof(VkDispatchIndirectCommand), "Dispatch indirect doesn't match!");
 	FVulkanBuffer* ArgumentBuffer = ResourceCast(ArgumentBufferRHI);
 
@@ -614,13 +618,14 @@ void FVulkanCommandListContext::RHIDrawPrimitive(uint32 BaseVertexIndex, uint32
 #endif
 	NumInstances = FMath::Max(1U, NumInstances);
 
-	RHI_DRAW_CALL_STATS(PendingGfxState->PrimitiveType, NumInstances*NumPrimitives);
-
 	CommitGraphicsResourceTables();
 
 	FVulkanCommandBuffer& CommandBuffer = GetCommandBuffer();
 	PendingGfxState->PrepareForDraw(CommandBuffer);
 	uint32 NumVertices = GetVertexCountForPrimitiveCount(NumPrimitives, PendingGfxState->PrimitiveType);
+
+	RHI_DRAW_CALL_STATS(PendingGfxState->PrimitiveType, NumVertices, NumPrimitives, NumInstances);
+
 	VulkanRHI::vkCmdDraw(CommandBuffer.GetHandle(), NumVertices, NumInstances, BaseVertexIndex, 0);
 
 	if (FVulkanPlatform::RegisterGPUWork() && IsImmediate())
@@ -646,7 +651,6 @@ void FVulkanCommandListContext::RHIDrawPrimitiveIndirect(FRHIBuffer* ArgumentBuf
 
 	FVulkanBuffer* ArgumentBuffer = ResourceCast(ArgumentBufferRHI);
 
-
 	VulkanRHI::vkCmdDrawIndirect(CommandBufferHandle, ArgumentBuffer->GetHandle(), ArgumentBuffer->GetOffset() + ArgumentOffset, 1, sizeof(VkDrawIndirectCommand));
 
 	if (FVulkanPlatform::RegisterGPUWork() && IsImmediate())
@@ -662,7 +666,7 @@ void FVulkanCommandListContext::RHIDrawIndexedPrimitive(FRHIBuffer* IndexBufferR
 	SCOPE_CYCLE_COUNTER(STAT_VulkanDrawCallTime);
 #endif
 	NumInstances = FMath::Max(1U, NumInstances);
-	RHI_DRAW_CALL_STATS(PendingGfxState->PrimitiveType, NumInstances*NumPrimitives);
+	RHI_DRAW_CALL_STATS(PendingGfxState->PrimitiveType, NumVertices, NumPrimitives, NumInstances);
 	checkf(GRHISupportsFirstInstance || FirstInstance == 0, TEXT("FirstInstance must be 0, see GRHISupportsFirstInstance"));
 
 	CommitGraphicsResourceTables();
@@ -738,6 +742,8 @@ void FVulkanCommandListContext::RHIDrawIndexedPrimitiveIndirect(FRHIBuffer* Inde
 #if PLATFORM_SUPPORTS_MESH_SHADERS
 void FVulkanCommandListContext::RHIDispatchMeshShader(uint32 ThreadGroupCountX, uint32 ThreadGroupCountY, uint32 ThreadGroupCountZ)
 {
+	RHI_DRAW_CALL_INC();
+
 	CommitGraphicsResourceTables();
 
 	FVulkanCommandBuffer& CommandBuffer = GetCommandBuffer();
@@ -754,6 +760,8 @@ void FVulkanCommandListContext::RHIDispatchMeshShader(uint32 ThreadGroupCountX,
 
 void FVulkanCommandListContext::RHIDispatchIndirectMeshShader(FRHIBuffer* ArgumentBufferRHI, uint32 ArgumentOffset)
 {
+	RHI_DRAW_CALL_INC();
+
 	CommitGraphicsResourceTables();
 
 	FVulkanCommandBuffer& CommandBuffer = GetCommandBuffer();
diff --git a/Engine/Source/Runtime/Windows/D3D11RHI/Private/D3D11Commands.cpp b/Engine/Source/Runtime/Windows/D3D11RHI/Private/D3D11Commands.cpp
--- a/Engine/Source/Runtime/Windows/D3D11RHI/Private/D3D11Commands.cpp
+++ b/Engine/Source/Runtime/Windows/D3D11RHI/Private/D3D11Commands.cpp
@@ -161,12 +161,16 @@ void FD3D11DynamicRHI::RHISetComputeShader(FRHIComputeShader* ComputeShaderRHI)
 
 void FD3D11DynamicRHI::RHIDispatchComputeShader(uint32 ThreadGroupCountX, uint32 ThreadGroupCountY, uint32 ThreadGroupCountZ) 
 { 
+	RHI_DISPATCH_CALL_INC();
+
 	FRHIComputeShader* ComputeShaderRHI = GetCurrentComputeShader();
 	FD3D11ComputeShader* ComputeShader = ResourceCast(ComputeShaderRHI);
 
 	StateCache.SetComputeShader(ComputeShader->Resource);
 
-	RegisterGPUDispatch(FIntVector(ThreadGroupCountX, ThreadGroupCountY, ThreadGroupCountZ));	
+#if (RHI_NEW_GPU_PROFILER == 0)
+	RegisterGPUDispatch(FIntVector(ThreadGroupCountX, ThreadGroupCountY, ThreadGroupCountZ));
+#endif
 
 	if (ComputeShader->bShaderNeedsGlobalConstantBuffer)
 	{
@@ -181,11 +185,15 @@ void FD3D11DynamicRHI::RHIDispatchComputeShader(uint32 ThreadGroupCountX, uint32
 
 void FD3D11DynamicRHI::RHIDispatchIndirectComputeShader(FRHIBuffer* ArgumentBufferRHI, uint32 ArgumentOffset)
 { 
+	RHI_DISPATCH_CALL_INC();
+
 	FRHIComputeShader* ComputeShaderRHI = GetCurrentComputeShader();
 	FD3D11ComputeShader* ComputeShader = ResourceCast(ComputeShaderRHI);
 	FD3D11Buffer* ArgumentBuffer = ResourceCast(ArgumentBufferRHI);
 
+#if (RHI_NEW_GPU_PROFILER == 0)
 	RegisterGPUDispatch(FIntVector(1, 1, 1));
+#endif
 
 	StateCache.SetComputeShader(ComputeShader->Resource);
 	
@@ -1090,14 +1098,17 @@ void FD3D11DynamicRHI::CommitComputeResourceTables(FD3D11ComputeShader* InComput
 
 void FD3D11DynamicRHI::RHIDrawPrimitive(uint32 BaseVertexIndex,uint32 NumPrimitives,uint32 NumInstances)
 {
-	RHI_DRAW_CALL_STATS(PrimitiveType, FMath::Max(NumInstances, 1U) * NumPrimitives);
-
 	CommitGraphicsResourceTables();
 	CommitNonComputeShaderConstants();
 
-	uint32 VertexCount = GetVertexCountForPrimitiveCount(NumPrimitives,PrimitiveType);
+	uint32 VertexCount = GetVertexCountForPrimitiveCount(NumPrimitives, PrimitiveType);
+
+	RHI_DRAW_CALL_STATS(PrimitiveType, VertexCount, NumPrimitives, NumInstances);
 
+#if (RHI_NEW_GPU_PROFILER == 0)
 	RegisterGPUWork(NumPrimitives * NumInstances, VertexCount * NumInstances);
+#endif
+
 	StateCache.SetPrimitiveTopology(GetD3D11PrimitiveType(PrimitiveType));
 	if(NumInstances > 1)
 	{
@@ -1117,7 +1128,9 @@ void FD3D11DynamicRHI::RHIDrawPrimitiveIndirect(FRHIBuffer* ArgumentBufferRHI, u
 
 	RHI_DRAW_CALL_INC();
 
+#if (RHI_NEW_GPU_PROFILER == 0)
 	RegisterGPUWork(0);
+#endif
 
 	CommitGraphicsResourceTables();
 	CommitNonComputeShaderConstants();
@@ -1135,7 +1148,9 @@ void FD3D11DynamicRHI::RHIDrawIndexedIndirect(FRHIBuffer* IndexBufferRHI, FRHIBu
 
 	RHI_DRAW_CALL_INC();
 
+#if (RHI_NEW_GPU_PROFILER == 0)
 	RegisterGPUWork(1);
+#endif
 
 	CommitGraphicsResourceTables();
 	CommitNonComputeShaderConstants();
@@ -1154,14 +1169,16 @@ void FD3D11DynamicRHI::RHIDrawIndexedIndirect(FRHIBuffer* IndexBufferRHI, FRHIBu
 
 void FD3D11DynamicRHI::RHIDrawIndexedPrimitive(FRHIBuffer* IndexBufferRHI, int32 BaseVertexIndex, uint32 FirstInstance, uint32 NumVertices, uint32 StartIndex, uint32 NumPrimitives, uint32 NumInstances)
 {
-	RHI_DRAW_CALL_STATS(PrimitiveType, FMath::Max(NumInstances, 1U) * NumPrimitives);
+	RHI_DRAW_CALL_STATS(PrimitiveType, NumVertices, NumPrimitives, NumInstances);
 
 	FD3D11Buffer* IndexBuffer = ResourceCast(IndexBufferRHI);
 
 	// called should make sure the input is valid, this avoid hidden bugs
 	ensure(NumPrimitives > 0);
 
+#if (RHI_NEW_GPU_PROFILER == 0)
 	RegisterGPUWork(NumPrimitives * NumInstances, NumVertices * NumInstances);
+#endif
 
 	CommitGraphicsResourceTables();
 	CommitNonComputeShaderConstants();
@@ -1201,7 +1218,9 @@ void FD3D11DynamicRHI::RHIDrawIndexedPrimitiveIndirect(FRHIBuffer* IndexBufferRH
 
 	RHI_DRAW_CALL_INC();
 
+#if (RHI_NEW_GPU_PROFILER == 0)
 	RegisterGPUWork(0);
+#endif
 	
 	CommitGraphicsResourceTables();
 	CommitNonComputeShaderConstants();
@@ -1265,7 +1284,9 @@ void FD3D11DynamicRHI::RHIClearMRTImpl(const bool* bClearColorArray, int32 NumCl
 		Direct3DDeviceIMContext->ClearDepthStencilView(DepthStencilView,ClearFlags,Depth,Stencil);
 	}
 
+#if (RHI_NEW_GPU_PROFILER == 0)
 	RegisterGPUWork(0);
+#endif
 }
 
 // Blocks the CPU until the GPU catches up and goes idle.
diff --git a/Engine/Source/Runtime/Windows/D3D11RHI/Private/D3D11RHI.cpp b/Engine/Source/Runtime/Windows/D3D11RHI/Private/D3D11RHI.cpp
--- a/Engine/Source/Runtime/Windows/D3D11RHI/Private/D3D11RHI.cpp
+++ b/Engine/Source/Runtime/Windows/D3D11RHI/Private/D3D11RHI.cpp
@@ -690,7 +690,9 @@ ID3D11ShaderResourceView* FD3D11DynamicRHI::RHIGetShaderResourceView(FRHITexture
 
 void FD3D11DynamicRHI::RHIRegisterWork(uint32 NumPrimitives)
 {
+#if (RHI_NEW_GPU_PROFILER == 0)
 	RegisterGPUWork(NumPrimitives);
+#endif
 }
 
 void FD3D11DynamicRHI::RHIVerifyResult(ID3D11Device* Device, HRESULT Result, const ANSICHAR* Code, const ANSICHAR* Filename, uint32 Line) const
diff --git a/Engine/Source/Runtime/Windows/D3D11RHI/Private/D3D11RHIPrivate.h b/Engine/Source/Runtime/Windows/D3D11RHI/Private/D3D11RHIPrivate.h
--- a/Engine/Source/Runtime/Windows/D3D11RHI/Private/D3D11RHIPrivate.h
+++ b/Engine/Source/Runtime/Windows/D3D11RHI/Private/D3D11RHIPrivate.h
@@ -871,10 +871,7 @@ protected:
 	bool bRenderDoc = false;
 
 public:
-#if RHI_NEW_GPU_PROFILER
-	void RegisterGPUWork(uint32 NumPrimitives = 0, uint32 NumVertices = 0)	{ /*checkNoEntry();*/ } // @todo - new gpu profiler
-	void RegisterGPUDispatch(FIntVector GroupCount)	                        { /*checkNoEntry();*/ } // @todo - new gpu profiler
-#else
+#if (RHI_NEW_GPU_PROFILER == 0)
 	void RegisterGPUWork(uint32 NumPrimitives = 0, uint32 NumVertices = 0)	{ GPUProfilingData.RegisterGPUWork(NumPrimitives, NumVertices); }
 	void RegisterGPUDispatch(FIntVector GroupCount)	                        { GPUProfilingData.RegisterGPUDispatch(GroupCount); }
 #endif
diff --git a/Engine/Source/Runtime/Windows/D3D11RHI/Private/D3D11RenderTarget.cpp b/Engine/Source/Runtime/Windows/D3D11RHI/Private/D3D11RenderTarget.cpp
--- a/Engine/Source/Runtime/Windows/D3D11RHI/Private/D3D11RenderTarget.cpp
+++ b/Engine/Source/Runtime/Windows/D3D11RHI/Private/D3D11RenderTarget.cpp
@@ -102,7 +102,9 @@ void FD3D11DynamicRHI::ResolveTextureUsingShader(
 				// Clear the destination texture.
 				if (bClearDestTexture)
 				{
+				#if (RHI_NEW_GPU_PROFILER == 0)
 					Context.RegisterGPUWork(0);
+				#endif
 
 					Context.Direct3DDeviceIMContext->ClearDepthStencilView(DestTextureDSV,D3D11_CLEAR_DEPTH | D3D11_CLEAR_STENCIL,0,0);
 				}
@@ -124,7 +126,9 @@ void FD3D11DynamicRHI::ResolveTextureUsingShader(
 				// Clear the destination texture.
 				if (bClearDestTexture)
 				{
+				#if (RHI_NEW_GPU_PROFILER == 0)
 					Context.RegisterGPUWork(0);
+				#endif
 
 					FLinearColor ClearColor(0,0,0,0);
 					Context.Direct3DDeviceIMContext->ClearRenderTargetView(DestTextureRTV,(float*)&ClearColor);
@@ -874,7 +878,9 @@ void FD3D11DynamicRHI::RHIEndRenderPass()
 
 void FD3D11DynamicRHI::ResolveTexture(UE::RHICore::FResolveTextureInfo Info)
 {
+#if (RHI_NEW_GPU_PROFILER == 0)
 	RegisterGPUWork();
+#endif
 
 	FD3D11Texture* SourceTexture      = ResourceCast(Info.SourceTexture);
 	const FRHITextureDesc& SourceDesc = SourceTexture->GetDesc();
diff --git a/Engine/Source/Runtime/Windows/D3D11RHI/Private/D3D11Texture.cpp b/Engine/Source/Runtime/Windows/D3D11RHI/Private/D3D11Texture.cpp
--- a/Engine/Source/Runtime/Windows/D3D11RHI/Private/D3D11Texture.cpp
+++ b/Engine/Source/Runtime/Windows/D3D11RHI/Private/D3D11Texture.cpp
@@ -1660,7 +1660,9 @@ void FD3D11DynamicRHI::RHICopyTexture(FRHITexture* SourceTextureRHI, FRHITexture
 
 	check(SourceTexture && DestTexture);
 
+#if (RHI_NEW_GPU_PROFILER == 0)
 	RegisterGPUWork();
+#endif
 
 	const FRHITextureDesc& SourceDesc = SourceTextureRHI->GetDesc();
 	const FRHITextureDesc& DestDesc = DestTextureRHI->GetDesc();
@@ -1732,7 +1734,9 @@ void FD3D11DynamicRHI::RHICopyBufferRegion(FRHIBuffer* DstBuffer, uint64 DstOffs
 	check(DstBufferD3D11 && SrcBufferD3D11);
 	check(DstOffset + NumBytes <= DstBuffer->GetSize() && SrcOffset + NumBytes <= SrcBuffer->GetSize());
 
+#if (RHI_NEW_GPU_PROFILER == 0)
 	RegisterGPUWork();
+#endif
 
 	D3D11_BOX SrcBox;
 	SrcBox.left = SrcOffset;
diff --git a/Engine/Source/Runtime/Windows/D3D11RHI/Private/D3D11UAV.cpp b/Engine/Source/Runtime/Windows/D3D11RHI/Private/D3D11UAV.cpp
--- a/Engine/Source/Runtime/Windows/D3D11RHI/Private/D3D11UAV.cpp
+++ b/Engine/Source/Runtime/Windows/D3D11RHI/Private/D3D11UAV.cpp
@@ -205,7 +205,9 @@ void FD3D11DynamicRHI::ClearUAV(TRHICommandList_RecursiveHazardous<FD3D11Dynamic
 			RHICmdList.RunOnContext([UnorderedAccessView, ClearValues](auto& Context)
 			{
 				Context.Direct3DDeviceIMContext->ClearUnorderedAccessViewUint(UnorderedAccessView->View, *reinterpret_cast<const UINT(*)[4]>(ClearValues));
+			#if (RHI_NEW_GPU_PROFILER == 0)
 				Context.RegisterGPUWork(1);
+			#endif
 			});
 		}
 		else
