Merging RHI CreateBuffer refactor to Main.

diff --git a/Engine/Source/Runtime/Windows/D3D11RHI/Private/D3D11Buffer.cpp b/Engine/Source/Runtime/Windows/D3D11RHI/Private/D3D11Buffer.cpp
--- a/Engine/Source/Runtime/Windows/D3D11RHI/Private/D3D11Buffer.cpp
+++ b/Engine/Source/Runtime/Windows/D3D11RHI/Private/D3D11Buffer.cpp
@@ -5,6 +5,7 @@
 =============================================================================*/
 
 #include "D3D11RHIPrivate.h"
+#include "RHICoreBufferInitializer.h"
 
 TAutoConsoleVariable<int32> GCVarUseSharedKeyedMutex(
 	TEXT("r.D3D11.UseSharedKeyMutex"),
@@ -32,62 +33,53 @@ void FD3D11Buffer::ReleaseOwnership()
 	Resource = nullptr;
 }
 
-FBufferRHIRef FD3D11DynamicRHI::RHICreateBuffer(FRHICommandListBase& RHICmdList, FRHIBufferDesc const& BufferDesc, ERHIAccess ResourceState, FRHIResourceCreateInfo& CreateInfo)
+static D3D11_BUFFER_DESC D3D11DescFromBufferDesc(const FRHIBufferDesc& BufferDesc)
 {
-	if (BufferDesc.IsNull())
-	{
-		return new FD3D11Buffer(nullptr, BufferDesc);
-	}
-
-	// Explicitly check that the size is nonzero before allowing CreateBuffer to opaquely fail.
-	checkf(BufferDesc.Size > 0, TEXT("Attempt to create buffer '%s' with size 0."), CreateInfo.DebugName ? CreateInfo.DebugName : TEXT("(null)"));
-
-	// Describe the buffer.
 	D3D11_BUFFER_DESC Desc{};
 	Desc.ByteWidth = BufferDesc.Size;
 
-	if (EnumHasAnyFlags(BufferDesc.Usage, BUF_AnyDynamic))
+	if (EnumHasAnyFlags(BufferDesc.Usage, EBufferUsageFlags::AnyDynamic))
 	{
 		Desc.Usage = D3D11_USAGE_DYNAMIC;
 		Desc.CPUAccessFlags = D3D11_CPU_ACCESS_WRITE;
 	}
 
-	if (EnumHasAnyFlags(BufferDesc.Usage, BUF_VertexBuffer))
+	if (EnumHasAnyFlags(BufferDesc.Usage, EBufferUsageFlags::VertexBuffer))
 	{
 		Desc.BindFlags |= D3D11_BIND_VERTEX_BUFFER;
 	}
 
-	if (EnumHasAnyFlags(BufferDesc.Usage, BUF_IndexBuffer))
+	if (EnumHasAnyFlags(BufferDesc.Usage, EBufferUsageFlags::IndexBuffer))
 	{
 		Desc.BindFlags |= D3D11_BIND_INDEX_BUFFER;
 	}
 
-	if (EnumHasAnyFlags(BufferDesc.Usage, BUF_ByteAddressBuffer))
+	if (EnumHasAnyFlags(BufferDesc.Usage, EBufferUsageFlags::ByteAddressBuffer))
 	{
 		Desc.MiscFlags |= D3D11_RESOURCE_MISC_BUFFER_ALLOW_RAW_VIEWS;
 	}
-	else if (EnumHasAnyFlags(BufferDesc.Usage, BUF_StructuredBuffer))
+	else if (EnumHasAnyFlags(BufferDesc.Usage, EBufferUsageFlags::StructuredBuffer))
 	{
 		Desc.StructureByteStride = BufferDesc.Stride;
 		Desc.MiscFlags |= D3D11_RESOURCE_MISC_BUFFER_STRUCTURED;
 	}
 
-	if (EnumHasAnyFlags(BufferDesc.Usage, BUF_ShaderResource))
+	if (EnumHasAnyFlags(BufferDesc.Usage, EBufferUsageFlags::ShaderResource))
 	{
 		Desc.BindFlags |= D3D11_BIND_SHADER_RESOURCE;
 	}
 
-	if (EnumHasAnyFlags(BufferDesc.Usage, BUF_UnorderedAccess))
+	if (EnumHasAnyFlags(BufferDesc.Usage, EBufferUsageFlags::UnorderedAccess))
 	{
 		Desc.BindFlags |= D3D11_BIND_UNORDERED_ACCESS;
 	}
 
-	if (EnumHasAnyFlags(BufferDesc.Usage, BUF_DrawIndirect))
+	if (EnumHasAnyFlags(BufferDesc.Usage, EBufferUsageFlags::DrawIndirect))
 	{
 		Desc.MiscFlags |= D3D11_RESOURCE_MISC_DRAWINDIRECT_ARGS;
 	}
 
-	if (EnumHasAnyFlags(BufferDesc.Usage, BUF_Shared))
+	if (EnumHasAnyFlags(BufferDesc.Usage, EBufferUsageFlags::Shared))
 	{
 		if (GCVarUseSharedKeyedMutex->GetInt() != 0)
 		{
@@ -99,15 +91,31 @@ FBufferRHIRef FD3D11DynamicRHI::RHICreateBuffer(FRHICommandListBase& RHICmdList,
 		}
 	}
 
+	return Desc;
+}
+
+FD3D11Buffer* FD3D11DynamicRHI::BeginCreateBufferInternal(const FRHIBufferCreateDesc& CreateDesc)
+{
+	// Explicitly check that the size is nonzero before allowing CreateBuffer to opaquely fail.
+	checkf(CreateDesc.Size > 0, TEXT("Attempt to create buffer '%s' with size 0."), CreateDesc.DebugName ? CreateDesc.DebugName : TEXT("(null)"));
+
+	return new FD3D11Buffer(nullptr, CreateDesc);
+}
+
+void FD3D11DynamicRHI::FinalizeCreateBufferInternal(FD3D11Buffer* Buffer, TConstArrayView<uint8> InitialData)
+{
+	const FRHIBufferDesc& BufferDesc = Buffer->GetDesc();
+	const D3D11_BUFFER_DESC Desc = D3D11DescFromBufferDesc(BufferDesc);
+
 	// If a resource array was provided for the resource, create the resource pre-populated
-	D3D11_SUBRESOURCE_DATA InitData;
-	D3D11_SUBRESOURCE_DATA* pInitData = NULL;
-	if(CreateInfo.ResourceArray)
+	D3D11_SUBRESOURCE_DATA InitData{};
+	D3D11_SUBRESOURCE_DATA* pInitData = nullptr;
+	if (InitialData.Num())
 	{
-		check(BufferDesc.Size == CreateInfo.ResourceArray->GetResourceDataSize());
-		InitData.pSysMem = CreateInfo.ResourceArray->GetResourceData();
-		InitData.SysMemPitch = BufferDesc.Size;
-		InitData.SysMemSlicePitch = 0;
+		check(BufferDesc.Size == InitialData.Num());
+		InitData.pSysMem = InitialData.GetData();
+		InitData.SysMemPitch = InitialData.Num();
+
 		pInitData = &InitData;
 	}
 
@@ -116,37 +124,78 @@ FBufferRHIRef FD3D11DynamicRHI::RHICreateBuffer(FRHICommandListBase& RHICmdList,
 		HRESULT hr = Direct3DDevice->CreateBuffer(&Desc, pInitData, BufferResource.GetInitReference());
 		if (FAILED(hr))
 		{
+			const FName BufferName = Buffer->GetName();
+
 			UE_LOG(LogD3D11RHI, Error, TEXT("Failed to create buffer '%s' with ByteWidth=%u, Usage=%d, BindFlags=0x%x, CPUAccessFlags=0x%x, MiscFlags=0x%x, StructureByteStride=%u, InitData=0x%p"),
-				CreateInfo.DebugName ? CreateInfo.DebugName : TEXT(""), Desc.ByteWidth, Desc.Usage, Desc.BindFlags, Desc.CPUAccessFlags, Desc.MiscFlags, Desc.StructureByteStride, pInitData);
+				*BufferName.ToString(), Desc.ByteWidth, Desc.Usage, Desc.BindFlags, Desc.CPUAccessFlags, Desc.MiscFlags, Desc.StructureByteStride, pInitData);
 			VerifyD3D11Result(hr, "CreateBuffer", __FILE__, __LINE__, Direct3DDevice);
 		}
 	}
 
-	if (CreateInfo.DebugName)
+	if (FName BufferName = Buffer->GetName(); BufferName != NAME_None)
 	{
-		BufferResource->SetPrivateData(WKPDID_D3DDebugObjectName, FCString::Strlen(CreateInfo.DebugName) + 1, TCHAR_TO_ANSI(CreateInfo.DebugName));
+		const FAnsiString BufferNameString(BufferName.ToString());
+		BufferResource->SetPrivateData(WKPDID_D3DDebugObjectName, BufferNameString.Len() + 1, *BufferNameString);
 	}
 
-	if (CreateInfo.DebugName)
+	Buffer->Resource = BufferResource;
+
+	D3D11BufferStats::UpdateBufferStats(*Buffer, true);
+}
+
+FD3D11Buffer* FD3D11DynamicRHI::CreateBufferInternal(const FRHIBufferCreateDesc& CreateDesc, TConstArrayView<uint8> InitialData)
+{
+	FD3D11Buffer* Buffer = BeginCreateBufferInternal(CreateDesc);
+	FinalizeCreateBufferInternal(Buffer, InitialData);
+	return Buffer;
+}
+
+FRHIBufferInitializer FD3D11DynamicRHI::RHICreateBufferInitializer(FRHICommandListBase& RHICmdList, const FRHIBufferCreateDesc& CreateDesc)
+{
+	if (CreateDesc.IsNull() || CreateDesc.InitAction == ERHIBufferInitAction::Default)
 	{
-		BufferResource->SetPrivateData(WKPDID_D3DDebugObjectName, FCString::Strlen(CreateInfo.DebugName) + 1, TCHAR_TO_ANSI(CreateInfo.DebugName));
+		return UE::RHICore::FDefaultBufferInitializer(RHICmdList, CreateBufferInternal(CreateDesc, {}));
 	}
 
-	if (CreateInfo.ResourceArray)
+	if (CreateDesc.InitAction == ERHIBufferInitAction::ResourceArray)
 	{
+		check(CreateDesc.InitialData);
+
+		FD3D11Buffer* Buffer = CreateBufferInternal(CreateDesc, CreateDesc.InitialData->GetResourceDataView<uint8>());
+
 		// Discard the resource array's contents.
-		CreateInfo.ResourceArray->Discard();
+		CreateDesc.InitialData->Discard();
+
+		return UE::RHICore::FDefaultBufferInitializer(RHICmdList, Buffer);
 	}
 
-	FD3D11Buffer* NewBuffer = new FD3D11Buffer(BufferResource, BufferDesc);
-	if (CreateInfo.DebugName)
+	if (CreateDesc.InitAction == ERHIBufferInitAction::Zeroed)
 	{
-		NewBuffer->SetName(CreateInfo.DebugName);
+		// Allocate zeroed memory and upload that on buffer creation. The contents of the buffer are "undefined" if initial data is not passed in, so we actually have to do something here.
+		void* UploadMemory = FMemory::MallocZeroed(CreateDesc.Size, 16);
+		FD3D11Buffer* Buffer = CreateBufferInternal(CreateDesc, TConstArrayView<uint8>(reinterpret_cast<uint8*>(UploadMemory), CreateDesc.Size));
+		FMemory::Free(UploadMemory);
+
+		return UE::RHICore::FDefaultBufferInitializer(RHICmdList, Buffer);
 	}
 
-	D3D11BufferStats::UpdateBufferStats(*NewBuffer, true);
+	if (CreateDesc.InitAction == ERHIBufferInitAction::Initializer)
+	{
+		FD3D11Buffer* Buffer = BeginCreateBufferInternal(CreateDesc);
+
+		// Allocate ad-hoc memory for writing.
+		void* UploadMemory = FMemory::Malloc(CreateDesc.Size, 16);
+
+		return UE::RHICore::FCustomBufferInitializer(RHICmdList, Buffer, UploadMemory, CreateDesc.Size,
+			[this, Buffer = TRefCountPtr<FD3D11Buffer>(Buffer), UploadMemory = UE::RHICore::FScopedMemory(UploadMemory)](FRHICommandListBase& RHICmdList) mutable
+			{
+				FinalizeCreateBufferInternal(Buffer, TConstArrayView<uint8>(reinterpret_cast<const uint8*>(UploadMemory.Pointer), Buffer->GetDesc().Size));
+				return TRefCountPtr<FRHIBuffer>(MoveTemp(Buffer));
+			}
+		);
+	}
 
-	return NewBuffer;
+	return UE::RHICore::HandleUnknownBufferInitializerInitAction(RHICmdList, CreateDesc);
 }
 
 void* FD3D11DynamicRHI::LockBuffer_BottomOfPipe(FRHICommandListBase& RHICmdList, FRHIBuffer* BufferRHI, uint32 Offset, uint32 Size, EResourceLockMode LockMode)
diff --git a/Engine/Source/Runtime/Windows/D3D11RHI/Private/D3D11RHIPrivate.h b/Engine/Source/Runtime/Windows/D3D11RHI/Private/D3D11RHIPrivate.h
--- a/Engine/Source/Runtime/Windows/D3D11RHI/Private/D3D11RHIPrivate.h
+++ b/Engine/Source/Runtime/Windows/D3D11RHI/Private/D3D11RHIPrivate.h
@@ -448,7 +448,13 @@ public:
 	virtual FBoundShaderStateRHIRef RHICreateBoundShaderState(FRHIVertexDeclaration* VertexDeclaration, FRHIVertexShader* VertexShader, FRHIPixelShader* PixelShader, FRHIGeometryShader* GeometryShader) final override;
 	virtual FUniformBufferRHIRef RHICreateUniformBuffer(const void* Contents, const FRHIUniformBufferLayout* Layout, EUniformBufferUsage Usage, EUniformBufferValidation Validation) final override;
 	virtual void RHIUpdateUniformBuffer(FRHICommandListBase& RHICmdList, FRHIUniformBuffer* UniformBufferRHI, const void* Contents) final override;
-	virtual FBufferRHIRef RHICreateBuffer(FRHICommandListBase& RHICmdList, FRHIBufferDesc const& Desc, ERHIAccess ResourceState, FRHIResourceCreateInfo& CreateInfo) final override;
+	
+	[[nodiscard]] FD3D11Buffer* BeginCreateBufferInternal(const FRHIBufferCreateDesc& CreateDesc);
+	[[nodiscard]] void FinalizeCreateBufferInternal(FD3D11Buffer* Buffer, TConstArrayView<uint8> InitialData);
+	[[nodiscard]] FD3D11Buffer* CreateBufferInternal(const FRHIBufferCreateDesc& CreateDesc, TConstArrayView<uint8> InitialData);
+
+	[[nodiscard]] virtual FRHIBufferInitializer RHICreateBufferInitializer(FRHICommandListBase& RHICmdList, const FRHIBufferCreateDesc& CreateDesc) final override;
+
 	virtual void RHIReplaceResources(FRHICommandListBase& RHICmdList, TArray<FRHIResourceReplaceInfo>&& ReplaceInfos) final override;
 	virtual void* LockBuffer_BottomOfPipe(FRHICommandListBase& RHICmdList, FRHIBuffer* Buffer, uint32 Offset, uint32 Size, EResourceLockMode LockMode) final override;
 	virtual void UnlockBuffer_BottomOfPipe(FRHICommandListBase& RHICmdList, FRHIBuffer* Buffer) final override;
diff --git a/Engine/Source/Runtime/Windows/D3D11RHI/Public/D3D11Resources.h b/Engine/Source/Runtime/Windows/D3D11RHI/Public/D3D11Resources.h
--- a/Engine/Source/Runtime/Windows/D3D11RHI/Public/D3D11Resources.h
+++ b/Engine/Source/Runtime/Windows/D3D11RHI/Public/D3D11Resources.h
@@ -409,8 +409,8 @@ public:
 
 	TRefCountPtr<ID3D11Buffer> Resource;
 
-	FD3D11Buffer(ID3D11Buffer* InResource, FRHIBufferDesc const& InDesc)
-		: FRHIBuffer(InDesc)
+	FD3D11Buffer(ID3D11Buffer* InResource, const FRHIBufferCreateDesc& InCreateDesc)
+		: FRHIBuffer(InCreateDesc)
 		, Resource(InResource)
 	{}
 
