Fixing Null buffer creation in D3D11.

diff --git a/Engine/Source/Runtime/Windows/D3D11RHI/Private/D3D11Buffer.cpp b/Engine/Source/Runtime/Windows/D3D11RHI/Private/D3D11Buffer.cpp
--- a/Engine/Source/Runtime/Windows/D3D11RHI/Private/D3D11Buffer.cpp
+++ b/Engine/Source/Runtime/Windows/D3D11RHI/Private/D3D11Buffer.cpp
@@ -97,7 +97,7 @@ static D3D11_BUFFER_DESC D3D11DescFromBufferDesc(const FRHIBufferDesc& BufferDes
 FD3D11Buffer* FD3D11DynamicRHI::BeginCreateBufferInternal(const FRHIBufferCreateDesc& CreateDesc)
 {
 	// Explicitly check that the size is nonzero before allowing CreateBuffer to opaquely fail.
-	checkf(CreateDesc.Size > 0, TEXT("Attempt to create buffer '%s' with size 0."), CreateDesc.DebugName ? CreateDesc.DebugName : TEXT("(null)"));
+	checkf(CreateDesc.Size > 0 || CreateDesc.IsNull(), TEXT("Attempt to create buffer '%s' with size 0."), CreateDesc.DebugName ? CreateDesc.DebugName : TEXT("(null)"));
 
 	return new FD3D11Buffer(nullptr, CreateDesc);
 }
@@ -107,6 +107,8 @@ void FD3D11DynamicRHI::FinalizeCreateBufferInternal(FD3D11Buffer* Buffer, TConst
 	const FRHIBufferDesc& BufferDesc = Buffer->GetDesc();
 	const D3D11_BUFFER_DESC Desc = D3D11DescFromBufferDesc(BufferDesc);
 
+	check(!BufferDesc.IsNull());
+
 	// If a resource array was provided for the resource, create the resource pre-populated
 	D3D11_SUBRESOURCE_DATA InitData{};
 	D3D11_SUBRESOURCE_DATA* pInitData = nullptr;
@@ -152,7 +154,13 @@ FD3D11Buffer* FD3D11DynamicRHI::CreateBufferInternal(const FRHIBufferCreateDesc&
 
 FRHIBufferInitializer FD3D11DynamicRHI::RHICreateBufferInitializer(FRHICommandListBase& RHICmdList, const FRHIBufferCreateDesc& CreateDesc)
 {
-	if (CreateDesc.IsNull() || CreateDesc.InitAction == ERHIBufferInitAction::Default)
+	if (CreateDesc.IsNull())
+	{
+		FD3D11Buffer* Buffer = BeginCreateBufferInternal(CreateDesc);
+		return UE::RHICore::FDefaultBufferInitializer(RHICmdList, Buffer);
+	}
+
+	if (CreateDesc.InitAction == ERHIBufferInitAction::Default)
 	{
 		return UE::RHICore::FDefaultBufferInitializer(RHICmdList, CreateBufferInternal(CreateDesc, {}));
 	}
