D3D11 Texture Creation cleanup before RHICreateTexture refactor
- Merging common code between CreateD3D11Texture2D and CreateD3D11Texture3D
- Adding SafeCreateTexture3D to mirror SafeCreateTexture2D
- Texture creation split into 2 phases: RHI object creation then D3D11Resource creation. This is important for the RHICreateTexture refactor.

diff --git a/Engine/Source/Runtime/RHI/Public/RHIUtilities.h b/Engine/Source/Runtime/RHI/Public/RHIUtilities.h
--- a/Engine/Source/Runtime/RHI/Public/RHIUtilities.h
+++ b/Engine/Source/Runtime/RHI/Public/RHIUtilities.h
@@ -14,6 +14,14 @@
 #include "RHICommandList.h"
 #include "Tasks/Task.h"
 
+namespace UE::RHI
+{
+	static uint32 CalculateMipBlockCount(uint32 Extent, uint32 MipIndex, uint32 BlockSize)
+	{
+		return FMath::Max<uint32>(1, FMath::DivideAndRoundUp(Extent >> MipIndex, BlockSize));
+	}
+}
+
 static inline EPixelFormat GetBlockCompressedFormatUAVAliasFormat(EPixelFormat Format)
 {
 	switch (Format)
diff --git a/Engine/Source/Runtime/RHICore/Internal/DXGIUtilities.h b/Engine/Source/Runtime/RHICore/Internal/DXGIUtilities.h
--- a/Engine/Source/Runtime/RHICore/Internal/DXGIUtilities.h
+++ b/Engine/Source/Runtime/RHICore/Internal/DXGIUtilities.h
@@ -142,6 +142,11 @@ namespace UE::DXGIUtilities
 		return InFormat;
 	}
 
+	inline DXGI_FORMAT GetPlatformTextureResourceFormat(EPixelFormat InPixelFormat, ETextureCreateFlags InFlags)
+	{
+		return GetPlatformTextureResourceFormat((DXGI_FORMAT)GPixelFormats[InPixelFormat].PlatformFormat, InFlags);
+	}
+
 	/** Find an appropriate DXGI format unordered access of the raw format. */
 	inline DXGI_FORMAT FindUnorderedAccessFormat(DXGI_FORMAT InFormat)
 	{
diff --git a/Engine/Source/Runtime/Windows/D3D11RHI/Private/D3D11RHIPrivate.h b/Engine/Source/Runtime/Windows/D3D11RHI/Private/D3D11RHIPrivate.h
--- a/Engine/Source/Runtime/Windows/D3D11RHI/Private/D3D11RHIPrivate.h
+++ b/Engine/Source/Runtime/Windows/D3D11RHI/Private/D3D11RHIPrivate.h
@@ -461,6 +461,12 @@ public:
 	virtual FRHICalcTextureSizeResult RHICalcTexturePlatformSize(FRHITextureDesc const& Desc, uint32 FirstMipIndex) final override;
 	virtual void RHIGetTextureMemoryStats(FTextureMemoryStats& OutStats) final override;
 	virtual bool RHIGetTextureMemoryVisualizeData(FColor* TextureData,int32 SizeX,int32 SizeY,int32 Pitch,int32 PixelSize) final override;
+
+	[[nodiscard]] FD3D11Texture* BeginCreateTextureInternal(const FRHITextureCreateDesc& CreateDesc);
+	[[nodiscard]] void  FinalizeCreateTexture2DInternal(FD3D11Texture* Texture, TConstArrayView<D3D11_SUBRESOURCE_DATA> InitialData);
+	[[nodiscard]] void  FinalizeCreateTexture3DInternal(FD3D11Texture* Texture, TConstArrayView<D3D11_SUBRESOURCE_DATA> InitialData);
+	[[nodiscard]] FD3D11Texture* CreateTextureInternal(const FRHITextureCreateDesc& CreateDesc);
+
 	virtual FTextureRHIRef RHICreateTexture(FRHICommandListBase& RHICmdList, const FRHITextureCreateDesc& CreateDesc) final override;
 	virtual FTextureRHIRef RHIAsyncCreateTexture2D(uint32 SizeX, uint32 SizeY, uint8 Format, uint32 NumMips, ETextureCreateFlags Flags, ERHIAccess InResourceState, void** InitialMipData, uint32 NumInitialMips, const TCHAR* DebugName, FGraphEventRef& OutCompletionEvent) final override;
 	virtual uint32 RHIComputeMemorySize(FRHITexture* TextureRHI) final override;
@@ -897,9 +903,6 @@ protected:
 
 	FD3D11Adapter Adapter;
 
-	FD3D11Texture* CreateD3D11Texture2D(FRHITextureCreateDesc const& CreateDesc, TConstArrayView<D3D11_SUBRESOURCE_DATA> InitialData = {});
-	FD3D11Texture* CreateD3D11Texture3D(FRHITextureCreateDesc const& CreateDesc);
-
 	FD3D11Texture* CreateTextureFromResource(bool bTextureArray, bool bCubeTexture, EPixelFormat Format, ETextureCreateFlags TexCreateFlags, const FClearValueBinding& ClearValueBinding, ID3D11Texture2D* TextureResource);
 
 	/** Initializes the constant buffers.  Called once at RHI initialization time. */
diff --git a/Engine/Source/Runtime/Windows/D3D11RHI/Private/D3D11Texture.cpp b/Engine/Source/Runtime/Windows/D3D11RHI/Private/D3D11Texture.cpp
--- a/Engine/Source/Runtime/Windows/D3D11RHI/Private/D3D11Texture.cpp
+++ b/Engine/Source/Runtime/Windows/D3D11RHI/Private/D3D11Texture.cpp
@@ -19,6 +19,7 @@
 #include "HAL/LowLevelMemTracker.h"
 #include "ProfilingDebugging/MemoryTrace.h"
 #include "RHICoreStats.h"
+#include "RHIUtilities.h"
 
 int64 FD3D11GlobalStats::GDedicatedVideoMemory = 0;
 int64 FD3D11GlobalStats::GDedicatedSystemMemory = 0;
@@ -94,9 +95,6 @@ FDynamicRHI::FRHICalcTextureSizeResult FD3D11DynamicRHI::RHICalcTexturePlatformS
 	return Result;
 }
 
-/**
- * Retrieves texture memory stats. 
- */
 void FD3D11DynamicRHI::RHIGetTextureMemoryStats(FTextureMemoryStats& OutStats)
 {
 	UE::RHICore::FillBaselineTextureMemoryStats(OutStats);
@@ -109,17 +107,6 @@ void FD3D11DynamicRHI::RHIGetTextureMemoryStats(FTextureMemoryStats& OutStats)
 	OutStats.LargestContiguousAllocation = OutStats.StreamingMemorySize;
 }
 
-/**
- * Fills a texture with to visualize the texture pool memory.
- *
- * @param	TextureData		Start address
- * @param	SizeX			Number of pixels along X
- * @param	SizeY			Number of pixels along Y
- * @param	Pitch			Number of bytes between each row
- * @param	PixelSize		Number of bytes each pixel represents
- *
- * @return true if successful, false otherwise
- */
 bool FD3D11DynamicRHI::RHIGetTextureMemoryVisualizeData( FColor* /*TextureData*/, int32 /*SizeX*/, int32 /*SizeY*/, int32 /*Pitch*/, int32 /*PixelSize*/ )
 {
 	// currently only implemented for console (Note: Keep this function for further extension. Talk to NiklasS for more info.)
@@ -201,66 +188,119 @@ static void SafeCreateTexture2D(ID3D11Device* Direct3DDevice, int32 UEFormat, co
 #endif // #if GUARDED_TEXTURE_CREATES
 }
 
-FD3D11Texture* FD3D11DynamicRHI::CreateD3D11Texture2D(FRHITextureCreateDesc const& CreateDesc, TConstArrayView<D3D11_SUBRESOURCE_DATA> InitialData)
+static void SafeCreateTexture3D(FD3D11Device* Direct3DDevice, int32 UEFormat, const D3D11_TEXTURE3D_DESC* TextureDesc, const D3D11_SUBRESOURCE_DATA* SubResourceData, ID3D11Texture3D** OutTexture2D, const TCHAR* DebugName)
 {
-	check(!CreateDesc.IsTexture3D());
-
-	const bool                bTextureArray = CreateDesc.IsTextureArray();
-	const bool                bCubeTexture  = CreateDesc.IsTextureCube();
-	const uint32              SizeX         = CreateDesc.Extent.X;
-	const uint32              SizeY         = CreateDesc.Extent.Y;
-	const uint32              SizeZ         = bCubeTexture ? CreateDesc.ArraySize * 6 : CreateDesc.ArraySize;
-	const EPixelFormat        Format        = CreateDesc.Format;
-	const uint32              NumMips       = CreateDesc.NumMips;
-	const uint32              NumSamples    = CreateDesc.NumSamples;
-	const ETextureCreateFlags Flags         = CreateDesc.Flags;
-
-	check(SizeX > 0 && SizeY > 0 && NumMips > 0);
-
-	if (bCubeTexture)
+#if GUARDED_TEXTURE_CREATES
+	bool bDriverCrash = true;
+	__try
 	{
-		checkf(SizeX <= GetMaxCubeTextureDimension(), TEXT("Requested cube texture size too large: %i, Max: %i, DebugName: '%s'"), SizeX, GetMaxCubeTextureDimension(), CreateDesc.DebugName ? CreateDesc.DebugName : TEXT(""));
-		check(SizeX == SizeY);
+#endif // #if GUARDED_TEXTURE_CREATES
+		VERIFYD3D11CREATETEXTURERESULT(
+			Direct3DDevice->CreateTexture3D(TextureDesc,SubResourceData,OutTexture2D),
+			UEFormat,
+			TextureDesc->Width,
+			TextureDesc->Height,
+			TextureDesc->Depth,
+			TextureDesc->Format,
+			TextureDesc->MipLevels,
+			TextureDesc->BindFlags,
+			TextureDesc->Usage,
+			TextureDesc->CPUAccessFlags,
+			TextureDesc->MiscFlags,			
+			0,
+			0,
+			SubResourceData ? SubResourceData->pSysMem : nullptr,
+			SubResourceData ? SubResourceData->SysMemPitch : 0,
+			SubResourceData ? SubResourceData->SysMemSlicePitch : 0,
+			Direct3DDevice,
+			DebugName
+			);
+#if GUARDED_TEXTURE_CREATES
+		bDriverCrash = false;
 	}
-	else
+	__finally
 	{
-		checkf(SizeX <= GetMax2DTextureDimension(), TEXT("Requested texture2d x size too large: %i, Max: %i, DebugName: '%s'"), SizeX, GetMax2DTextureDimension(), CreateDesc.DebugName ? CreateDesc.DebugName : TEXT(""));
-		checkf(SizeY <= GetMax2DTextureDimension(), TEXT("Requested texture2d y size too large: %i, Max: %i, DebugName: '%s'"), SizeY, GetMax2DTextureDimension(), CreateDesc.DebugName ? CreateDesc.DebugName : TEXT(""));
+		if (bDriverCrash)
+		{
+			UE_LOG(LogD3D11RHI,Error,
+				TEXT("Driver crashed while creating texture: %ux%ux%u %s(0x%08x) with %u mips, PF_ %d"),
+				TextureDesc->Width,
+				TextureDesc->Height,
+				TextureDesc->Depth,
+				UE::DXGIUtilities::GetFormatString(TextureDesc->Format),
+				(uint32)TextureDesc->Format,
+				TextureDesc->MipLevels,
+				UEFormat
+				);
+		}
 	}
+#endif // #if GUARDED_TEXTURE_CREATES
+}
 
-	if (bTextureArray)
+FD3D11Texture* FD3D11DynamicRHI::BeginCreateTextureInternal(const FRHITextureCreateDesc& CreateDesc)
+{
+	FD3D11Texture* Texture = new FD3D11Texture(CreateDesc);
+	return Texture;
+}
+
+inline uint32 GetMaxMSAAQuality(uint32 SampleCount)
+{
+	if (SampleCount <= DX_MAX_MSAA_COUNT)
 	{
-		checkf(SizeZ <= GetMaxTextureArrayLayers(), TEXT("Requested texture array size too large: %i, Max: %i, DebugName: '%s'"), SizeZ, GetMaxTextureArrayLayers(), CreateDesc.DebugName ? CreateDesc.DebugName : TEXT(""));
+		return 0;
 	}
+	return 0xffffffff;
+}
 
-	SCOPE_CYCLE_COUNTER(STAT_D3D11CreateTextureTime);
+enum class ED3D11TextureCreateViewFlags
+{
+	None        = 0,
+	SRV         = 1 << 0,
+	RTV         = 1 << 1,
+	DSV         = 1 << 2,
+};
+ENUM_CLASS_FLAGS(ED3D11TextureCreateViewFlags);
+
+// Setup a D3D11_TEXTURE#D_DESC structure and return a set of flags telling the caller which view types should be created.
+template<typename TDesc>
+static ED3D11TextureCreateViewFlags SetupD3D11TextureCommonDesc(TDesc& D3D11Desc, const FRHITextureDesc& TextureDesc, DXGI_FORMAT PlatformResourceFormat)
+{
+	FMemory::Memzero(D3D11Desc);
+
+	const bool                bTextureArray = TextureDesc.IsTextureArray();
+	const bool                bCubeTexture = TextureDesc.IsTextureCube();
+	const uint32              SizeX = TextureDesc.Extent.X;
+	const uint32              SizeY = TextureDesc.Extent.Y;
+	const uint32              SizeZ = bCubeTexture ? TextureDesc.ArraySize * 6 : TextureDesc.ArraySize;
+	const EPixelFormat        Format = TextureDesc.Format;
+	const uint32              NumMips = TextureDesc.NumMips;
+	const uint32              NumSamples = TextureDesc.NumSamples;
+	const ETextureCreateFlags Flags = TextureDesc.Flags;
+
+	check(SizeX > 0 && SizeY > 0 && NumMips > 0);
 
-	const bool bSRGB = EnumHasAnyFlags(Flags, TexCreate_SRGB);
+	const bool bSRGB = EnumHasAnyFlags(Flags, ETextureCreateFlags::SRGB);
 
-	const DXGI_FORMAT PlatformResourceFormat = UE::DXGIUtilities::GetPlatformTextureResourceFormat((DXGI_FORMAT)GPixelFormats[Format].PlatformFormat, Flags);
-	const DXGI_FORMAT PlatformShaderResourceFormat = UE::DXGIUtilities::FindShaderResourceFormat(PlatformResourceFormat, bSRGB);
-	const DXGI_FORMAT PlatformRenderTargetFormat = UE::DXGIUtilities::FindShaderResourceFormat(PlatformResourceFormat, bSRGB);
-	
 	uint32 CPUAccessFlags = 0;
 	D3D11_USAGE TextureUsage = D3D11_USAGE_DEFAULT;
 	bool bCreateShaderResource = true;
 
 	uint32 ActualMSAAQuality = GetMaxMSAAQuality(NumSamples);
 	check(ActualMSAAQuality != 0xffffffff);
-	check(NumSamples == 1 || !EnumHasAnyFlags(Flags, TexCreate_Shared));
+	check(NumSamples == 1 || !EnumHasAnyFlags(Flags, ETextureCreateFlags::Shared));
 
 	const bool bIsMultisampled = NumSamples > 1;
 
-	if (EnumHasAnyFlags(Flags, TexCreate_CPUReadback))
+	if (EnumHasAnyFlags(Flags, ETextureCreateFlags::CPUReadback))
 	{
-		check(!EnumHasAnyFlags(Flags, TexCreate_RenderTargetable| TexCreate_DepthStencilTargetable| TexCreate_ShaderResource));
+		check(!EnumHasAnyFlags(Flags, ETextureCreateFlags::RenderTargetable | ETextureCreateFlags::DepthStencilTargetable | ETextureCreateFlags::ShaderResource));
 
 		CPUAccessFlags = D3D11_CPU_ACCESS_READ;
 		TextureUsage = D3D11_USAGE_STAGING;
 		bCreateShaderResource = false;
 	}
 
-	if (EnumHasAnyFlags(Flags, TexCreate_CPUWritable))
+	if (EnumHasAnyFlags(Flags, ETextureCreateFlags::CPUWritable))
 	{
 		CPUAccessFlags = D3D11_CPU_ACCESS_WRITE;
 		TextureUsage = D3D11_USAGE_STAGING;
@@ -268,76 +308,68 @@ FD3D11Texture* FD3D11DynamicRHI::CreateD3D11Texture2D(FRHITextureCreateDesc cons
 	}
 
 	// Describe the texture.
-	D3D11_TEXTURE2D_DESC TextureDesc;
-	ZeroMemory( &TextureDesc, sizeof( D3D11_TEXTURE2D_DESC ) );
-	TextureDesc.Width = SizeX;
-	TextureDesc.Height = SizeY;
-	TextureDesc.MipLevels = NumMips;
-	TextureDesc.ArraySize = SizeZ;
-	TextureDesc.Format = PlatformResourceFormat;
-	TextureDesc.SampleDesc.Count = NumSamples;
-	TextureDesc.SampleDesc.Quality = ActualMSAAQuality;
-	TextureDesc.Usage = TextureUsage;
-	TextureDesc.BindFlags = bCreateShaderResource? D3D11_BIND_SHADER_RESOURCE : 0;
-	TextureDesc.CPUAccessFlags = CPUAccessFlags;
-	TextureDesc.MiscFlags = bCubeTexture ? D3D11_RESOURCE_MISC_TEXTURECUBE : 0;
-
-	ApplyBC7SoftwareAdapterWorkaround(Adapter.bSoftwareAdapter, TextureDesc);
+	D3D11Desc.Width = SizeX;
+	D3D11Desc.Height = SizeY;
+	D3D11Desc.MipLevels = NumMips;
+	D3D11Desc.Format = PlatformResourceFormat;
+	D3D11Desc.Usage = TextureUsage;
+	D3D11Desc.BindFlags = bCreateShaderResource ? D3D11_BIND_SHADER_RESOURCE : 0;
+	D3D11Desc.CPUAccessFlags = CPUAccessFlags;
+	D3D11Desc.MiscFlags = bCubeTexture ? D3D11_RESOURCE_MISC_TEXTURECUBE : 0;
 
 	// NV12/P010 doesn't support SRV in NV12 format so don't create SRV for it.
 	// Todo: add support for SRVs of underneath luminance & chrominance textures.
 	if (Format == PF_NV12 || Format == PF_P010)
 	{
-		// JoeG - I moved this to below the bind flags because it is valid to bind R8 or B8G8 to this
-		// And creating a SRV afterward would fail because of the missing bind flags
+		// This has to be set after the bind flags because it is valid to bind R8 or B8G8 to this
+		// and creating a SRV afterward would fail because of the missing bind flags
 		bCreateShaderResource = false;
 	}
 
-	if (EnumHasAnyFlags(Flags, TexCreate_DisableSRVCreation))
+	if (EnumHasAnyFlags(Flags, ETextureCreateFlags::DisableSRVCreation))
 	{
 		bCreateShaderResource = false;
 	}
 
-	if (EnumHasAnyFlags(Flags, TexCreate_Shared))
+	if (EnumHasAnyFlags(Flags, ETextureCreateFlags::Shared))
 	{
 		if (GCVarUseSharedKeyedMutex->GetInt() != 0)
 		{
-			TextureDesc.MiscFlags |= D3D11_RESOURCE_MISC_SHARED_KEYEDMUTEX;
+			D3D11Desc.MiscFlags |= D3D11_RESOURCE_MISC_SHARED_KEYEDMUTEX;
 		}
 		else
 		{
-			TextureDesc.MiscFlags |= D3D11_RESOURCE_MISC_SHARED;
+			D3D11Desc.MiscFlags |= D3D11_RESOURCE_MISC_SHARED;
 		}
 	}
 
 	// Set up the texture bind flags.
 	bool bCreateRTV = false;
 	bool bCreateDSV = false;
-	bool bCreatedRTVPerSlice = false;
 
-	if(EnumHasAnyFlags(Flags, TexCreate_RenderTargetable))
+	if (EnumHasAnyFlags(Flags, ETextureCreateFlags::RenderTargetable))
 	{
-		check(!EnumHasAnyFlags(Flags, TexCreate_DepthStencilTargetable| TexCreate_ResolveTargetable));
-		TextureDesc.BindFlags |= D3D11_BIND_RENDER_TARGET;
-		bCreateRTV = true;		
+		check(!EnumHasAnyFlags(Flags, ETextureCreateFlags::DepthStencilTargetable | ETextureCreateFlags::ResolveTargetable));
+		D3D11Desc.BindFlags |= D3D11_BIND_RENDER_TARGET;
+		bCreateRTV = true;
 	}
-	else if(EnumHasAnyFlags(Flags, TexCreate_DepthStencilTargetable))
+	else if (EnumHasAnyFlags(Flags, ETextureCreateFlags::DepthStencilTargetable))
 	{
-		check(!EnumHasAnyFlags(Flags, TexCreate_RenderTargetable| TexCreate_ResolveTargetable));
-		TextureDesc.BindFlags |= D3D11_BIND_DEPTH_STENCIL; 
+		check(!EnumHasAnyFlags(Flags, ETextureCreateFlags::RenderTargetable | ETextureCreateFlags::ResolveTargetable));
+		D3D11Desc.BindFlags |= D3D11_BIND_DEPTH_STENCIL;
 		bCreateDSV = true;
 	}
-	else if(EnumHasAnyFlags(Flags, TexCreate_ResolveTargetable))
+	else if (EnumHasAnyFlags(Flags, ETextureCreateFlags::ResolveTargetable))
 	{
-		check(!EnumHasAnyFlags(Flags, TexCreate_RenderTargetable| TexCreate_DepthStencilTargetable));
-		if(Format == PF_DepthStencil || Format == PF_ShadowDepth || Format == PF_D24)
+		check(!EnumHasAnyFlags(Flags, ETextureCreateFlags::RenderTargetable | ETextureCreateFlags::DepthStencilTargetable));
+		if (Format == PF_DepthStencil || Format == PF_ShadowDepth || Format == PF_D24)
 		{
-			TextureDesc.BindFlags |= D3D11_BIND_DEPTH_STENCIL; 
+			D3D11Desc.BindFlags |= D3D11_BIND_DEPTH_STENCIL;
 			bCreateDSV = true;
 		}
 		else
 		{
-			TextureDesc.BindFlags |= D3D11_BIND_RENDER_TARGET;
+			D3D11Desc.BindFlags |= D3D11_BIND_RENDER_TARGET;
 			bCreateRTV = true;
 		}
 	}
@@ -348,193 +380,274 @@ FD3D11Texture* FD3D11DynamicRHI::CreateD3D11Texture2D(FRHITextureCreateDesc cons
 		bCreateRTV = false;
 	}
 
-	if (EnumHasAnyFlags(Flags, TexCreate_UAV))
+	if (EnumHasAnyFlags(Flags, ETextureCreateFlags::UAV))
 	{
-		TextureDesc.BindFlags |= D3D11_BIND_UNORDERED_ACCESS;
+		D3D11Desc.BindFlags |= D3D11_BIND_UNORDERED_ACCESS;
 	}
 
-	if (bCreateDSV && !EnumHasAnyFlags(Flags, TexCreate_ShaderResource))
+	if (bCreateDSV && !EnumHasAnyFlags(Flags, ETextureCreateFlags::ShaderResource))
 	{
-		TextureDesc.BindFlags &= ~D3D11_BIND_SHADER_RESOURCE;
+		D3D11Desc.BindFlags &= ~D3D11_BIND_SHADER_RESOURCE;
 		bCreateShaderResource = false;
 	}
 
-	TRefCountPtr<ID3D11Texture2D> TextureResource;
-	TRefCountPtr<ID3D11ShaderResourceView> ShaderResourceView;
-	TArray<TRefCountPtr<ID3D11RenderTargetView> > RenderTargetViews;
-	TRefCountPtr<ID3D11DepthStencilView> DepthStencilViews[FExclusiveDepthStencil::MaxIndex];
+	ED3D11TextureCreateViewFlags CreateViewFlags = ED3D11TextureCreateViewFlags::None;
+	if (bCreateShaderResource)
+	{
+		EnumAddFlags(CreateViewFlags, ED3D11TextureCreateViewFlags::SRV);
+	}
+	if (bCreateRTV)
+	{
+		EnumAddFlags(CreateViewFlags, ED3D11TextureCreateViewFlags::RTV);
+	}
+	if (bCreateDSV)
+	{
+		EnumAddFlags(CreateViewFlags, ED3D11TextureCreateViewFlags::DSV);
+	}
+	return CreateViewFlags;
+}
+
+TRefCountPtr<ID3D11RenderTargetView> CreateRTV(FD3D11Device* Direct3DDevice, ID3D11Resource* Resource, const FRHITextureDesc& TextureDesc, DXGI_FORMAT PlatformResourceFormat, uint32 MipIndex, uint32 SliceIndex, uint32 SliceCount)
+{
+	D3D11_RENDER_TARGET_VIEW_DESC RTVDesc;
+	FMemory::Memzero(RTVDesc);
 
-	TArray<D3D11_SUBRESOURCE_DATA> SubResourceData;
-	D3D11_SUBRESOURCE_DATA const* pSubresourceData = nullptr;
+	RTVDesc.Format = UE::DXGIUtilities::FindShaderResourceFormat(PlatformResourceFormat, EnumHasAnyFlags(TextureDesc.Flags, ETextureCreateFlags::SRGB));
 
-	if (InitialData.Num())
+	if (TextureDesc.IsTexture3D())
 	{
-		// Caller provided initla data structs.
-		check(InitialData.Num() == NumMips * SizeZ);
-		pSubresourceData = InitialData.GetData();
+		RTVDesc.ViewDimension = D3D11_RTV_DIMENSION_TEXTURE3D;
+		RTVDesc.Texture3D.MipSlice = MipIndex;
+		RTVDesc.Texture3D.FirstWSlice = 0;
+		RTVDesc.Texture3D.WSize = TextureDesc.Depth;
 	}
-	else if (CreateDesc.BulkData)
+	else if (TextureDesc.IsTextureArray() || TextureDesc.IsTextureCube())
 	{
-		uint8* Data = (uint8*)CreateDesc.BulkData->GetResourceBulkData();
+		if (TextureDesc.NumSamples > 1)
+		{
+			RTVDesc.ViewDimension = D3D11_RTV_DIMENSION_TEXTURE2DMSARRAY;
+			RTVDesc.Texture2DMSArray.FirstArraySlice = SliceIndex;
+			RTVDesc.Texture2DMSArray.ArraySize = SliceCount;
+		}
+		else
+		{
+			RTVDesc.ViewDimension = D3D11_RTV_DIMENSION_TEXTURE2DARRAY;
+			RTVDesc.Texture2DArray.FirstArraySlice = SliceIndex;
+			RTVDesc.Texture2DArray.ArraySize = SliceCount;
+			RTVDesc.Texture2DArray.MipSlice = MipIndex;
+		}
+	}
+	else
+	{
+		if (TextureDesc.NumSamples > 1)
+		{
+			RTVDesc.ViewDimension = D3D11_RTV_DIMENSION_TEXTURE2DMS;
+			// Nothing to set
+		}
+		else
+		{
+			RTVDesc.ViewDimension = D3D11_RTV_DIMENSION_TEXTURE2D;
+			RTVDesc.Texture2D.MipSlice = MipIndex;
+		}
+	}
 
-		// each mip of each array slice counts as a subresource
-		SubResourceData.AddZeroed(NumMips * SizeZ);
+	TRefCountPtr<ID3D11RenderTargetView> RenderTargetView;
+	VERIFYD3D11RESULT_EX(Direct3DDevice->CreateRenderTargetView(Resource, &RTVDesc, RenderTargetView.GetInitReference()), Direct3DDevice);
 
-		uint32 SliceOffset = 0;
-		for (uint32 ArraySliceIndex = 0; ArraySliceIndex < SizeZ; ++ArraySliceIndex)
-		{			
-			uint32 MipOffset = 0;
-			for(uint32 MipIndex = 0;MipIndex < NumMips;++MipIndex)
-			{
-				uint32 DataOffset = SliceOffset + MipOffset;
-				uint32 SubResourceIndex = ArraySliceIndex * NumMips + MipIndex;
+	return RenderTargetView;
+}
 
-				uint32 NumBlocksX = FMath::Max<uint32>(1,((SizeX >> MipIndex) + GPixelFormats[Format].BlockSizeX-1) / GPixelFormats[Format].BlockSizeX);
-				uint32 NumBlocksY = FMath::Max<uint32>(1,((SizeY >> MipIndex) + GPixelFormats[Format].BlockSizeY-1) / GPixelFormats[Format].BlockSizeY);
+TRefCountPtr<ID3D11ShaderResourceView> CreateSRV(FD3D11Device* Direct3DDevice, ID3D11Resource* Resource, const FRHITextureDesc& TextureDesc, DXGI_FORMAT PlatformResourceFormat)
+{
+	const DXGI_FORMAT PlatformShaderResourceFormat = UE::DXGIUtilities::FindShaderResourceFormat(PlatformResourceFormat, EnumHasAnyFlags(TextureDesc.Flags, ETextureCreateFlags::SRGB));
 
-				SubResourceData[SubResourceIndex].pSysMem = &Data[DataOffset];
-				SubResourceData[SubResourceIndex].SysMemPitch      =  NumBlocksX * GPixelFormats[Format].BlockBytes;
-				SubResourceData[SubResourceIndex].SysMemSlicePitch =  NumBlocksX * NumBlocksY * SubResourceData[MipIndex].SysMemPitch;
+	D3D11_SHADER_RESOURCE_VIEW_DESC SRVDesc;
+	FMemory::Memzero(SRVDesc);
 
-				MipOffset                                  += NumBlocksY * SubResourceData[MipIndex].SysMemPitch;
-			}
-			SliceOffset += MipOffset;
+	SRVDesc.Format = PlatformShaderResourceFormat;
+
+	switch (TextureDesc.Dimension)
+	{
+	case ETextureDimension::Texture2D:
+		if (TextureDesc.NumSamples > 1)
+		{
+			SRVDesc.ViewDimension = D3D11_SRV_DIMENSION_TEXTURE2DMS;
+			// Nothing to set
+		}
+		else
+		{
+			SRVDesc.ViewDimension = D3D11_SRV_DIMENSION_TEXTURE2D;
+			SRVDesc.Texture2D.MostDetailedMip = 0;
+			SRVDesc.Texture2D.MipLevels = TextureDesc.NumMips;
 		}
+		break;
+	case ETextureDimension::Texture2DArray:
+		if (TextureDesc.NumSamples > 1)
+		{
+			SRVDesc.ViewDimension = D3D11_SRV_DIMENSION_TEXTURE2DMSARRAY;
+			SRVDesc.Texture2DMSArray.FirstArraySlice = 0;
+			SRVDesc.Texture2DMSArray.ArraySize = TextureDesc.ArraySize;
+		}
+		else
+		{
+			SRVDesc.ViewDimension = D3D11_SRV_DIMENSION_TEXTURE2DARRAY;
+			SRVDesc.Texture2DArray.MostDetailedMip = 0;
+			SRVDesc.Texture2DArray.MipLevels = TextureDesc.NumMips;
+			SRVDesc.Texture2DArray.FirstArraySlice = 0;
+			SRVDesc.Texture2DArray.ArraySize = TextureDesc.ArraySize;
+		}
+		break;
+	case ETextureDimension::Texture3D:
+		SRVDesc.ViewDimension = D3D11_SRV_DIMENSION_TEXTURE3D;
+		SRVDesc.Texture3D.MipLevels = TextureDesc.NumMips;
+		SRVDesc.Texture3D.MostDetailedMip = 0;
+		break;
+	case ETextureDimension::TextureCube:
+		SRVDesc.ViewDimension = D3D11_SRV_DIMENSION_TEXTURECUBE;
+		SRVDesc.TextureCube.MostDetailedMip = 0;
+		SRVDesc.TextureCube.MipLevels = TextureDesc.NumMips;
+		break;
+	case ETextureDimension::TextureCubeArray:
+		SRVDesc.ViewDimension = D3D11_SRV_DIMENSION_TEXTURECUBEARRAY;
+		SRVDesc.TextureCubeArray.MostDetailedMip = 0;
+		SRVDesc.TextureCubeArray.MipLevels = TextureDesc.NumMips;
+		SRVDesc.TextureCubeArray.First2DArrayFace = 0;
+		SRVDesc.TextureCubeArray.NumCubes = TextureDesc.ArraySize;
+		break;
+	default:
+		break;
+	}
+
+	TRefCountPtr<ID3D11ShaderResourceView> SRV;
+	VERIFYD3D11RESULT_EX(Direct3DDevice->CreateShaderResourceView(Resource, &SRVDesc, SRV.GetInitReference()), Direct3DDevice);
+
+	return SRV;
+}
 
-		pSubresourceData = SubResourceData.GetData();
-	}
+void FD3D11DynamicRHI::FinalizeCreateTexture2DInternal(FD3D11Texture* Texture, TConstArrayView<D3D11_SUBRESOURCE_DATA> InitialData)
+{
+	const FString DebugName = Texture->GetName().ToString();
+	const FRHITextureDesc& TextureDesc = Texture->GetDesc();
 
-#if INTEL_EXTENSIONS
-	if (EnumHasAnyFlags(Flags, ETextureCreateFlags::Atomic64Compatible) && IsRHIDeviceIntel() && GRHISupportsAtomicUInt64)
-	{
-		INTC_D3D11_TEXTURE2D_DESC IntelDesc{};
-		IntelDesc.EmulatedTyped64bitAtomics = true;
-		IntelDesc.pD3D11Desc = &TextureDesc;
+	check(!TextureDesc.IsTexture3D());
 
-		VERIFYD3D11RESULT(INTC_D3D11_CreateTexture2D(IntelExtensionContext, &IntelDesc, pSubresourceData, TextureResource.GetInitReference()));
+	const bool                bTextureArray = TextureDesc.IsTextureArray();
+	const bool                bCubeTexture  = TextureDesc.IsTextureCube();
+	const uint32              SizeX         = TextureDesc.Extent.X;
+	const uint32              SizeY         = TextureDesc.Extent.Y;
+	const uint32              SizeZ         = bCubeTexture ? TextureDesc.ArraySize * 6 : TextureDesc.ArraySize;
+	const EPixelFormat        Format        = TextureDesc.Format;
+	const uint32              NumMips       = TextureDesc.NumMips;
+	const uint32              NumSamples    = TextureDesc.NumSamples;
+	const ETextureCreateFlags Flags         = TextureDesc.Flags;
+
+	check(SizeX > 0 && SizeY > 0 && NumMips > 0);
+
+	if (bCubeTexture)
+	{
+		checkf(SizeX <= GetMaxCubeTextureDimension(), TEXT("Requested cube texture size too large: %i, Max: %i, DebugName: '%s'"), SizeX, GetMaxCubeTextureDimension(), *DebugName);
+		check(SizeX == SizeY);
 	}
 	else
-#endif
 	{
-		SafeCreateTexture2D(Direct3DDevice, Format, &TextureDesc, pSubresourceData, TextureResource.GetInitReference(), CreateDesc.DebugName);
+		checkf(SizeX <= GetMax2DTextureDimension(), TEXT("Requested texture2d x size too large: %i, Max: %i, DebugName: '%s'"), SizeX, GetMax2DTextureDimension(), *DebugName);
+		checkf(SizeY <= GetMax2DTextureDimension(), TEXT("Requested texture2d y size too large: %i, Max: %i, DebugName: '%s'"), SizeY, GetMax2DTextureDimension(), *DebugName);
 	}
 
-	if (bCreateRTV)
+	if (bTextureArray)
 	{
-		// Create a render target view for each mip
-		for (uint32 MipIndex = 0; MipIndex < NumMips; MipIndex++)
-		{
-			if (EnumHasAnyFlags(Flags, TexCreate_TargetArraySlicesIndependently) && (bTextureArray || bCubeTexture))
-			{
-				bCreatedRTVPerSlice = true;
+		checkf(SizeZ <= GetMaxTextureArrayLayers(), TEXT("Requested texture array size too large: %i, Max: %i, DebugName: '%s'"), SizeZ, GetMaxTextureArrayLayers(), *DebugName);
+	}
 
-				for (uint32 SliceIndex = 0; SliceIndex < TextureDesc.ArraySize; SliceIndex++)
-				{
-					D3D11_RENDER_TARGET_VIEW_DESC RTVDesc;
-					FMemory::Memzero(RTVDesc);
+	SCOPE_CYCLE_COUNTER(STAT_D3D11CreateTextureTime);
 
-					RTVDesc.Format = PlatformRenderTargetFormat;
+	const DXGI_FORMAT PlatformResourceFormat = UE::DXGIUtilities::GetPlatformTextureResourceFormat(Format, Flags);
 
-					if (bIsMultisampled)
-					{
-						RTVDesc.ViewDimension = D3D11_RTV_DIMENSION_TEXTURE2DMSARRAY;
-						RTVDesc.Texture2DMSArray.FirstArraySlice = SliceIndex;
-						RTVDesc.Texture2DMSArray.ArraySize = 1;
-					}
-					else
-					{
-						RTVDesc.ViewDimension = D3D11_RTV_DIMENSION_TEXTURE2DARRAY;
-						RTVDesc.Texture2DArray.FirstArraySlice = SliceIndex;
-						RTVDesc.Texture2DArray.ArraySize = 1;
-						RTVDesc.Texture2DArray.MipSlice = MipIndex;
-					}
+	const uint32 ActualMSAAQuality = GetMaxMSAAQuality(TextureDesc.NumSamples);
+	check(ActualMSAAQuality != 0xffffffff);
+	check(TextureDesc.NumSamples == 1 || !EnumHasAnyFlags(Flags, ETextureCreateFlags::Shared));
 
-					TRefCountPtr<ID3D11RenderTargetView> RenderTargetView;
-					VERIFYD3D11RESULT_EX(Direct3DDevice->CreateRenderTargetView(TextureResource,&RTVDesc,RenderTargetView.GetInitReference()), Direct3DDevice);
-					RenderTargetViews.Add(RenderTargetView);
-				}
-			}
-			else
-			{
-				D3D11_RENDER_TARGET_VIEW_DESC RTVDesc;
-				FMemory::Memzero(RTVDesc);
+	// Describe the texture.
+	D3D11_TEXTURE2D_DESC Desc;
+	
+	ED3D11TextureCreateViewFlags CreateViewFlags = SetupD3D11TextureCommonDesc(Desc, TextureDesc, PlatformResourceFormat);
 
-				RTVDesc.Format = PlatformRenderTargetFormat;
+	// Texture2D specific vars
+	Desc.ArraySize = SizeZ;
+	Desc.SampleDesc.Count = NumSamples;
+	Desc.SampleDesc.Quality = ActualMSAAQuality;
 
-				if (bTextureArray || bCubeTexture)
-				{
-					if (bIsMultisampled)
-					{
-						RTVDesc.ViewDimension = D3D11_RTV_DIMENSION_TEXTURE2DMSARRAY;
-						RTVDesc.Texture2DMSArray.FirstArraySlice = 0;
-						RTVDesc.Texture2DMSArray.ArraySize = TextureDesc.ArraySize;
-					}
-					else
-					{
-						RTVDesc.ViewDimension = D3D11_RTV_DIMENSION_TEXTURE2DARRAY;
-						RTVDesc.Texture2DArray.FirstArraySlice = 0;
-						RTVDesc.Texture2DArray.ArraySize = TextureDesc.ArraySize;
-						RTVDesc.Texture2DArray.MipSlice = MipIndex;
-					}
-				}
-				else
-				{
-					if (bIsMultisampled)
-					{
-						RTVDesc.ViewDimension = D3D11_RTV_DIMENSION_TEXTURE2DMS;
-						// Nothing to set
-					}
-					else
-					{
-						RTVDesc.ViewDimension = D3D11_RTV_DIMENSION_TEXTURE2D;
-						RTVDesc.Texture2D.MipSlice = MipIndex;
-					}
-				}
+	ApplyBC7SoftwareAdapterWorkaround(Adapter.bSoftwareAdapter, Desc);
 
-				TRefCountPtr<ID3D11RenderTargetView> RenderTargetView;
-				VERIFYD3D11RESULT_EX(Direct3DDevice->CreateRenderTargetView(TextureResource,&RTVDesc,RenderTargetView.GetInitReference()), Direct3DDevice);
-				RenderTargetViews.Add(RenderTargetView);
-			}
-		}
+	const D3D11_SUBRESOURCE_DATA* pSubresourceData = nullptr;
+	if (InitialData.Num())
+	{
+		// Caller provided initial data.
+		check(InitialData.Num() == NumMips * SizeZ);
+		pSubresourceData = InitialData.GetData();
 	}
-	
-	if (bCreateDSV)
+
+	TRefCountPtr<ID3D11Texture2D> TextureResource;
+
+#if INTEL_EXTENSIONS
+	if (EnumHasAnyFlags(TextureDesc.Flags, ETextureCreateFlags::Atomic64Compatible) && IsRHIDeviceIntel() && GRHIGlobals.SupportsAtomicUInt64)
 	{
-		// Create a depth-stencil-view for the texture.
-		D3D11_DEPTH_STENCIL_VIEW_DESC DSVDesc;
-		FMemory::Memzero(DSVDesc);
+		INTC_D3D11_TEXTURE2D_DESC IntelDesc{};
+		IntelDesc.EmulatedTyped64bitAtomics = true;
+		IntelDesc.pD3D11Desc = &Desc;
 
-		DSVDesc.Format = UE::DXGIUtilities::FindDepthStencilFormat(PlatformResourceFormat);
+		VERIFYD3D11RESULT(INTC_D3D11_CreateTexture2D(IntelExtensionContext, &IntelDesc, pSubresourceData, TextureResource.GetInitReference()));
+	}
+	else
+#endif
+	{
+		SafeCreateTexture2D(Direct3DDevice, Format, &Desc, pSubresourceData, TextureResource.GetInitReference(), *DebugName);
+	}
 
-		if (bTextureArray || bCubeTexture)
+	TArray<TRefCountPtr<ID3D11RenderTargetView>> RenderTargetViews;
+	bool bCreatedRTVPerSlice = false;
+
+	if (EnumHasAnyFlags(CreateViewFlags, ED3D11TextureCreateViewFlags::RTV))
+	{
+		if (EnumHasAnyFlags(TextureDesc.Flags, ETextureCreateFlags::TargetArraySlicesIndependently) && (TextureDesc.IsTextureArray() || TextureDesc.IsTextureCube()))
 		{
-			if (bIsMultisampled)
-			{
-				DSVDesc.ViewDimension = D3D11_DSV_DIMENSION_TEXTURE2DMSARRAY;
-				DSVDesc.Texture2DMSArray.FirstArraySlice = 0;
-				DSVDesc.Texture2DMSArray.ArraySize = TextureDesc.ArraySize;
-			}
-			else
+			bCreatedRTVPerSlice = true;
+
+			for (uint32 MipIndex = 0; MipIndex < NumMips; MipIndex++)
 			{
-				DSVDesc.ViewDimension = D3D11_DSV_DIMENSION_TEXTURE2DARRAY;
-				DSVDesc.Texture2DArray.FirstArraySlice = 0;
-				DSVDesc.Texture2DArray.ArraySize = TextureDesc.ArraySize;
-				DSVDesc.Texture2DArray.MipSlice = 0;
+				for (uint32 SliceIndex = 0; SliceIndex < Desc.ArraySize; SliceIndex++)
+				{
+					RenderTargetViews.Emplace(CreateRTV(Direct3DDevice, TextureResource, TextureDesc, PlatformResourceFormat, MipIndex, SliceIndex, 1));
+				}
 			}
 		}
 		else
 		{
-			if (bIsMultisampled)
+			for (uint32 MipIndex = 0; MipIndex < NumMips; MipIndex++)
 			{
-				DSVDesc.ViewDimension = D3D11_DSV_DIMENSION_TEXTURE2DMS;
-				// Nothing to set
-			}
-			else
-			{
-				DSVDesc.ViewDimension = D3D11_DSV_DIMENSION_TEXTURE2D;
-				DSVDesc.Texture2D.MipSlice = 0;
+				RenderTargetViews.Emplace(CreateRTV(Direct3DDevice, TextureResource, TextureDesc, PlatformResourceFormat, MipIndex, 0, Desc.ArraySize));
 			}
 		}
+	}
+	
+	TRefCountPtr<ID3D11DepthStencilView> DepthStencilViews[FExclusiveDepthStencil::MaxIndex];
+
+	if (EnumHasAnyFlags(CreateViewFlags, ED3D11TextureCreateViewFlags::DSV))
+	{
+		const D3D11_DSV_DIMENSION DSVDimension =
+			(bTextureArray || bCubeTexture)
+			? (TextureDesc.NumSamples > 1 ? D3D11_DSV_DIMENSION_TEXTURE2DMSARRAY : D3D11_DSV_DIMENSION_TEXTURE2DARRAY)
+			: (TextureDesc.NumSamples > 1 ? D3D11_DSV_DIMENSION_TEXTURE2DMS : D3D11_DSV_DIMENSION_TEXTURE2D);
+
+		// Create a depth-stencil-view for the texture.
+		CD3D11_DEPTH_STENCIL_VIEW_DESC DSVDesc(
+			DSVDimension,
+			UE::DXGIUtilities::FindDepthStencilFormat((DXGI_FORMAT)GPixelFormats[TextureDesc.Format].PlatformFormat),
+			0, // MipSlice
+			0, // FirstArraySlice
+			Desc.ArraySize,
+			0 // Flags
+		);
 
 		for (uint32 AccessType = 0; AccessType < FExclusiveDepthStencil::MaxIndex; ++AccessType)
 		{
@@ -552,228 +665,74 @@ FD3D11Texture* FD3D11DynamicRHI::CreateD3D11Texture2D(FRHITextureCreateDesc cons
 	check(IsValidRef(TextureResource));
 
 	// Create a shader resource view for the texture.
-	if (bCreateShaderResource)
-	{
-		{
-			D3D11_SHADER_RESOURCE_VIEW_DESC SRVDesc;
-			FMemory::Memzero(SRVDesc);
-
-			SRVDesc.Format = PlatformShaderResourceFormat;
-
-			if (bCubeTexture && bTextureArray)
-			{
-				SRVDesc.ViewDimension = D3D11_SRV_DIMENSION_TEXTURECUBEARRAY;
-				SRVDesc.TextureCubeArray.MostDetailedMip = 0;
-				SRVDesc.TextureCubeArray.MipLevels = NumMips;
-				SRVDesc.TextureCubeArray.First2DArrayFace = 0;
-				SRVDesc.TextureCubeArray.NumCubes = SizeZ / 6;
-			}
-			else if(bCubeTexture)
-			{
-				SRVDesc.ViewDimension = D3D11_SRV_DIMENSION_TEXTURECUBE;
-				SRVDesc.TextureCube.MostDetailedMip = 0;
-				SRVDesc.TextureCube.MipLevels = NumMips;
-			}
-			else if(bTextureArray)
-			{
-				if (bIsMultisampled)
-				{
-					SRVDesc.ViewDimension = D3D11_SRV_DIMENSION_TEXTURE2DMSARRAY;
-					SRVDesc.Texture2DMSArray.FirstArraySlice = 0;
-					SRVDesc.Texture2DMSArray.ArraySize = TextureDesc.ArraySize;
-				}
-				else
-				{
-					SRVDesc.ViewDimension = D3D11_SRV_DIMENSION_TEXTURE2DARRAY;
-					SRVDesc.Texture2DArray.MostDetailedMip = 0;
-					SRVDesc.Texture2DArray.MipLevels = NumMips;
-					SRVDesc.Texture2DArray.FirstArraySlice = 0;
-					SRVDesc.Texture2DArray.ArraySize = TextureDesc.ArraySize;
-				}
-			}
-			else
-			{
-				if (bIsMultisampled)
-				{
-					SRVDesc.ViewDimension = D3D11_SRV_DIMENSION_TEXTURE2DMS;
-					// Nothing to set
-				}
-				else
-				{
-					SRVDesc.ViewDimension = D3D11_SRV_DIMENSION_TEXTURE2D;
-					SRVDesc.Texture2D.MostDetailedMip = 0;
-					SRVDesc.Texture2D.MipLevels = NumMips;
-				}
-			}
-			VERIFYD3D11RESULT_EX(Direct3DDevice->CreateShaderResourceView(TextureResource,&SRVDesc,ShaderResourceView.GetInitReference()), Direct3DDevice);
-		}
+	TRefCountPtr<ID3D11ShaderResourceView> ShaderResourceView;
 
+	if (EnumHasAnyFlags(CreateViewFlags, ED3D11TextureCreateViewFlags::SRV))
+	{
+		ShaderResourceView = CreateSRV(Direct3DDevice, TextureResource, TextureDesc, PlatformResourceFormat);
 		check(IsValidRef(ShaderResourceView));
 	}
 
-	FD3D11Texture* Texture2D = new FD3D11Texture(
-		CreateDesc,
+	Texture->FinalizeCreation(
 		TextureResource,
 		ShaderResourceView,
-		TextureDesc.ArraySize,
+		Desc.ArraySize,
 		bCreatedRTVPerSlice,
 		RenderTargetViews,
 		DepthStencilViews
 	);
-
-	if (CreateDesc.BulkData)
-	{
-		CreateDesc.BulkData->Discard();
-	}
-
-	return Texture2D;
 }
 
-FD3D11Texture* FD3D11DynamicRHI::CreateD3D11Texture3D(FRHITextureCreateDesc const& CreateDesc)
+void FD3D11DynamicRHI::FinalizeCreateTexture3DInternal(FD3D11Texture* Texture, TConstArrayView<D3D11_SUBRESOURCE_DATA> InitialData)
 {
-	check(CreateDesc.IsTexture3D());
-	check(CreateDesc.ArraySize == 1);
-	const uint32              SizeX   = CreateDesc.Extent.X;
-	const uint32              SizeY   = CreateDesc.Extent.Y;
-	const uint32              SizeZ   = CreateDesc.Depth;
-	const EPixelFormat        Format  = CreateDesc.Format;
-	const uint32              NumMips = CreateDesc.NumMips;
-	const ETextureCreateFlags Flags   = CreateDesc.Flags;
-
-	SCOPE_CYCLE_COUNTER(STAT_D3D11CreateTextureTime);
-	
-	const bool bSRGB = EnumHasAnyFlags(Flags, TexCreate_SRGB);
-
-	const DXGI_FORMAT PlatformResourceFormat = (DXGI_FORMAT)GPixelFormats[Format].PlatformFormat;
-	const DXGI_FORMAT PlatformShaderResourceFormat = UE::DXGIUtilities::FindShaderResourceFormat(PlatformResourceFormat, bSRGB);
-	const DXGI_FORMAT PlatformRenderTargetFormat = UE::DXGIUtilities::FindShaderResourceFormat(PlatformResourceFormat, bSRGB);
+	const FString DebugName = Texture->GetName().ToString();
+	const FRHITextureDesc& TextureDesc = Texture->GetDesc();
 
-	// Describe the texture.
-	D3D11_TEXTURE3D_DESC TextureDesc;
-	ZeroMemory( &TextureDesc, sizeof( D3D11_TEXTURE3D_DESC ) );
-	TextureDesc.Width = SizeX;
-	TextureDesc.Height = SizeY;
-	TextureDesc.Depth = SizeZ;
-	TextureDesc.MipLevels = NumMips;
-	TextureDesc.Format = PlatformResourceFormat;
-	TextureDesc.Usage = D3D11_USAGE_DEFAULT;
-	TextureDesc.BindFlags = D3D11_BIND_SHADER_RESOURCE;
-	TextureDesc.CPUAccessFlags = 0;
-	TextureDesc.MiscFlags = 0;
-
-	if (EnumHasAnyFlags(Flags, TexCreate_Shared))
-	{
-		if (GCVarUseSharedKeyedMutex->GetInt() != 0)
-		{
-			TextureDesc.MiscFlags |= D3D11_RESOURCE_MISC_SHARED_KEYEDMUTEX;
-		}
-		else
-		{
-			TextureDesc.MiscFlags |= D3D11_RESOURCE_MISC_SHARED;
-		}
-	}
-
-	if (EnumHasAnyFlags(Flags, TexCreate_UAV))
-	{
-		TextureDesc.BindFlags |= D3D11_BIND_UNORDERED_ACCESS;
-	}
-
-	bool bCreateRTV = false;
-
-	if(EnumHasAnyFlags(Flags, TexCreate_RenderTargetable))
-	{
-		TextureDesc.BindFlags |= D3D11_BIND_RENDER_TARGET;		
-		bCreateRTV = true;
-	}
+	check(TextureDesc.IsTexture3D());
+	check(TextureDesc.ArraySize == 1);
 
-	bool bCreateShaderResource = true;
+	SCOPE_CYCLE_COUNTER(STAT_D3D11CreateTextureTime);
 
-	if (EnumHasAnyFlags(Flags, TexCreate_CPUReadback))
-	{
-		check(!EnumHasAnyFlags(Flags, TexCreate_RenderTargetable | TexCreate_DepthStencilTargetable | TexCreate_ShaderResource));
+	// Set up the texture bind flags.
+	check(!EnumHasAnyFlags(TextureDesc.Flags, ETextureCreateFlags::DepthStencilTargetable | ETextureCreateFlags::ResolveTargetable));
 
-		TextureDesc.CPUAccessFlags = D3D11_CPU_ACCESS_READ;
-		TextureDesc.Usage = D3D11_USAGE_STAGING;
-		TextureDesc.BindFlags = 0;
-		bCreateShaderResource = false;
-	}
+	const DXGI_FORMAT PlatformResourceFormat = UE::DXGIUtilities::GetPlatformTextureResourceFormat(TextureDesc.Format, TextureDesc.Flags);
 
-	if (EnumHasAnyFlags(Flags, TexCreate_DisableSRVCreation))
-	{
-		bCreateShaderResource = false;
-	}
+	// Describe the texture.
+	D3D11_TEXTURE3D_DESC Desc;
 
-	// Set up the texture bind flags.
-	check(!EnumHasAnyFlags(Flags, TexCreate_DepthStencilTargetable | TexCreate_ResolveTargetable));
+	const ED3D11TextureCreateViewFlags CreateViewFlags = SetupD3D11TextureCommonDesc(Desc, TextureDesc, PlatformResourceFormat);
 
-	TArray<D3D11_SUBRESOURCE_DATA> SubResourceData;
+	// Texture3D specific vars
+	Desc.Depth = TextureDesc.Depth;
 
-	if (CreateDesc.BulkData)
+	const D3D11_SUBRESOURCE_DATA* pSubresourceData = nullptr;
+	if (InitialData.Num())
 	{
-		uint8* Data = (uint8*)CreateDesc.BulkData->GetResourceBulkData();
-		SubResourceData.AddZeroed(NumMips);
-		uint32 MipOffset = 0;
-		for(uint32 MipIndex = 0;MipIndex < NumMips;++MipIndex)
-		{
-			SubResourceData[MipIndex].pSysMem = &Data[MipOffset];
-			SubResourceData[MipIndex].SysMemPitch      =  FMath::Max<uint32>(1,SizeX >> MipIndex) * GPixelFormats[Format].BlockBytes;
-			SubResourceData[MipIndex].SysMemSlicePitch =  FMath::Max<uint32>(1,SizeY >> MipIndex) * SubResourceData[MipIndex].SysMemPitch;
-			MipOffset                                  += FMath::Max<uint32>(1,SizeZ >> MipIndex) * SubResourceData[MipIndex].SysMemSlicePitch;
-		}
+		// Caller provided initial data.
+		check(InitialData.Num() == TextureDesc.NumMips);
+		pSubresourceData = InitialData.GetData();
 	}
 
 	TRefCountPtr<ID3D11Texture3D> TextureResource;
-	const D3D11_SUBRESOURCE_DATA* SubResData = CreateDesc.BulkData != nullptr ? (const D3D11_SUBRESOURCE_DATA*)SubResourceData.GetData() : nullptr;
-	VERIFYD3D11CREATETEXTURERESULT(
-		Direct3DDevice->CreateTexture3D(&TextureDesc, SubResData,TextureResource.GetInitReference()),
-		Format,
-		SizeX,
-		SizeY,
-		SizeZ,
-		PlatformShaderResourceFormat,
-		NumMips,
-		TextureDesc.BindFlags,
-		TextureDesc.Usage,
-		TextureDesc.CPUAccessFlags,
-		TextureDesc.MiscFlags,
-		0,
-		0,
-		SubResData ? SubResData->pSysMem : nullptr,
-		SubResData ? SubResData->SysMemPitch : 0,
-		SubResData ? SubResData->SysMemSlicePitch : 0,
-		Direct3DDevice,
-		CreateDesc.DebugName
-		);
+	SafeCreateTexture3D(Direct3DDevice, TextureDesc.Format, &Desc, pSubresourceData, TextureResource.GetInitReference(), *DebugName);
 
 	// Create a shader resource view for the texture.
 	TRefCountPtr<ID3D11ShaderResourceView> ShaderResourceView;
-	if(bCreateShaderResource)
+
+	if (EnumHasAnyFlags(CreateViewFlags, ED3D11TextureCreateViewFlags::SRV))
 	{
-		D3D11_SHADER_RESOURCE_VIEW_DESC SRVDesc;
-		SRVDesc.Format = PlatformShaderResourceFormat;
-		SRVDesc.ViewDimension = D3D11_SRV_DIMENSION_TEXTURE3D;
-		SRVDesc.Texture3D.MipLevels = NumMips;
-		SRVDesc.Texture3D.MostDetailedMip = 0;
-		VERIFYD3D11RESULT_EX(Direct3DDevice->CreateShaderResourceView(TextureResource,&SRVDesc,ShaderResourceView.GetInitReference()), Direct3DDevice);
+		ShaderResourceView = CreateSRV(Direct3DDevice, TextureResource, TextureDesc, PlatformResourceFormat);
 	}
 
 	TRefCountPtr<ID3D11RenderTargetView> RenderTargetView;
-	if(bCreateRTV)
-	{
-		// Create a render-target-view for the texture.
-		D3D11_RENDER_TARGET_VIEW_DESC RTVDesc;
-		FMemory::Memzero(&RTVDesc,sizeof(RTVDesc));
-		RTVDesc.Format = PlatformRenderTargetFormat;
-		RTVDesc.ViewDimension = D3D11_RTV_DIMENSION_TEXTURE3D;
-		RTVDesc.Texture3D.MipSlice = 0;
-		RTVDesc.Texture3D.FirstWSlice = 0;
-		RTVDesc.Texture3D.WSize = SizeZ;
 
-		VERIFYD3D11RESULT_EX(Direct3DDevice->CreateRenderTargetView(TextureResource,&RTVDesc,RenderTargetView.GetInitReference()), Direct3DDevice);
+	if (EnumHasAnyFlags(CreateViewFlags, ED3D11TextureCreateViewFlags::RTV))
+	{
+		RenderTargetView = CreateRTV(Direct3DDevice, TextureResource, TextureDesc, PlatformResourceFormat, 0, 0, 1);
 	}
 
-	FD3D11Texture* Texture3D = new FD3D11Texture(
-		CreateDesc,
+	Texture->FinalizeCreation(
 		TextureResource,
 		ShaderResourceView,
 		1,     // InRTVArraySize
@@ -781,24 +740,86 @@ FD3D11Texture* FD3D11DynamicRHI::CreateD3D11Texture3D(FRHITextureCreateDesc cons
 		{ RenderTargetView },
 		{}
 	);
+}
+
+static const D3D11_SUBRESOURCE_DATA* FillSubresourceData(TArray<D3D11_SUBRESOURCE_DATA>& SubresourceData, const FRHITextureCreateDesc& CreateDesc)
+{
+	const D3D11_SUBRESOURCE_DATA* pSubresourceData = nullptr;
 
 	if (CreateDesc.BulkData)
 	{
-		CreateDesc.BulkData->Discard();
+		const FPixelFormatInfo& PixelFormat = GPixelFormats[CreateDesc.Format];
+		const uint8* Data = reinterpret_cast<const uint8*>(CreateDesc.BulkData->GetResourceBulkData());
+
+		const uint32 FaceCount = CreateDesc.IsTextureCube() ? 6 : 1;
+		const uint32 ArrayCount = CreateDesc.ArraySize * FaceCount;
+		const uint32 MipCount = CreateDesc.NumMips;
+
+		// each mip of each array slice counts as a subresource
+		SubresourceData.AddZeroed(MipCount * ArrayCount);
+
+		uint32 SliceOffset = 0;
+		for (uint32 ArraySliceIndex = 0; ArraySliceIndex < ArrayCount; ++ArraySliceIndex)
+		{
+			uint32 MipOffset = 0;
+			for (uint32 MipIndex = 0; MipIndex < MipCount; ++MipIndex)
+			{
+				const uint32 DataOffset = SliceOffset + MipOffset;
+				const uint32 SubResourceIndex = ArraySliceIndex * MipCount + MipIndex;
+
+				const uint32 NumBlocksX = UE::RHI::CalculateMipBlockCount(CreateDesc.Extent.X, MipIndex, PixelFormat.BlockSizeX);
+				const uint32 NumBlocksY = UE::RHI::CalculateMipBlockCount(CreateDesc.Extent.Y, MipIndex, PixelFormat.BlockSizeY);
+				const uint32 NumBlocksZ = UE::RHI::CalculateMipBlockCount(CreateDesc.Depth,    MipIndex, PixelFormat.BlockSizeZ);
+
+				const uint32 RowPitch = NumBlocksX * PixelFormat.BlockBytes;
+				const uint32 SlicePitch = NumBlocksY * RowPitch;
+
+				SubresourceData[SubResourceIndex].pSysMem          = &Data[DataOffset];
+				SubresourceData[SubResourceIndex].SysMemPitch      = RowPitch;
+				SubresourceData[SubResourceIndex].SysMemSlicePitch = SlicePitch;
+
+				MipOffset += SlicePitch * NumBlocksZ;
+			}
+
+			SliceOffset += MipOffset;
+		}
+
+		pSubresourceData = SubresourceData.GetData();
 	}
 
-	return Texture3D;
+	return pSubresourceData;
+}
+
+FD3D11Texture* FD3D11DynamicRHI::CreateTextureInternal(const FRHITextureCreateDesc& CreateDesc)
+{
+	FD3D11Texture* Texture = BeginCreateTextureInternal(CreateDesc);
+
+	TArray<D3D11_SUBRESOURCE_DATA> SubresourceData;
+	FillSubresourceData(SubresourceData, CreateDesc);
+
+	if (CreateDesc.IsTexture3D())
+	{
+		FinalizeCreateTexture3DInternal(Texture, TConstArrayView<D3D11_SUBRESOURCE_DATA>(SubresourceData));
+	}
+	else
+	{
+		FinalizeCreateTexture2DInternal(Texture, TConstArrayView<D3D11_SUBRESOURCE_DATA>(SubresourceData));
+	}
+
+	return Texture;
 }
 
-/*-----------------------------------------------------------------------------
-	2D texture support.
------------------------------------------------------------------------------*/
 
 FTextureRHIRef FD3D11DynamicRHI::RHICreateTexture(FRHICommandListBase&, const FRHITextureCreateDesc& CreateDesc)
 {
-	return CreateDesc.IsTexture3D()
-		? CreateD3D11Texture3D(CreateDesc)
-		: CreateD3D11Texture2D(CreateDesc);
+	FD3D11Texture* Texture = CreateTextureInternal(CreateDesc);
+
+	if (CreateDesc.BulkData)
+	{
+		CreateDesc.BulkData->Discard();
+	}
+
+	return Texture;
 }
 
 FTextureRHIRef FD3D11DynamicRHI::RHIAsyncCreateTexture2D(uint32 SizeX, uint32 SizeY, uint8 Format, uint32 NumMips, ETextureCreateFlags Flags, ERHIAccess InResourceState, void** InitialMipData, uint32 NumInitialMips, const TCHAR* DebugName, FGraphEventRef& OutCompletionEvent)
@@ -838,14 +859,15 @@ FTextureRHIRef FD3D11DynamicRHI::RHIAsyncCreateTexture2D(uint32 SizeX, uint32 Si
 		SubresourceData[MipIndex].SysMemSlicePitch = MipSize;
 	}
 
-	const FRHITextureCreateDesc Desc =
+	const FRHITextureCreateDesc CreateDesc =
 		FRHITextureCreateDesc::Create2D(DebugName, SizeX, SizeY, (EPixelFormat)Format)
 		.SetClearValue(FClearValueBinding::None)
 		.SetFlags(Flags)
 		.SetNumMips(NumMips)
 		.DetermineInititialState();
 
-	FTextureRHIRef Texture = CreateD3D11Texture2D(Desc, SubresourceData);
+	FD3D11Texture* Texture = BeginCreateTextureInternal(CreateDesc);
+	FinalizeCreateTexture2DInternal(Texture, SubresourceData);
 
 	if (TempBufferSize != ZeroBufferSize)
 	{
@@ -857,12 +879,6 @@ FTextureRHIRef FD3D11DynamicRHI::RHIAsyncCreateTexture2D(uint32 SizeX, uint32 Si
 	return Texture;
 }
 
-/**
- * Computes the size in memory required by a given texture.
- *
- * @param	TextureRHI		- Texture we want to know the size of
- * @return					- Size in Bytes
- */
 uint32 FD3D11DynamicRHI::RHIComputeMemorySize(FRHITexture* TextureRHI)
 {
 	if(!TextureRHI)
@@ -874,17 +890,6 @@ uint32 FD3D11DynamicRHI::RHIComputeMemorySize(FRHITexture* TextureRHI)
 	return Texture->GetMemorySize();
 }
 
-/**
- * Asynchronous texture copy helper
- *
- * @param NewTexture2DRHI		- Texture to reallocate
- * @param Texture2DRHI			- Texture to reallocate
- * @param NewMipCount			- New number of mip-levels
- * @param NewSizeX				- New width, in pixels
- * @param NewSizeY				- New height, in pixels
- * @param RequestStatus			- Will be decremented by 1 when the reallocation is complete (success or failure).
- * @return						- New reference to the texture, or an invalid reference upon failure
- */
 void FD3D11DynamicRHI::RHIAsyncCopyTexture2DCopy(FRHITexture* NewTexture2DRHI, FRHITexture* Texture2DRHI, int32 NewMipCount, int32 NewSizeX, int32 NewSizeY, FThreadSafeCounter* RequestStatus)
 {
 	FD3D11Texture* Texture2D = ResourceCast(Texture2DRHI);
@@ -915,22 +920,6 @@ void FD3D11DynamicRHI::RHIAsyncCopyTexture2DCopy(FRHITexture* NewTexture2DRHI, F
 	RequestStatus->Decrement();
 }
 
-
-/**
- * Starts an asynchronous texture reallocation. It may complete immediately if the reallocation
- * could be performed without any reshuffling of texture memory, or if there isn't enough memory.
- * The specified status counter will be decremented by 1 when the reallocation is complete (success or failure).
- *
- * Returns a new reference to the texture, which will represent the new mip count when the reallocation is complete.
- * RHIGetAsyncReallocateTexture2DStatus() can be used to check the status of an ongoing or completed reallocation.
- *
- * @param Texture2D		- Texture to reallocate
- * @param NewMipCount	- New number of mip-levels
- * @param NewSizeX		- New width, in pixels
- * @param NewSizeY		- New height, in pixels
- * @param RequestStatus	- Will be decremented by 1 when the reallocation is complete (success or failure).
- * @return				- New reference to the texture, or an invalid reference upon failure
- */
 FTextureRHIRef FD3D11DynamicRHI::RHIAsyncReallocateTexture2D(FRHITexture* Texture2DRHI, int32 NewMipCount, int32 NewSizeX, int32 NewSizeY, FThreadSafeCounter* RequestStatus)
 {
 	FD3D11Texture* Texture2D = ResourceCast(Texture2DRHI);
@@ -944,7 +933,7 @@ FTextureRHIRef FD3D11DynamicRHI::RHIAsyncReallocateTexture2D(FRHITexture* Textur
 	CreateDesc.NumMips = NewMipCount;
 
 	// Allocate a new texture.
-	FD3D11Texture* NewTexture2D = CreateD3D11Texture2D(CreateDesc);
+	FD3D11Texture* NewTexture2D = CreateTextureInternal(CreateDesc);
 
 	RHIAsyncCopyTexture2DCopy(NewTexture2D, Texture2D, NewMipCount, NewSizeX, NewSizeY, RequestStatus);
 
@@ -976,7 +965,7 @@ FTextureRHIRef FD3D11DynamicRHI::AsyncReallocateTexture2D_RenderThread(
 		CreateDesc.Extent = FIntPoint(NewSizeX, NewSizeY);
 		CreateDesc.NumMips = NewMipCount;
 
-		NewTexture2D = CreateD3D11Texture2D(CreateDesc);
+		NewTexture2D = CreateTextureInternal(CreateDesc);
 
 		RunOnRHIThread([this, NewTexture2D, Texture2D, NewMipCount, NewSizeX, NewSizeY, RequestStatus]()
 		{
@@ -1314,7 +1303,7 @@ FD3D11Texture* FD3D11DynamicRHI::CreateTextureFromResource(bool bTextureArray, b
 
 	const bool bSRGB = EnumHasAnyFlags(TexCreateFlags, TexCreate_SRGB);
 
-	const DXGI_FORMAT PlatformResourceFormat = UE::DXGIUtilities::GetPlatformTextureResourceFormat((DXGI_FORMAT)GPixelFormats[Format].PlatformFormat, TexCreateFlags);
+	const DXGI_FORMAT PlatformResourceFormat = UE::DXGIUtilities::GetPlatformTextureResourceFormat(Format, TexCreateFlags);
 	const DXGI_FORMAT PlatformShaderResourceFormat = UE::DXGIUtilities::FindShaderResourceFormat(PlatformResourceFormat, bSRGB);
 	const DXGI_FORMAT PlatformRenderTargetFormat = UE::DXGIUtilities::FindShaderResourceFormat(PlatformResourceFormat, bSRGB);
 
@@ -1572,23 +1561,27 @@ FTextureRHIRef FD3D11DynamicRHI::RHICreateTextureCubeFromResource(EPixelFormat F
 	return CreateTextureFromResource(false, true, Format, TexCreateFlags, ClearValueBinding, TextureResource);
 }
 
-FD3D11Texture::FD3D11Texture(
-	  const FRHITextureCreateDesc& InDesc
-	, ID3D11Resource* InResource
-	, ID3D11ShaderResourceView* InShaderResourceView
-	, int32 InRTVArraySize
-	, bool bInCreatedRTVsPerSlice
-	, TConstArrayView<TRefCountPtr<ID3D11RenderTargetView>> InRenderTargetViews
-	, TConstArrayView<TRefCountPtr<ID3D11DepthStencilView>> InDepthStencilViews
-	)
-	: FRHITexture         (InDesc)
-	, Resource            (InResource)
-	, ShaderResourceView  (InShaderResourceView)
-	, RenderTargetViews   (InRenderTargetViews)
-	, RTVArraySize        (InRTVArraySize)
-	, bCreatedRTVsPerSlice(bInCreatedRTVsPerSlice)
-	, bAlias              (false)
+FD3D11Texture::FD3D11Texture(const FRHITextureCreateDesc& InDesc)
+	: FRHITexture(InDesc)
+	, bCreatedRTVsPerSlice(false)
+	, bAlias(false)
+{
+}
+
+void FD3D11Texture::FinalizeCreation(
+	ID3D11Resource* InResource,
+	ID3D11ShaderResourceView* InShaderResourceView,
+	int32 InRTVArraySize,
+	bool bInCreatedRTVsPerSlice,
+	TConstArrayView<TRefCountPtr<ID3D11RenderTargetView>> InRenderTargetViews,
+	TConstArrayView<TRefCountPtr<ID3D11DepthStencilView>> InDepthStencilViews)
 {
+	Resource = InResource;
+	ShaderResourceView = InShaderResourceView;
+	RenderTargetViews = InRenderTargetViews;
+	RTVArraySize = InRTVArraySize;
+	bCreatedRTVsPerSlice = bInCreatedRTVsPerSlice;
+
 	// Set the DSVs for all the access type combinations
 	if (InDepthStencilViews.Num())
 	{
diff --git a/Engine/Source/Runtime/Windows/D3D11RHI/Public/D3D11Resources.h b/Engine/Source/Runtime/Windows/D3D11RHI/Public/D3D11Resources.h
--- a/Engine/Source/Runtime/Windows/D3D11RHI/Public/D3D11Resources.h
+++ b/Engine/Source/Runtime/Windows/D3D11RHI/Public/D3D11Resources.h
@@ -153,8 +153,8 @@ private:
 class FD3D11Texture final : public FRHITexture, public FD3D11ViewableResource
 {
 public:
-	D3D11RHI_API explicit FD3D11Texture(
-		const FRHITextureCreateDesc& InDesc,
+	FD3D11Texture(const FRHITextureCreateDesc& InDesc);
+	void FinalizeCreation(
 		ID3D11Resource* InResource,
 		ID3D11ShaderResourceView* InShaderResourceView,
 		int32 InRTVArraySize,
@@ -163,6 +163,19 @@ public:
 		TConstArrayView<TRefCountPtr<ID3D11DepthStencilView>> InDepthStencilViews
 	);
 
+	FD3D11Texture(
+		const FRHITextureCreateDesc& InDesc,
+		ID3D11Resource* InResource,
+		ID3D11ShaderResourceView* InShaderResourceView,
+		int32 InRTVArraySize,
+		bool bInCreatedRTVsPerSlice,
+		TConstArrayView<TRefCountPtr<ID3D11RenderTargetView>> InRenderTargetViews,
+		TConstArrayView<TRefCountPtr<ID3D11DepthStencilView>> InDepthStencilViews)
+		: FD3D11Texture(InDesc)
+	{
+		FinalizeCreation(InResource, InShaderResourceView, InRTVArraySize, bInCreatedRTVsPerSlice, InRenderTargetViews, InDepthStencilViews);
+	}
+
 	enum EAliasResourceParam { CreateAlias };
 	D3D11RHI_API explicit FD3D11Texture(FD3D11Texture const& Other, const FString& Name, EAliasResourceParam);
 	D3D11RHI_API void AliasResource(FD3D11Texture const& Other);
@@ -300,7 +313,7 @@ private:
 	/** A depth-stencil targetable view of the texture. */
 	TRefCountPtr<ID3D11DepthStencilView> DepthStencilViews[FExclusiveDepthStencil::MaxIndex];
 
-	int32 RTVArraySize;
+	int32 RTVArraySize = 0;
 
 	uint8 bCreatedRTVsPerSlice : 1;
 	uint8 bAlias : 1;
