More RHI cleanup
- Adding RHITextureUtils.h which currently contains the logic for calculating texture sizes and subresource offsets.
- Changed some D3D11RHI texture creation code to use these new utilities.
- Reworking D3D11RHI texture creation to operate on initial data as a view of the data rather than BulkData being passed in. This will make it easier to implement the Initializer refactor.

diff --git a/Engine/Source/Runtime/RHI/Public/RHITextureUtils.h b/Engine/Source/Runtime/RHI/Public/RHITextureUtils.h
new file mode 100644
--- /dev/null
+++ b/Engine/Source/Runtime/RHI/Public/RHITextureUtils.h
@@ -0,0 +1,99 @@
+// Copyright Epic Games, Inc. All Rights Reserved.
+
+#pragma once
+
+#include "RHIResources.h"
+
+namespace UE::RHITextureUtils
+{
+	static uint32 CalculateMipBlockCount(uint32 Extent, uint32 MipIndex, uint32 BlockSize)
+	{
+		return FMath::Max<uint32>(1, FMath::DivideAndRoundUp(Extent >> MipIndex, BlockSize));
+	}
+
+	static FUintVector3 CalculateMipBlockCounts(const FRHITextureDesc& Desc, uint32 MipIndex, const FPixelFormatInfo& PixelFormat)
+	{
+		const uint32 NumBlocksX = UE::RHITextureUtils::CalculateMipBlockCount(Desc.Extent.X, MipIndex, PixelFormat.BlockSizeX);
+		const uint32 NumBlocksY = UE::RHITextureUtils::CalculateMipBlockCount(Desc.Extent.Y, MipIndex, PixelFormat.BlockSizeY);
+		const uint32 NumBlocksZ = UE::RHITextureUtils::CalculateMipBlockCount(Desc.Depth,    MipIndex, PixelFormat.BlockSizeZ);
+
+		return FUintVector3(NumBlocksX, NumBlocksY, NumBlocksZ);
+	}
+
+	static FUintVector3 CalculateMipBlockCounts(const FRHITextureDesc& Desc, uint32 MipIndex)
+	{
+		const FPixelFormatInfo& PixelFormat = GPixelFormats[Desc.Format];
+		return CalculateMipBlockCounts(Desc, MipIndex, PixelFormat);
+	}
+
+	static uint64 CalculateTextureMipSize(const FRHITextureDesc& Desc, uint32 MipIndex, uint64& OutStride)
+	{
+		const FPixelFormatInfo& PixelFormat = GPixelFormats[Desc.Format];
+
+		const FUintVector3 BlockCounts = CalculateMipBlockCounts(Desc, MipIndex, PixelFormat);
+		const uint64 MipStride = BlockCounts.X * PixelFormat.BlockBytes;
+		const uint64 SubresourceSize = MipStride * BlockCounts.Y * BlockCounts.Z;
+
+		OutStride = MipStride;
+		return SubresourceSize;
+	}
+
+	static uint64 CalculateTexturePlaneSize(const FRHITextureDesc& Desc)
+	{
+		uint64 PlaneSize = 0;
+		uint64 MipStride = 0;
+		for (uint32 MipIndex = 0; MipIndex < Desc.NumMips; MipIndex++)
+		{
+			PlaneSize += CalculateTextureMipSize(Desc, MipIndex, MipStride);
+		}
+		return PlaneSize;
+	}
+
+	static uint64 CalculateTextureSize(const FRHITextureDesc& Desc)
+	{
+		const uint64 PlaneSize = CalculateTexturePlaneSize(Desc);
+		const uint64 FaceCount = Desc.IsTextureCube() ? 6 : 1;
+		return PlaneSize * Desc.ArraySize * Desc.Depth * FaceCount;
+	}
+
+	inline void SubresourceIndexToIndices(uint32 SubresourceIndex, uint32 NumMips, uint32 NumSlices, uint32& OutFaceIndex, uint32& OutArrayIndex, uint32& OutMipIndex)
+	{
+		OutMipIndex = (SubresourceIndex % NumMips);
+		OutArrayIndex = (SubresourceIndex / NumMips) % NumSlices;
+		OutFaceIndex = (SubresourceIndex / (NumMips % NumSlices));
+	}
+
+	static uint64 CalculateSubresourceOffset(const FRHITextureDesc& Desc, int32 FaceIndex, int32 ArrayIndex, int32 MipIndex, uint64& OutStride, uint64& OutSize)
+	{
+		const uint64 FaceCount = Desc.IsTextureCube() ? 6 : 1;
+		const uint64 PlaneIndex = FaceIndex + ArrayIndex * FaceCount;
+
+		uint64 PlaneSize = 0;
+		if (PlaneIndex > 0)
+		{
+			PlaneSize = CalculateTexturePlaneSize(Desc);
+		}
+
+		const uint64 PlaneOffset = PlaneIndex * PlaneSize;
+
+		uint64 MipOffset = 0;
+		uint64 MipStride = 0;
+
+		for (int32 Index = 0; Index <= MipIndex; Index++)
+		{
+			const uint64 MipSize = CalculateTextureMipSize(Desc, Index, MipStride);
+
+			if (Index == MipIndex)
+			{
+				OutSize = MipSize;
+				OutStride = MipStride;
+			}
+			else
+			{
+				MipOffset += MipSize;
+			}
+		}
+
+		return PlaneOffset + MipOffset;
+	}
+}
diff --git a/Engine/Source/Runtime/RHI/Public/RHIUtilities.h b/Engine/Source/Runtime/RHI/Public/RHIUtilities.h
--- a/Engine/Source/Runtime/RHI/Public/RHIUtilities.h
+++ b/Engine/Source/Runtime/RHI/Public/RHIUtilities.h
@@ -14,14 +14,6 @@
 #include "RHICommandList.h"
 #include "Tasks/Task.h"
 
-namespace UE::RHI
-{
-	static uint32 CalculateMipBlockCount(uint32 Extent, uint32 MipIndex, uint32 BlockSize)
-	{
-		return FMath::Max<uint32>(1, FMath::DivideAndRoundUp(Extent >> MipIndex, BlockSize));
-	}
-}
-
 static inline EPixelFormat GetBlockCompressedFormatUAVAliasFormat(EPixelFormat Format)
 {
 	switch (Format)
diff --git a/Engine/Source/Runtime/Windows/D3D11RHI/Private/D3D11RHIPrivate.h b/Engine/Source/Runtime/Windows/D3D11RHI/Private/D3D11RHIPrivate.h
--- a/Engine/Source/Runtime/Windows/D3D11RHI/Private/D3D11RHIPrivate.h
+++ b/Engine/Source/Runtime/Windows/D3D11RHI/Private/D3D11RHIPrivate.h
@@ -463,9 +463,10 @@ public:
 	virtual bool RHIGetTextureMemoryVisualizeData(FColor* TextureData,int32 SizeX,int32 SizeY,int32 Pitch,int32 PixelSize) final override;
 
 	[[nodiscard]] FD3D11Texture* BeginCreateTextureInternal(const FRHITextureCreateDesc& CreateDesc);
-	[[nodiscard]] void  FinalizeCreateTexture2DInternal(FD3D11Texture* Texture, TConstArrayView<D3D11_SUBRESOURCE_DATA> InitialData);
-	[[nodiscard]] void  FinalizeCreateTexture3DInternal(FD3D11Texture* Texture, TConstArrayView<D3D11_SUBRESOURCE_DATA> InitialData);
-	[[nodiscard]] FD3D11Texture* CreateTextureInternal(const FRHITextureCreateDesc& CreateDesc);
+	void  FinalizeCreateTexture2DInternal(FD3D11Texture* Texture, TConstArrayView<D3D11_SUBRESOURCE_DATA> InitialData);
+	void  FinalizeCreateTexture3DInternal(FD3D11Texture* Texture, TConstArrayView<D3D11_SUBRESOURCE_DATA> InitialData);
+	FD3D11Texture* FinalizeCreateTextureInternal(FD3D11Texture* Texture, const FRHITextureDesc& InDesc, TConstArrayView<uint8> InitialData);
+	[[nodiscard]] FD3D11Texture* CreateTextureInternal(const FRHITextureCreateDesc& CreateDesc, TConstArrayView<uint8> InitialData);
 
 	virtual FTextureRHIRef RHICreateTexture(FRHICommandListBase& RHICmdList, const FRHITextureCreateDesc& CreateDesc) final override;
 	virtual FTextureRHIRef RHIAsyncCreateTexture2D(uint32 SizeX, uint32 SizeY, uint8 Format, uint32 NumMips, ETextureCreateFlags Flags, ERHIAccess InResourceState, void** InitialMipData, uint32 NumInitialMips, const TCHAR* DebugName, FGraphEventRef& OutCompletionEvent) final override;
diff --git a/Engine/Source/Runtime/Windows/D3D11RHI/Private/D3D11Texture.cpp b/Engine/Source/Runtime/Windows/D3D11RHI/Private/D3D11Texture.cpp
--- a/Engine/Source/Runtime/Windows/D3D11RHI/Private/D3D11Texture.cpp
+++ b/Engine/Source/Runtime/Windows/D3D11RHI/Private/D3D11Texture.cpp
@@ -20,6 +20,7 @@
 #include "ProfilingDebugging/MemoryTrace.h"
 #include "RHICoreStats.h"
 #include "RHIUtilities.h"
+#include "RHITextureUtils.h"
 
 int64 FD3D11GlobalStats::GDedicatedVideoMemory = 0;
 int64 FD3D11GlobalStats::GDedicatedSystemMemory = 0;
@@ -742,14 +743,13 @@ void FD3D11DynamicRHI::FinalizeCreateTexture3DInternal(FD3D11Texture* Texture, T
 	);
 }
 
-static const D3D11_SUBRESOURCE_DATA* FillSubresourceData(TArray<D3D11_SUBRESOURCE_DATA>& SubresourceData, const FRHITextureCreateDesc& CreateDesc)
+static const D3D11_SUBRESOURCE_DATA* FillSubresourceData(TArray<D3D11_SUBRESOURCE_DATA>& SubresourceData, const FRHITextureDesc& CreateDesc, TConstArrayView<uint8> InitialData)
 {
 	const D3D11_SUBRESOURCE_DATA* pSubresourceData = nullptr;
 
-	if (CreateDesc.BulkData)
+	if (InitialData.Num())
 	{
 		const FPixelFormatInfo& PixelFormat = GPixelFormats[CreateDesc.Format];
-		const uint8* Data = reinterpret_cast<const uint8*>(CreateDesc.BulkData->GetResourceBulkData());
 
 		const uint32 FaceCount = CreateDesc.IsTextureCube() ? 6 : 1;
 		const uint32 ArrayCount = CreateDesc.ArraySize * FaceCount;
@@ -767,18 +767,16 @@ static const D3D11_SUBRESOURCE_DATA* FillSubresourceData(TArray<D3D11_SUBRESOURC
 				const uint32 DataOffset = SliceOffset + MipOffset;
 				const uint32 SubResourceIndex = ArraySliceIndex * MipCount + MipIndex;
 
-				const uint32 NumBlocksX = UE::RHI::CalculateMipBlockCount(CreateDesc.Extent.X, MipIndex, PixelFormat.BlockSizeX);
-				const uint32 NumBlocksY = UE::RHI::CalculateMipBlockCount(CreateDesc.Extent.Y, MipIndex, PixelFormat.BlockSizeY);
-				const uint32 NumBlocksZ = UE::RHI::CalculateMipBlockCount(CreateDesc.Depth,    MipIndex, PixelFormat.BlockSizeZ);
+				const FUintVector3 BlockCounts = UE::RHITextureUtils::CalculateMipBlockCounts(CreateDesc, MipIndex, PixelFormat);
 
-				const uint32 RowPitch = NumBlocksX * PixelFormat.BlockBytes;
-				const uint32 SlicePitch = NumBlocksY * RowPitch;
+				const uint32 RowPitch = BlockCounts.X * PixelFormat.BlockBytes;
+				const uint32 SlicePitch = BlockCounts.Y * RowPitch;
 
-				SubresourceData[SubResourceIndex].pSysMem          = &Data[DataOffset];
+				SubresourceData[SubResourceIndex].pSysMem          = &InitialData[DataOffset];
 				SubresourceData[SubResourceIndex].SysMemPitch      = RowPitch;
 				SubresourceData[SubResourceIndex].SysMemSlicePitch = SlicePitch;
 
-				MipOffset += SlicePitch * NumBlocksZ;
+				MipOffset += SlicePitch * BlockCounts.Z;
 			}
 
 			SliceOffset += MipOffset;
@@ -790,14 +788,12 @@ static const D3D11_SUBRESOURCE_DATA* FillSubresourceData(TArray<D3D11_SUBRESOURC
 	return pSubresourceData;
 }
 
-FD3D11Texture* FD3D11DynamicRHI::CreateTextureInternal(const FRHITextureCreateDesc& CreateDesc)
+FD3D11Texture* FD3D11DynamicRHI::FinalizeCreateTextureInternal(FD3D11Texture* Texture, const FRHITextureDesc& InDesc, TConstArrayView<uint8> InitialData)
 {
-	FD3D11Texture* Texture = BeginCreateTextureInternal(CreateDesc);
-
 	TArray<D3D11_SUBRESOURCE_DATA> SubresourceData;
-	FillSubresourceData(SubresourceData, CreateDesc);
+	FillSubresourceData(SubresourceData, InDesc, InitialData);
 
-	if (CreateDesc.IsTexture3D())
+	if (InDesc.IsTexture3D())
 	{
 		FinalizeCreateTexture3DInternal(Texture, TConstArrayView<D3D11_SUBRESOURCE_DATA>(SubresourceData));
 	}
@@ -809,11 +805,21 @@ FD3D11Texture* FD3D11DynamicRHI::CreateTextureInternal(const FRHITextureCreateDe
 	return Texture;
 }
 
+FD3D11Texture* FD3D11DynamicRHI::CreateTextureInternal(const FRHITextureCreateDesc& CreateDesc, TConstArrayView<uint8> InitialData)
+{
+	FD3D11Texture* Texture = BeginCreateTextureInternal(CreateDesc);
+	return FinalizeCreateTextureInternal(Texture, CreateDesc, InitialData);
+}
 
 FTextureRHIRef FD3D11DynamicRHI::RHICreateTexture(FRHICommandListBase&, const FRHITextureCreateDesc& CreateDesc)
 {
-	FD3D11Texture* Texture = CreateTextureInternal(CreateDesc);
+	TConstArrayView<uint8> InitialData;
+	if (CreateDesc.BulkData)
+	{
+		InitialData = CreateDesc.BulkData->GetBulkDataView<uint8>();
+	}
 
+	FD3D11Texture* Texture = CreateTextureInternal(CreateDesc, InitialData);
 	if (CreateDesc.BulkData)
 	{
 		CreateDesc.BulkData->Discard();
@@ -824,26 +830,29 @@ FTextureRHIRef FD3D11DynamicRHI::RHICreateTexture(FRHICommandListBase&, const FR
 
 FTextureRHIRef FD3D11DynamicRHI::RHIAsyncCreateTexture2D(uint32 SizeX, uint32 SizeY, uint8 Format, uint32 NumMips, ETextureCreateFlags Flags, ERHIAccess InResourceState, void** InitialMipData, uint32 NumInitialMips, const TCHAR* DebugName, FGraphEventRef& OutCompletionEvent)
 {
+	const FPixelFormatInfo& PixelFormat = GPixelFormats[Format];
+
 	TArray<D3D11_SUBRESOURCE_DATA, TInlineAllocator<12>> SubresourceData;
 	SubresourceData.SetNumUninitialized(NumMips);
 
 	for (uint32 MipIndex = 0; MipIndex < NumInitialMips; ++MipIndex)
 	{
-		uint32 NumBlocksX = FMath::Max<uint32>(1, ((SizeX >> MipIndex) +GPixelFormats[Format].BlockSizeX-1) / GPixelFormats[Format].BlockSizeX);
-		uint32 NumBlocksY = FMath::Max<uint32>(1, ((SizeY >> MipIndex) +GPixelFormats[Format].BlockSizeY-1) / GPixelFormats[Format].BlockSizeY);
+		const uint32 NumBlocksX = UE::RHITextureUtils::CalculateMipBlockCount(SizeX, MipIndex, PixelFormat.BlockSizeX);
+		const uint32 NumBlocksY = UE::RHITextureUtils::CalculateMipBlockCount(SizeY, MipIndex, PixelFormat.BlockSizeY);
 
 		SubresourceData[MipIndex].pSysMem = InitialMipData[MipIndex];
-		SubresourceData[MipIndex].SysMemPitch = NumBlocksX * GPixelFormats[Format].BlockBytes;
-		SubresourceData[MipIndex].SysMemSlicePitch = NumBlocksX * NumBlocksY * GPixelFormats[Format].BlockBytes;
+		SubresourceData[MipIndex].SysMemPitch = NumBlocksX * PixelFormat.BlockBytes;
+		SubresourceData[MipIndex].SysMemSlicePitch = NumBlocksX * NumBlocksY * PixelFormat.BlockBytes;
 	}
 
 	void* TempBuffer = ZeroBuffer;
 	uint32 TempBufferSize = ZeroBufferSize;
 	for (uint32 MipIndex = NumInitialMips; MipIndex < NumMips; ++MipIndex)
 	{
-		uint32 NumBlocksX = FMath::Max<uint32>(1, ((SizeX >> MipIndex) +GPixelFormats[Format].BlockSizeX-1) / GPixelFormats[Format].BlockSizeX);
-		uint32 NumBlocksY = FMath::Max<uint32>(1, ((SizeY >> MipIndex) +GPixelFormats[Format].BlockSizeY-1) / GPixelFormats[Format].BlockSizeY);
-		uint32 MipSize = NumBlocksX * NumBlocksY * GPixelFormats[Format].BlockBytes;
+		const uint32 NumBlocksX = UE::RHITextureUtils::CalculateMipBlockCount(SizeX, MipIndex, PixelFormat.BlockSizeX);
+		const uint32 NumBlocksY = UE::RHITextureUtils::CalculateMipBlockCount(SizeY, MipIndex, PixelFormat.BlockSizeY);
+
+		uint32 MipSize = NumBlocksX * NumBlocksY * PixelFormat.BlockBytes;
 
 		if (MipSize > TempBufferSize)
 		{
@@ -855,7 +864,7 @@ FTextureRHIRef FD3D11DynamicRHI::RHIAsyncCreateTexture2D(uint32 SizeX, uint32 Si
 		}
 
 		SubresourceData[MipIndex].pSysMem = TempBuffer;
-		SubresourceData[MipIndex].SysMemPitch = NumBlocksX * GPixelFormats[Format].BlockBytes;
+		SubresourceData[MipIndex].SysMemPitch = NumBlocksX * PixelFormat.BlockBytes;
 		SubresourceData[MipIndex].SysMemSlicePitch = MipSize;
 	}
 
@@ -933,7 +942,7 @@ FTextureRHIRef FD3D11DynamicRHI::RHIAsyncReallocateTexture2D(FRHITexture* Textur
 	CreateDesc.NumMips = NewMipCount;
 
 	// Allocate a new texture.
-	FD3D11Texture* NewTexture2D = CreateTextureInternal(CreateDesc);
+	FD3D11Texture* NewTexture2D = CreateTextureInternal(CreateDesc, {});
 
 	RHIAsyncCopyTexture2DCopy(NewTexture2D, Texture2D, NewMipCount, NewSizeX, NewSizeY, RequestStatus);
 
@@ -965,7 +974,7 @@ FTextureRHIRef FD3D11DynamicRHI::AsyncReallocateTexture2D_RenderThread(
 		CreateDesc.Extent = FIntPoint(NewSizeX, NewSizeY);
 		CreateDesc.NumMips = NewMipCount;
 
-		NewTexture2D = CreateTextureInternal(CreateDesc);
+		NewTexture2D = CreateTextureInternal(CreateDesc, {});
 
 		RunOnRHIThread([this, NewTexture2D, Texture2D, NewMipCount, NewSizeX, NewSizeY, RequestStatus]()
 		{
