Refactor of audio device swap for XAudio2 and WASAPI backends.

diff --git a/Engine/Source/Runtime/AudioDeviceEnumeration/Windows/WindowsMMDeviceEnumeration/Private/ConversionHelpers.h b/Engine/Source/Runtime/AudioDeviceEnumeration/Windows/WindowsMMDeviceEnumeration/Private/ConversionHelpers.h
--- a/Engine/Source/Runtime/AudioDeviceEnumeration/Windows/WindowsMMDeviceEnumeration/Private/ConversionHelpers.h
+++ b/Engine/Source/Runtime/AudioDeviceEnumeration/Windows/WindowsMMDeviceEnumeration/Private/ConversionHelpers.h
@@ -4,8 +4,6 @@
 
 #include "HAL/Platform.h"
 
-#if PLATFORM_WINDOWS
-
 #include "AudioMixer.h"	
 #include "Windows/AllowWindowsPlatformTypes.h"
 
@@ -52,5 +50,3 @@ namespace Audio
 	}
 
 }
-
-#endif //PLATFORM_WINDOWS
diff --git a/Engine/Source/Runtime/AudioDeviceEnumeration/Windows/WindowsMMDeviceEnumeration/Private/WindowsMMDeviceEnumerationModule.cpp b/Engine/Source/Runtime/AudioDeviceEnumeration/Windows/WindowsMMDeviceEnumeration/Private/WindowsMMDeviceEnumerationModule.cpp
--- a/Engine/Source/Runtime/AudioDeviceEnumeration/Windows/WindowsMMDeviceEnumeration/Private/WindowsMMDeviceEnumerationModule.cpp
+++ b/Engine/Source/Runtime/AudioDeviceEnumeration/Windows/WindowsMMDeviceEnumeration/Private/WindowsMMDeviceEnumerationModule.cpp
@@ -3,7 +3,6 @@
 #include "Modules/ModuleManager.h"
 #include "HAL/Platform.h"
 
-#if PLATFORM_WINDOWS
 #include "Windows/AllowWindowsPlatformTypes.h"
 
 THIRD_PARTY_INCLUDES_START
@@ -12,7 +11,6 @@ THIRD_PARTY_INCLUDES_START
 THIRD_PARTY_INCLUDES_END
 
 #include "Windows/HideWindowsPlatformTypes.h"
-#endif // PLATFORM_WINDOWS
 
 #include "WindowsMMNotificationClient.h"
 #include "WindowsMMDeviceInfoCache.h"
diff --git a/Engine/Source/Runtime/AudioDeviceEnumeration/Windows/WindowsMMDeviceEnumeration/Private/WindowsMMDeviceInfoCache.cpp b/Engine/Source/Runtime/AudioDeviceEnumeration/Windows/WindowsMMDeviceEnumeration/Private/WindowsMMDeviceInfoCache.cpp
--- a/Engine/Source/Runtime/AudioDeviceEnumeration/Windows/WindowsMMDeviceEnumeration/Private/WindowsMMDeviceInfoCache.cpp
+++ b/Engine/Source/Runtime/AudioDeviceEnumeration/Windows/WindowsMMDeviceEnumeration/Private/WindowsMMDeviceInfoCache.cpp
@@ -2,8 +2,6 @@
 
 #include "WindowsMMDeviceInfoCache.h"
 
-#if PLATFORM_WINDOWS
-
 #include "Windows/AllowWindowsPlatformTypes.h"
 #include "Windows/AllowWindowsPlatformAtomics.h"
 
@@ -17,6 +15,7 @@ THIRD_PARTY_INCLUDES_END
 
 #include "WindowsMMStringUtils.h"
 #include "ConversionHelpers.h"
+#include "Misc/ScopeRWLock.h"
 #include "WindowsMMCvarUtils.h"
 #include "WindowsMMDeviceEnumerationLog.h"
 
@@ -294,6 +293,8 @@ bool FWindowsMMDeviceCache::EnumerateHardwareTopology(const TComPtr<IMMDevice>&
 {
 	SCOPED_NAMED_EVENT(FWindowsMMDeviceCache_EnumerateHardwareTopology, FColor::Blue);
 
+#if PLATFORM_WINDOWS
+
 	if (!DeviceEnumerator)
 	{
 		return false;
@@ -357,6 +358,9 @@ bool FWindowsMMDeviceCache::EnumerateHardwareTopology(const TComPtr<IMMDevice>&
 	OutInfo.FilterId = TempFilterId;
 
 	return true;
+#else
+	return false;
+#endif //PLATFORM_WINDOWS
 }
 
 void FWindowsMMDeviceCache::EnumerateEndpoints()
@@ -728,5 +732,3 @@ TOptional<Audio::FAudioPlatformDeviceInfo> FWindowsMMDeviceCache::FindActiveOutp
 }
 
 } // namespace Audio
-
-#endif //PLATFORM_WINDOWS
diff --git a/Engine/Source/Runtime/AudioDeviceEnumeration/Windows/WindowsMMDeviceEnumeration/Private/WindowsMMNotificationClient.cpp b/Engine/Source/Runtime/AudioDeviceEnumeration/Windows/WindowsMMDeviceEnumeration/Private/WindowsMMNotificationClient.cpp
--- a/Engine/Source/Runtime/AudioDeviceEnumeration/Windows/WindowsMMDeviceEnumeration/Private/WindowsMMNotificationClient.cpp
+++ b/Engine/Source/Runtime/AudioDeviceEnumeration/Windows/WindowsMMDeviceEnumeration/Private/WindowsMMNotificationClient.cpp
@@ -2,8 +2,6 @@
 
 #include "WindowsMMNotificationClient.h"
 
-#if PLATFORM_WINDOWS
-
 #include "Windows/AllowWindowsPlatformTypes.h"
 
 THIRD_PARTY_INCLUDES_START
@@ -19,6 +17,7 @@ THIRD_PARTY_INCLUDES_END
 
 namespace Audio
 {
+#if PLATFORM_WINDOWS
 	bool FWindowsMMNotificationClient::RegisterForSessionNotifications(const TComPtr<IMMDevice>& InDevice)
 	{
 		FScopeLock Lock(&SessionRegistrationCS);
@@ -112,6 +111,7 @@ namespace Audio
 	{
 		return S_OK;
 	}
+#endif //PLATFORM_WINDOWS
 
 	void FWindowsMMNotificationClient::UnRegisterDeviceDeviceChangedListener(Audio::IAudioMixerDeviceChangedListener* DeviceChangedListener)
 	{
@@ -434,19 +434,24 @@ namespace Audio
 
 	FWindowsMMNotificationClient::~FWindowsMMNotificationClient()
 	{
+#if PLATFORM_WINDOWS
 		UnregisterForSessionNotifications();
+#endif //PLATFORM_WINDOWS
 
 		if (DeviceEnumerator)
 		{
 			DeviceEnumerator->UnregisterEndpointNotificationCallback(this);
 		}
 
+#if PLATFORM_WINDOWS
 		if (bComInitialized)
 		{
 			FWindowsPlatformMisc::CoUninitialize();
 		}
+#endif //PLATFORM_WINDOWS
 	}
 
+#if PLATFORM_WINDOWS
 	void FWindowsMMNotificationClient::UnregisterForSessionNotifications()
 	{
 		FScopeLock Lock(&SessionRegistrationCS);
@@ -468,19 +473,21 @@ namespace Audio
 		// Reset this flag.
 		bHasDisconnectSessionHappened = false;
 	}
+#endif //PLATFORM_WINDOWS
 
 	FWindowsMMNotificationClient::FWindowsMMNotificationClient() 
-		: Ref(1)
-		, DeviceEnumerator(nullptr)
-		, bHasDisconnectSessionHappened(false)
 	{
+#if PLATFORM_WINDOWS
 		bComInitialized = FWindowsPlatformMisc::CoInitialize();
+#endif //PLATFORM_WINDOWS
+
 		HRESULT Result = CoCreateInstance(__uuidof(MMDeviceEnumerator), nullptr, CLSCTX_INPROC_SERVER, IID_PPV_ARGS(&DeviceEnumerator));
 		if (Result == S_OK)
 		{
 			DeviceEnumerator->RegisterEndpointNotificationCallback(this);
 		}
 
+#if PLATFORM_WINDOWS
 		// Register for session events from default endpoint.
 		if (DeviceEnumerator)
 		{
@@ -490,8 +497,7 @@ namespace Audio
 				RegisterForSessionNotifications(DefaultDevice);
 			}
 		}
+#endif //PLATFORM_WINDOWS
 	}
 
 }// namespace Audio
-
-#endif //PLATFORM_WINDOWS
diff --git a/Engine/Source/Runtime/AudioDeviceEnumeration/Windows/WindowsMMDeviceEnumeration/Private/WindowsMMStringUtils.cpp b/Engine/Source/Runtime/AudioDeviceEnumeration/Windows/WindowsMMDeviceEnumeration/Private/WindowsMMStringUtils.cpp
--- a/Engine/Source/Runtime/AudioDeviceEnumeration/Windows/WindowsMMDeviceEnumeration/Private/WindowsMMStringUtils.cpp
+++ b/Engine/Source/Runtime/AudioDeviceEnumeration/Windows/WindowsMMDeviceEnumeration/Private/WindowsMMStringUtils.cpp
@@ -5,18 +5,19 @@
 #include "Containers/UnrealString.h"
 #include "Misc/AssertionMacros.h"
 
-#if PLATFORM_WINDOWS
-
+#include "Windows/AllowWindowsPlatformTypes.h"
+THIRD_PARTY_INCLUDES_START
+#include <winerror.h>
+#include <audiopolicy.h>			// IAudioSessionEvents
 #if !NO_LOGGING
 // PSStringFromPropertyKey needs this lib, which we use for logging the Property GUIDs.
 #pragma comment(lib, "Propsys.lib")
 #endif //!NO_LOGGING
-
-#endif //PLATFORM_WINDOWS
+THIRD_PARTY_INCLUDES_END
+#include "Windows/HideWindowsPlatformTypes.h"
 
 namespace Audio
 {
-#if PLATFORM_WINDOWS
 	const TCHAR* ToString(AudioSessionDisconnectReason InDisconnectReason)
 	{
 #if !NO_LOGGING
@@ -76,6 +77,7 @@ namespace Audio
 
 	FString ToFString(const PROPERTYKEY InKey)
 	{
+#if PLATFORM_WINDOWS
 #if !NO_LOGGING
 #define IF_PROP_STRING(PROP) if(InKey.fmtid == PROP.fmtid) return TEXT(#PROP)
 
@@ -99,9 +101,9 @@ namespace Audio
 			return FString(KeyString);
 		}
 #endif //!NO_LOGGING
+#endif //PLATFORM_WINDOWS
 		return TEXT("Unknown");
 	}
-#endif //PLATFORM_WINDOWS
 
 	const TCHAR* ToString(EAudioDeviceRole InRole)
 	{
@@ -143,6 +145,77 @@ namespace Audio
 			ChannelString.Append(TEXT("|"));
 		}
 		return ChannelString;
-	}	
+	}
+
+	FString AudioClientErrorToFString(HRESULT InResult)
+	{
+#define CASE_AND_STRING(RESULT) case HRESULT(RESULT): return TEXT(#RESULT)
+
+		switch (InResult)
+		{
+		case HRESULT(REGDB_E_CLASSNOTREG):				return TEXT("REGDB_E_CLASSNOTREG");
+		case HRESULT(CLASS_E_NOAGGREGATION):			return TEXT("CLASS_E_NOAGGREGATION");
+		case HRESULT(E_NOINTERFACE):					return TEXT("E_NOINTERFACE");
+		case HRESULT(E_POINTER):						return TEXT("E_POINTER");
+		case HRESULT(E_INVALIDARG):						return TEXT("E_INVALIDARG");
+		case HRESULT(E_OUTOFMEMORY):					return TEXT("E_OUTOFMEMORY");
+
+		// AudioClient.h
+		CASE_AND_STRING(AUDCLNT_E_NOT_INITIALIZED);
+		CASE_AND_STRING(AUDCLNT_E_ALREADY_INITIALIZED);
+		CASE_AND_STRING(AUDCLNT_E_WRONG_ENDPOINT_TYPE);
+		CASE_AND_STRING(AUDCLNT_E_DEVICE_INVALIDATED);
+		CASE_AND_STRING(AUDCLNT_E_NOT_STOPPED);
+		CASE_AND_STRING(AUDCLNT_E_BUFFER_TOO_LARGE);
+		CASE_AND_STRING(AUDCLNT_E_OUT_OF_ORDER);
+		CASE_AND_STRING(AUDCLNT_E_UNSUPPORTED_FORMAT);
+		CASE_AND_STRING(AUDCLNT_E_INVALID_SIZE);
+		CASE_AND_STRING(AUDCLNT_E_DEVICE_IN_USE);
+		CASE_AND_STRING(AUDCLNT_E_BUFFER_OPERATION_PENDING);
+		CASE_AND_STRING(AUDCLNT_E_THREAD_NOT_REGISTERED);
+		CASE_AND_STRING(AUDCLNT_E_EXCLUSIVE_MODE_NOT_ALLOWED);
+		CASE_AND_STRING(AUDCLNT_E_ENDPOINT_CREATE_FAILED);
+		CASE_AND_STRING(AUDCLNT_E_SERVICE_NOT_RUNNING);
+		CASE_AND_STRING(AUDCLNT_E_EVENTHANDLE_NOT_EXPECTED);
+		CASE_AND_STRING(AUDCLNT_E_EXCLUSIVE_MODE_ONLY);
+		CASE_AND_STRING(AUDCLNT_E_BUFDURATION_PERIOD_NOT_EQUAL);
+		CASE_AND_STRING(AUDCLNT_E_EVENTHANDLE_NOT_SET);
+		CASE_AND_STRING(AUDCLNT_E_INCORRECT_BUFFER_SIZE);
+		CASE_AND_STRING(AUDCLNT_E_BUFFER_SIZE_ERROR);
+		CASE_AND_STRING(AUDCLNT_E_CPUUSAGE_EXCEEDED);
+		CASE_AND_STRING(AUDCLNT_E_BUFFER_ERROR);
+		CASE_AND_STRING(AUDCLNT_E_BUFFER_SIZE_NOT_ALIGNED);
+		CASE_AND_STRING(AUDCLNT_E_INVALID_DEVICE_PERIOD);
+		CASE_AND_STRING(AUDCLNT_E_INVALID_STREAM_FLAG);
+		CASE_AND_STRING(AUDCLNT_E_ENDPOINT_OFFLOAD_NOT_CAPABLE);
+		CASE_AND_STRING(AUDCLNT_E_OUT_OF_OFFLOAD_RESOURCES);
+		CASE_AND_STRING(AUDCLNT_E_OFFLOAD_MODE_ONLY);
+		CASE_AND_STRING(AUDCLNT_E_NONOFFLOAD_MODE_ONLY);
+		CASE_AND_STRING(AUDCLNT_E_RESOURCES_INVALIDATED);
+		CASE_AND_STRING(AUDCLNT_E_RAW_MODE_UNSUPPORTED);
+		CASE_AND_STRING(AUDCLNT_E_ENGINE_PERIODICITY_LOCKED);
+		CASE_AND_STRING(AUDCLNT_E_ENGINE_FORMAT_LOCKED);
+		CASE_AND_STRING(AUDCLNT_S_BUFFER_EMPTY);
+		CASE_AND_STRING(AUDCLNT_S_THREAD_ALREADY_REGISTERED);
+		CASE_AND_STRING(AUDCLNT_S_POSITION_STALLED);
+
+#undef CASE_AND_STRING
+
+		case HRESULT(0xe000020b):						return TEXT("ERROR_NO_SUCH_DEVINST");
+		case HRESULT(0x80010108):						return TEXT("RPC_E_DISCONNECTED");
+
+		default:
+		{
+			// We don't know this error, ask this system if it does.
+			TCHAR Buffer[1024] = { 0 };
+			FString Msg(FPlatformMisc::GetSystemErrorMessage(Buffer, UE_ARRAY_COUNT(Buffer), InResult));
+
+			// Anything to return? Otherwise "UNKNOWN"
+			return !Msg.IsEmpty() ?
+				Msg :
+				TEXT("UNKNOWN");
+		}
+		}
+	}
 }
 
diff --git a/Engine/Source/Runtime/AudioDeviceEnumeration/Windows/WindowsMMDeviceEnumeration/Public/ScopedCom.h b/Engine/Source/Runtime/AudioDeviceEnumeration/Windows/WindowsMMDeviceEnumeration/Public/ScopedCom.h
--- a/Engine/Source/Runtime/AudioDeviceEnumeration/Windows/WindowsMMDeviceEnumeration/Public/ScopedCom.h
+++ b/Engine/Source/Runtime/AudioDeviceEnumeration/Windows/WindowsMMDeviceEnumeration/Public/ScopedCom.h
@@ -3,14 +3,13 @@
 #pragma once
 
 #include "HAL/Platform.h"
-
-#if PLATFORM_WINDOWS
+#include "Windows/WindowsPlatformMisc.h"
 
 #include "Windows/AllowWindowsPlatformTypes.h"
-
 THIRD_PARTY_INCLUDES_START
 #include <combaseapi.h>
 THIRD_PARTY_INCLUDES_END
+#include "Windows/HideWindowsPlatformTypes.h"
 
 namespace Audio
 {
@@ -42,8 +41,28 @@ namespace Audio
 			}
 		}
 	};
-}
 
-#include "Windows/HideWindowsPlatformTypes.h"
+	class FScopedCoInitialize final
+	{
+	public:
+		FScopedCoInitialize()
+		{
+#if PLATFORM_WINDOWS
+			bCoInitialized = FPlatformMisc::CoInitialize(ECOMModel::Multithreaded);
+#endif	
+		}
 
-#endif // PLATFORM_WINDOWS
+		~FScopedCoInitialize()
+		{
+#if PLATFORM_WINDOWS
+			if (bCoInitialized)
+			{
+				FPlatformMisc::CoUninitialize();
+			}
+#endif
+		}
+	private:
+		bool bCoInitialized = false;
+	};
+
+}
diff --git a/Engine/Source/Runtime/AudioDeviceEnumeration/Windows/WindowsMMDeviceEnumeration/Public/WindowsMMDeviceInfoCache.h b/Engine/Source/Runtime/AudioDeviceEnumeration/Windows/WindowsMMDeviceEnumeration/Public/WindowsMMDeviceInfoCache.h
--- a/Engine/Source/Runtime/AudioDeviceEnumeration/Windows/WindowsMMDeviceEnumeration/Public/WindowsMMDeviceInfoCache.h
+++ b/Engine/Source/Runtime/AudioDeviceEnumeration/Windows/WindowsMMDeviceEnumeration/Public/WindowsMMDeviceInfoCache.h
@@ -4,12 +4,8 @@
 
 #include "HAL/Platform.h"
 
-#if PLATFORM_WINDOWS
-
 #include "AudioMixer.h"
-#include "Misc/ScopeRWLock.h"
 #include "ScopedCom.h"
-#include <atomic>
 
 #include "Windows/AllowWindowsPlatformTypes.h"
 #include "Windows/AllowWindowsPlatformAtomics.h"
@@ -112,5 +108,3 @@ struct WINDOWSMMDEVICEENUMERATION_API FWindowsMMDeviceCache : IAudioMixerDeviceC
 };
 
 }// namespace Audio
-
-#endif //PLATFORM_WINDOWS
diff --git a/Engine/Source/Runtime/AudioDeviceEnumeration/Windows/WindowsMMDeviceEnumeration/Public/WindowsMMNotificationClient.h b/Engine/Source/Runtime/AudioDeviceEnumeration/Windows/WindowsMMDeviceEnumeration/Public/WindowsMMNotificationClient.h
--- a/Engine/Source/Runtime/AudioDeviceEnumeration/Windows/WindowsMMDeviceEnumeration/Public/WindowsMMNotificationClient.h
+++ b/Engine/Source/Runtime/AudioDeviceEnumeration/Windows/WindowsMMDeviceEnumeration/Public/WindowsMMNotificationClient.h
@@ -4,8 +4,6 @@
 
 #include "HAL/Platform.h"
 
-#if PLATFORM_WINDOWS
-
 #include <atomic>
 #include "AudioMixer.h"	
 #include "Misc/ScopeRWLock.h"
@@ -27,16 +25,15 @@ namespace Audio
 {
 	class IAudioMixerDeviceChangedListener;
 
-	class FWindowsMMNotificationClient final : public IMMNotificationClient, public IAudioSessionEvents
+	class FWindowsMMNotificationClient final : public IMMNotificationClient
+#if PLATFORM_WINDOWS
+		, public IAudioSessionEvents
+#endif //PLATFORM_WINDOWS
 	{
 	public:
 		WINDOWSMMDEVICEENUMERATION_API FWindowsMMNotificationClient();
 		~FWindowsMMNotificationClient();
 
-		WINDOWSMMDEVICEENUMERATION_API bool RegisterForSessionNotifications(const TComPtr<IMMDevice>& InDevice);
-		WINDOWSMMDEVICEENUMERATION_API bool RegisterForSessionNotifications(const FString& InDeviceId);
-		WINDOWSMMDEVICEENUMERATION_API void UnregisterForSessionNotifications();
-
 		// TODO: Ideally we'd use the cache instead of ask for this.
 		WINDOWSMMDEVICEENUMERATION_API bool IsRenderDevice(const FString& InDeviceId) const;
 
@@ -63,6 +60,11 @@ namespace Audio
 		HRESULT STDMETHODCALLTYPE OnPropertyValueChanged(LPCWSTR pwstrDeviceId, const PROPERTYKEY key);
 		// End IMMNotificationClient overrides
 
+#if PLATFORM_WINDOWS
+		WINDOWSMMDEVICEENUMERATION_API bool RegisterForSessionNotifications(const TComPtr<IMMDevice>& InDevice);
+		WINDOWSMMDEVICEENUMERATION_API bool RegisterForSessionNotifications(const FString& InDeviceId);
+		WINDOWSMMDEVICEENUMERATION_API void UnregisterForSessionNotifications();
+
 		// Begin IAudioSessionEvents overrides
 		HRESULT STDMETHODCALLTYPE OnDisplayNameChanged(
 			LPCWSTR NewDisplayName,
@@ -95,20 +97,20 @@ namespace Audio
 		// End IAudioSessionEvents overrides
 
 	private:
-		LONG Ref;
-		TSet<Audio::IAudioMixerDeviceChangedListener*> Listeners;
-		FRWLock ListenersSetRwLock;
-
-		TComPtr<IMMDeviceEnumerator> DeviceEnumerator;
-
 		FCriticalSection SessionRegistrationCS;
 		TComPtr<IAudioSessionManager> SessionManager;
 		TComPtr<IAudioSessionControl> SessionControls;
 		TComPtr<IMMDevice> DeviceListeningToSessionEvents;
 
-		bool bComInitialized;
-		std::atomic<bool> bHasDisconnectSessionHappened;
-	};
-}// namespace Audio
+		bool bComInitialized = false;
+		std::atomic<bool> bHasDisconnectSessionHappened = false;
 
 #endif //PLATFORM_WINDOWS
+		
+	private:
+		LONG Ref = 1; // Start with self-reference
+		TSet<Audio::IAudioMixerDeviceChangedListener*> Listeners;
+		FRWLock ListenersSetRwLock;
+		TComPtr<IMMDeviceEnumerator> DeviceEnumerator = nullptr;
+	};
+}// namespace Audio
diff --git a/Engine/Source/Runtime/AudioDeviceEnumeration/Windows/WindowsMMDeviceEnumeration/Public/WindowsMMStringUtils.h b/Engine/Source/Runtime/AudioDeviceEnumeration/Windows/WindowsMMDeviceEnumeration/Public/WindowsMMStringUtils.h
--- a/Engine/Source/Runtime/AudioDeviceEnumeration/Windows/WindowsMMDeviceEnumeration/Public/WindowsMMStringUtils.h
+++ b/Engine/Source/Runtime/AudioDeviceEnumeration/Windows/WindowsMMDeviceEnumeration/Public/WindowsMMStringUtils.h
@@ -5,32 +5,26 @@
 #include "AudioMixer.h"
 #include "HAL/Platform.h"
 
-#if PLATFORM_WINDOWS
-
-#include "Windows/AllowWindowsPlatformTypes.h"
-
 THIRD_PARTY_INCLUDES_START
+#include <intsafe.h>
+#include "Windows/AllowWindowsPlatformTypes.h"
 
 #include <mmdeviceapi.h>			// IMMNotificationClient
 #include <audiopolicy.h>			// IAudioSessionEvents
 
-THIRD_PARTY_INCLUDES_END
-
 #include "Windows/HideWindowsPlatformTypes.h"
+THIRD_PARTY_INCLUDES_END
 
 namespace Audio
 {
-#if PLATFORM_WINDOWS
 	WINDOWSMMDEVICEENUMERATION_API const TCHAR* ToString(AudioSessionDisconnectReason InDisconnectReason);
 	WINDOWSMMDEVICEENUMERATION_API const TCHAR* ToString(ERole InRole);
 	WINDOWSMMDEVICEENUMERATION_API const TCHAR* ToString(EDataFlow InFlow);
 	WINDOWSMMDEVICEENUMERATION_API FString ToFString(const PROPERTYKEY Key);
-#endif //PLATFORM_WINDOWS
 
 	WINDOWSMMDEVICEENUMERATION_API FString ToFString(const TArray<EAudioMixerChannel::Type>& InChannels);
 
 	WINDOWSMMDEVICEENUMERATION_API const TCHAR* ToString(EAudioDeviceRole InRole);
 	WINDOWSMMDEVICEENUMERATION_API const TCHAR* ToString(EAudioDeviceState InState);
+	WINDOWSMMDEVICEENUMERATION_API FString AudioClientErrorToFString(HRESULT InResult);
 }
- 
-#endif //PLATFORM_WINDOWS
diff --git a/Engine/Source/Runtime/AudioMixerCore/Private/AudioMixer.cpp b/Engine/Source/Runtime/AudioMixerCore/Private/AudioMixer.cpp
--- a/Engine/Source/Runtime/AudioMixerCore/Private/AudioMixer.cpp
+++ b/Engine/Source/Runtime/AudioMixerCore/Private/AudioMixer.cpp
@@ -1,6 +1,7 @@
 // Copyright Epic Games, Inc. All Rights Reserved.
 #include "AudioMixer.h"
 
+#include "Async/Async.h"
 #include "AudioDefines.h"
 #include "AudioMixerTrace.h"
 #include "DSP/BufferVectorOperations.h"
@@ -325,7 +326,6 @@ namespace Audio
 	IAudioMixerPlatformInterface::IAudioMixerPlatformInterface()
 		: bWarnedBufferUnderrun(false)
 		, AudioRenderEvent(nullptr)
-		, bIsInDeviceSwap(false)
 		, AudioFadeEvent(nullptr)
 		, NumOutputBuffers(0)
 		, FadeVolume(0.0f)
@@ -333,7 +333,6 @@ namespace Audio
 		, bPerformingFade(true)
 		, bFadedOut(false)
 		, bIsDeviceInitialized(false)
-		, bMoveAudioStreamToNewAudioDevice(false)
 		, bIsUsingNullDevice(false)
 		, bIsGeneratingAudio(false)
 		, InstanceID(NextInstanceID.Increment())
@@ -468,7 +467,7 @@ namespace Audio
 	void IAudioMixerPlatformInterface::StartRunningNullDevice()
 	{
 		UE_LOG(LogAudioMixer, Display, TEXT("StartRunningNullDevice() called, InstanceID=%d"), InstanceID);
-		SCOPED_NAMED_EVENT(FMixerPlatformXAudio2_StartRunningNullDevice, FColor::Blue);
+		SCOPED_NAMED_EVENT(IAudioMixerPlatformInterface_StartRunningNullDevice, FColor::Blue);
 		
 		auto ThrowAwayBuffer = [this]() { this->ReadNextBuffer(); };
 		float SafeSampleRate = OpenStreamParams.SampleRate > 0.f ? OpenStreamParams.SampleRate : 48000.f;
@@ -497,7 +496,7 @@ namespace Audio
 	void IAudioMixerPlatformInterface::StopRunningNullDevice()
 	{		
 		UE_LOG(LogAudioMixer, Display, TEXT("StopRunningNullDevice() called, InstanceID=%d"), InstanceID);
-		SCOPED_NAMED_EVENT(FMixerPlatformXAudio2_StopRunningNullDevice, FColor::Blue);
+		SCOPED_NAMED_EVENT(IAudioMixerPlatformInterface_StopRunningNullDevice, FColor::Blue);
 
 		if (NullDeviceCallback.IsValid())
 		{
@@ -555,13 +554,6 @@ namespace Audio
 		}
 #endif //!UE_BUILD_SHIPPING
 
-		// If we are flushing buffers for our output voice and this is being called on the audio thread directly,
-		// early exit.
-		if (bIsInDeviceSwap)
-		{
-			return;
-		}
-
 		// If we are currently swapping devices and OnBufferEnd is being triggered in an XAudio2Thread,
 		// early exit.
 		if (!DeviceSwapCriticalSection.TryLock())
@@ -569,7 +561,7 @@ namespace Audio
 			return;
 		}
 
-		// Don't read any more audio if we're not running or changing device
+		// Don't read any more audio if we're not running
 		if (AudioStreamInfo.StreamState != EAudioOutputStreamState::Running)
 		{
 			DeviceSwapCriticalSection.Unlock();
@@ -632,7 +624,7 @@ namespace Audio
 	{
 		SCOPED_NAMED_EVENT(IAudioMixerPlatformInterface_BeginGeneratingAudio, FColor::Blue);
 		
-		checkf(!bIsGeneratingAudio, TEXT("BeginGeneratingAudio() is being run with StreamState = %i and bIsGeneratingAudio = %i"), AudioStreamInfo.StreamState, !!bIsGeneratingAudio);
+		checkf(!bIsGeneratingAudio, TEXT("BeginGeneratingAudio() is being run with StreamState = %i and bIsGeneratingAudio = %i"), AudioStreamInfo.StreamState.load(), !!bIsGeneratingAudio);
 
 		bIsGeneratingAudio = true;
 
@@ -900,7 +892,11 @@ namespace Audio
 
 	bool IAudioMixer::ShouldUseThreadedDeviceSwap()
 	{
+#if PLATFORM_WINDOWS 
 		return bUseThreadedDeviceSwapCVar != 0;
+#else //PLATFORM_WINDOWS
+		return false;
+#endif //PLATFORM_WINDOWS
 	}
 
 	bool IAudioMixer::ShouldUseDeviceInfoCache()
@@ -917,7 +913,214 @@ namespace Audio
 		return bRecycleThreadsCVar != 0;
 	}
 
+	FAudioMixerPlatformSwappable::FAudioMixerPlatformSwappable()
+	{
+	}
+
+	bool FAudioMixerPlatformSwappable::RequestDeviceSwap(const FString& DeviceID, const bool bInForce, const TCHAR* InReason)
+	{
+		if (!AllowDeviceSwap(bInForce))
+		{
+			UE_LOG(LogAudioMixer, Display, TEXT("NOT-ALLOWING attempt to swap audio render device to new device: '%s', because: '%s', force=%d"),
+				!DeviceID.IsEmpty() ? *DeviceID : TEXT("[System Default]"),
+				InReason ? InReason : TEXT("None specified"),
+				(int32)bInForce
+			);
+			return false;
+		}
+
+		// This critical section protects that device swap context below.
+		if (DeviceSwapCriticalSection.TryLock())
+		{
+			UE_LOG(LogAudioMixer, Display, TEXT("Attempt to swap audio render device to new device: '%s', because: '%s', force=%d"),
+				!DeviceID.IsEmpty() ? *DeviceID : TEXT("[System Default]"),
+				InReason ? InReason : TEXT("None specified"),
+				(int32)bInForce
+			);
+
+			// Set up the context for this device swap
+			InitializeDeviceSwapContext(DeviceID, InReason);
+
+			// Set the flag indicating we wish to begin a device swap
+			AudioStreamInfo.StreamState = EAudioOutputStreamState::SwappingDevice;
+		
+			DeviceSwapCriticalSection.Unlock();
+
+			return true;
+		}
+		
+		UE_LOG(LogAudioMixer, Display, TEXT("Dropping device swap request; device: '%s', because: '%s', force=%d"),
+		       !DeviceID.IsEmpty() ? *DeviceID : TEXT("[System Default]"),
+		       InReason ? InReason : TEXT("None specified"),
+		       (int32)bInForce);
+
+		return false;
+	}
 
+	bool FAudioMixerPlatformSwappable::CheckAudioDeviceChange()
+	{
+		SCOPED_NAMED_EVENT(FAudioMixerPlatformSwappable_CheckAudioDeviceChange, FColor::Blue);
+
+		// Use threaded version? (It also requires the info cache).
+		if (IAudioMixer::ShouldUseThreadedDeviceSwap() && ShouldUseDeviceInfoCache())
+		{
+			return CheckThreadedDeviceSwap();
+		}
+		
+		if (AudioStreamInfo.StreamState == EAudioOutputStreamState::SwappingDevice)
+		{
+			return MoveAudioStreamToNewAudioDevice();
+		}
+
+		return false;
+	}
+	
+	bool FAudioMixerPlatformSwappable::AllowDeviceSwap(const bool bInForceSwap)
+	{
+		const double CurrentTime = FPlatformTime::Seconds();
+		const double TimeDelta = CurrentTime - LastDeviceSwapTime;
+
+		LastDeviceSwapTime = CurrentTime;
+
+		if (!IsInitialized())
+		{
+			UE_LOG(LogAudioMixer, Display, TEXT("Unable to device swap until backend is initialized."));
+			return false;
+		}
+		
+		// If the device is not running or if a device swap is already in progress, don't "double-trigger" a swap.
+		if (ActiveDeviceSwap.IsValid())
+		{
+			UE_LOG(LogAudioMixer, Display, TEXT("Ignoring device swap request, swap already requested."));
+			return false;
+		}
+
+		// Only device swap if the audio device is running
+		if (AudioStreamInfo.StreamState != EAudioOutputStreamState::Running)
+		{
+			UE_LOG(LogAudioMixer, Display, TEXT("Ignoring device swap request, AudioStreamInfo.StreamState: %d"), AudioStreamInfo.StreamState.load());
+			return false;
+		}
+
+		// Some devices spam device swap notifications, so we want to rate-limit them to prevent double/triple triggering.
+		static constexpr double MinSwapTimeSeconds = 10.0 / 1000.0;
+		if (TimeDelta <= MinSwapTimeSeconds && !bInForceSwap)
+		{
+			UE_LOG(LogAudioMixer, Display, TEXT("IAudioMixerPlatformInterface::AllowDeviceSwap ignoring device swap due to rate-limit; LastDeviceSwapTime: %f CurrentTime: %f"), LastDeviceSwapTime, CurrentTime);
+			return false;
+		}
+
+		return true;
+	}
+	
+	bool FAudioMixerPlatformSwappable::CheckThreadedDeviceSwap()
+	{
+		bool bDidStopGeneratingAudio = false;
+		SCOPED_NAMED_EVENT(FAudioMixerPlatformSwappable_CheckThreadedDeviceSwap, FColor::Blue);
+
+		// Because this is called every tick, utilize the check-lock-check patter to avoid
+		// unnecessary locking.
+		if (AudioStreamInfo.StreamState == EAudioOutputStreamState::SwappingDevice || ActiveDeviceSwap.IsValid())
+		{
+			// Lock and verify this device is still running. 
+			FScopeLock Lock(&DeviceSwapCriticalSection);
+			
+			// Check the stream state again now that the lock is acquired in case the device is being closed down
+			EAudioOutputStreamState::Type StreamState = AudioStreamInfo.StreamState;
+			
+			// Start a job?
+			if (StreamState == EAudioOutputStreamState::SwappingDevice && !ActiveDeviceSwap.IsValid())
+			{
+				SCOPED_NAMED_EVENT(FAudioMixerPlatformSwappable_CheckThreadedDeviceSwap_StartAsyncSwap, FColor::Blue);
+				
+				if (!PreDeviceSwap())
+				{
+					UE_LOG(LogAudioMixer, Warning, TEXT("FAudioMixerPlatformSwappable::CheckThreadedDeviceSwap PreDeviceSwap() failed"));
+					return bDidStopGeneratingAudio;
+				}
+				
+				EnqueueAsyncDeviceSwap();
+			}
+			else
+			{
+				if (ActiveDeviceSwap.IsReady())  // Finished?
+				{
+					SCOPED_NAMED_EVENT(FAudioMixerPlatformSwappable_CheckThreadedDeviceSwap_FinishedJobSwitchOver, FColor::Blue);
+					const FDeviceSwapResult* DeviceSwapResult = GetDeviceSwapResult();
+					
+					if (bIsUsingNullDevice)
+					{
+						StopRunningNullDevice();
+					}
+
+					StopGeneratingAudio();
+					bDidStopGeneratingAudio = true;
+
+					if (!PostDeviceSwap())
+					{
+						UE_LOG(LogAudioMixer, Warning, TEXT("FAudioMixerPlatformSwappable::CheckThreadedDeviceSwap PostDeviceSwap() failed"));
+						return bDidStopGeneratingAudio;
+					}
+				}
+			}
+		}
+
+		return bDidStopGeneratingAudio;
+	}
+	
+	bool FAudioMixerPlatformSwappable::MoveAudioStreamToNewAudioDevice()
+	{
+		SCOPED_NAMED_EVENT(FAudioMixerPlatformSwappable_MoveAudioStreamToNewAudioDevice, FColor::Blue);
+
+		FScopeLock Lock(&DeviceSwapCriticalSection);
+		bool bDidStopGeneratingAudio = false;
+
+		if (!PreDeviceSwap())
+		{
+			UE_LOG(LogAudioMixer, Warning, TEXT("FAudioMixerPlatformSwappable::MoveAudioStreamToNewAudioDevice PreDeviceSwap() failed"));
+			return bDidStopGeneratingAudio;
+		}
+
+		// Swap devices
+		SynchronousDeviceSwap();
+
+		const FDeviceSwapResult* DeviceSwapResult = GetDeviceSwapResult();
+
+		// Stop null device if currently running
+		if (bIsUsingNullDevice)
+		{
+			StopRunningNullDevice();
+		}
+
+		// Device swaps require reinitialization of output buffers to handle
+		// different channel formats. Stop generating audio to protect against
+		// accessing the OutputBuffer.
+		StopGeneratingAudio();
+		bDidStopGeneratingAudio = true;
+		
+		if (!PostDeviceSwap())
+		{
+			UE_LOG(LogAudioMixer, Warning, TEXT("FAudioMixerPlatformSwappable::MoveAudioStreamToNewAudioDevice CompleteDeviceSwap() failed"));
+			return bDidStopGeneratingAudio;
+		}
+		
+		return bDidStopGeneratingAudio; 
+	}
+
+	void FAudioMixerPlatformSwappable::ResumePlaybackOnNewDevice()
+	{
+		SCOPED_NAMED_EVENT(FAudioMixerPlatformSwappable_ResumePlaybackOnNewDevice, FColor::Blue);
+		
+		int32 NumSamplesPopped = 0;
+		TArrayView<const uint8> PoppedAudio = OutputBuffer.PopBufferData(NumSamplesPopped);
+		SubmitBuffer(PoppedAudio.GetData());
+
+		check(OpenStreamParams.NumFrames * AudioStreamInfo.DeviceInfo.NumChannels == OutputBuffer.GetNumSamples());
+
+		StartAudioStream();
+		AudioRenderEvent->Trigger();
+	}
+	
 }
 
 FAudioPlatformSettings FAudioPlatformSettings::GetPlatformSettings(const TCHAR* PlatformSettingsConfigFile)
diff --git a/Engine/Source/Runtime/AudioMixerCore/Public/AudioMixer.h b/Engine/Source/Runtime/AudioMixerCore/Public/AudioMixer.h
--- a/Engine/Source/Runtime/AudioMixerCore/Public/AudioMixer.h
+++ b/Engine/Source/Runtime/AudioMixerCore/Public/AudioMixer.h
@@ -2,6 +2,7 @@
 
 #pragma once
 
+#include "Async/Future.h"
 #include "AudioMixerLog.h"
 #include "AudioMixerNullDevice.h"
 #include "AudioMixerTypes.h"
@@ -209,6 +210,7 @@ namespace Audio
 		AUDIOMIXERCORE_API static bool ShouldRecycleThreads();
 
 		/** Called by AudioMixer if it should use Cache for DeviceInfo Enumeration */
+		UE_DEPRECATED(5.6, "ShouldUseDeviceInfoCache functionality moved to IAudioMixerPlatformInterface.")
 		AUDIOMIXERCORE_API static bool ShouldUseDeviceInfoCache();
 
 
@@ -277,7 +279,7 @@ namespace Audio
 		FAudioPlatformDeviceInfo DeviceInfo;
 
 		/** The state of the output audio stream. */
-		EAudioOutputStreamState::Type StreamState;
+		std::atomic<EAudioOutputStreamState::Type> StreamState;
 
 		/** The callback to use for platform-independent layer. */
 		IAudioMixer* AudioMixer;
@@ -428,6 +430,31 @@ namespace Audio
 		virtual FString GetDeviceId() const { return FString(); }
 	};
 
+	
+	struct FDeviceSwapContext
+	{
+		FDeviceSwapContext() = delete;
+		FDeviceSwapContext(const FString& InRequestedDeviceID, const FString& InReason) :
+			RequestedDeviceId(InRequestedDeviceID),
+			DeviceSwapReason(InReason)
+		{}
+		virtual ~FDeviceSwapContext() = default;
+		
+		FString RequestedDeviceId;
+		FString DeviceSwapReason;
+		TOptional<FAudioPlatformDeviceInfo> NewDevice;
+	};
+
+	struct FDeviceSwapResult
+	{
+		virtual ~FDeviceSwapResult() = default;
+		
+		virtual bool IsNewDeviceReady() const { return false; }
+
+		FAudioPlatformDeviceInfo DeviceInfo;
+		FString SwapReason;
+		double SuccessfulDurationMs = 0.0;
+	};
 
 	/** Abstract interface for mixer platform. */
 	class IAudioMixerPlatformInterface : public FRunnable,
@@ -495,7 +522,9 @@ namespace Audio
 		virtual bool StopAudioStream() = 0;
 
 		/** Resets the audio stream to use a new audio device with the given device ID (empty string means default). */
-		virtual bool MoveAudioStreamToNewAudioDevice(const FString& InNewDeviceId) { return true;  }
+		UE_DEPRECATED(5.6, "Please use new version MoveAudioStreamToNewAudioDevice which takes no parameters.")
+		virtual bool MoveAudioStreamToNewAudioDevice(const FString& InNewDeviceId) { return MoveAudioStreamToNewAudioDevice();  }
+		virtual bool MoveAudioStreamToNewAudioDevice() { return true;  }
 
 		/** Sends a command to swap which output device is being used */
 		virtual bool RequestDeviceSwap(const FString& DeviceID, bool bInForce, const TCHAR* InReason = nullptr) { return false; }
@@ -539,6 +568,12 @@ namespace Audio
 		// Get the DeviceInfo Cache if one exists.
 		virtual IAudioPlatformDeviceInfoCache* GetDeviceInfoCache() const { return nullptr;  }
 
+		// Determine if the given device info struct is valid for use with this platform
+		virtual bool IsDeviceInfoValid(const FAudioPlatformDeviceInfo& InDeviceInfo) const { return false; }
+
+		// Subclasses can override to determine if it should use Cache for DeviceInfo Enumeration
+		virtual bool ShouldUseDeviceInfoCache() const { return false; }
+	
 	public: // Public Functions
 		//~ Begin FRunnable
 		AUDIOMIXERCORE_API uint32 Run() override;
@@ -693,6 +728,7 @@ namespace Audio
 		FCriticalSection DeviceSwapCriticalSection;
 
 		/** This is used if we are attempting to TryLock on DeviceSwapCriticalSection, but a buffer callback is being called in the current thread. */
+		UE_DEPRECATED(5.6, "bIsInDeviceSwap has been deprecated. AudioStreamInfo.StreamState can provide similar functionality.")
 		FThreadSafeBool bIsInDeviceSwap;
 
 		/** Event allows you to block until fadeout is complete. */
@@ -720,7 +756,9 @@ namespace Audio
 		FThreadSafeBool bFadedOut;
 		FThreadSafeBool bIsDeviceInitialized;
 
+		UE_DEPRECATED(5.6, "bMoveAudioStreamToNewAudioDevice has been deprecated. AudioStreamInfo.StreamState can provide similar functionality.")
 		FThreadSafeBool bMoveAudioStreamToNewAudioDevice;
+
 		FThreadSafeBool bIsUsingNullDevice;
 		FThreadSafeBool bIsGeneratingAudio;
 
@@ -733,6 +771,68 @@ namespace Audio
 	private:
 		TUniquePtr<FMixerNullCallback> NullDeviceCallback;
 	};
+
+	/** Audio mixer platform objects can subclass this in order to add device swap capabilities. */
+	class AUDIOMIXERCORE_API FAudioMixerPlatformSwappable : public IAudioMixerPlatformInterface
+	{
+	public:
+		FAudioMixerPlatformSwappable();
+		virtual ~FAudioMixerPlatformSwappable() override = default;
+
+		//~ Begin IAudioMixerPlatformInterface
+		virtual bool RequestDeviceSwap(const FString& DeviceID, const bool bInForce, const TCHAR* InReason) override;
+		virtual bool CheckAudioDeviceChange() override;
+		virtual bool MoveAudioStreamToNewAudioDevice() override;
+		virtual void ResumePlaybackOnNewDevice() override;
+		//~ End IAudioMixerPlatformInterface
+
+		/** Called to determine if the current device swap request should be allowed to proceed. */
+		virtual bool AllowDeviceSwap(const bool bInForceSwap);
+
+		/** Initializes a new device swap context with the given parameters */
+		virtual bool InitializeDeviceSwapContext(const FString& InRequestedDeviceID, const TCHAR* InReason) = 0;
+
+		/** Called repeatedly to update an active, async device swap */
+		virtual bool CheckThreadedDeviceSwap();
+
+		/** Called at the beginning of a device swap to perform any needed initialization */
+		virtual bool PreDeviceSwap() { return true; }
+
+		/** Kicks of an async device swap task */
+		virtual void EnqueueAsyncDeviceSwap() = 0;
+
+		/** Performs a device swap synchronously in the current thread */
+		virtual void SynchronousDeviceSwap() = 0;
+
+		/** Called after a device swap completes, providing an opportunity for any needed cleanup */
+		virtual bool PostDeviceSwap() { return true; }
+
+	protected:
+		/** Used in OnDeviceAdded for determining if an added device is the same as the original device. */
+		FString GetOriginalAudioDeviceId() const { return OriginalAudioDeviceId; }
+		void SetOriginalAudioDeviceId(const FString& InAudioDeviceId) { OriginalAudioDeviceId = InAudioDeviceId; }
+		
+		/** Future which holds result of device swap upon completion. */
+		void SetActiveDeviceSwapFuture(TFuture<TUniquePtr<FDeviceSwapResult>>&& InFuture) { ActiveDeviceSwap = MoveTemp(InFuture); } 
+		void ResetActiveDeviceSwapFuture() { ActiveDeviceSwap.Reset(); }
+
+		/** The results produced from a device swap operation. The containing future retains memory ownership.
+		 *  Will be valid until PostDeviceSwap returns.
+		 */
+		const FDeviceSwapResult* GetDeviceSwapResult() const { return ActiveDeviceSwap.IsValid() ? ActiveDeviceSwap.Get().Get() : nullptr; }
+		FDeviceSwapResult* GetDeviceSwapResult() { return ActiveDeviceSwap.IsValid() ? ActiveDeviceSwap.Get().Get() : nullptr; }
+
+	private:
+		/** Used in OnDeviceAdded for determining if an added device is the same as the original device. */
+		FString OriginalAudioDeviceId;
+
+		/** Future which holds result of device swap upon completion. */
+		TFuture<TUniquePtr<FDeviceSwapResult>> ActiveDeviceSwap;
+		
+		/** Used to rate limit device swap attempts to ignore multiple requests */
+		double LastDeviceSwapTime = 0.0;
+	};
+
 }
 
 /**
diff --git a/Engine/Source/Runtime/AudioMixerCore/Public/AudioMixerTypes.h b/Engine/Source/Runtime/AudioMixerCore/Public/AudioMixerTypes.h
--- a/Engine/Source/Runtime/AudioMixerCore/Public/AudioMixerTypes.h
+++ b/Engine/Source/Runtime/AudioMixerCore/Public/AudioMixerTypes.h
@@ -38,6 +38,9 @@ namespace Audio {
 
 			/** The audio output stream is open and running. */
 			Running,
+
+			/** The audio device hardware is being swapped */
+			SwappingDevice,
 		};
 	}
 
diff --git a/Engine/Source/Runtime/Windows/AudioMixerWasapi/Private/AudioMixerWasapi.cpp b/Engine/Source/Runtime/Windows/AudioMixerWasapi/Private/AudioMixerWasapi.cpp
--- a/Engine/Source/Runtime/Windows/AudioMixerWasapi/Private/AudioMixerWasapi.cpp
+++ b/Engine/Source/Runtime/Windows/AudioMixerWasapi/Private/AudioMixerWasapi.cpp
@@ -1,27 +1,36 @@
 // Copyright Epic Games, Inc. All Rights Reserved.
 
 #include "AudioMixerWasapi.h"
-#include "WasapiAudioUtils.h"
-#include "WasapiDefaultDeviceMgr.h"
+
+#include "Async/Async.h"
+#include "ScopedCom.h"
 #include "WasapiAggregateDeviceMgr.h"
+#include "WasapiDefaultDeviceMgr.h"
 
 namespace Audio
 {
 	FAudioMixerWasapi::FAudioMixerWasapi(const bool bInUseAggregateDevice) :
 		bUseWasapiAggregateDevice(bInUseAggregateDevice)
 	{
-		if (bUseWasapiAggregateDevice)
+		CreateDeviceManager(bInUseAggregateDevice, DeviceManager);
+	}
+
+	FAudioMixerWasapi::~FAudioMixerWasapi()
+	{
+	}
+
+	void FAudioMixerWasapi::CreateDeviceManager(const bool bInUseAggregateDevice, TUniquePtr<IAudioMixerWasapiDeviceManager>& InDeviceManager)
+	{
+		if (bInUseAggregateDevice)
 		{
-			ensure(DeviceManager = MakeUnique<FWasapiAggregateDeviceMgr>());
+			InDeviceManager = MakeUnique<FWasapiAggregateDeviceMgr>();
 		}
 		else
 		{
-			ensure(DeviceManager = MakeUnique<FWasapiDefaultDeviceMgr>());
+			InDeviceManager = MakeUnique<FWasapiDefaultDeviceMgr>();
 		}
-	}
 
-	FAudioMixerWasapi::~FAudioMixerWasapi()
-	{
+		ensure(InDeviceManager);
 	}
 
 	bool FAudioMixerWasapi::InitializeHardware()
@@ -39,13 +48,22 @@ namespace Audio
 			return false;
 		}
 
-		TFunction<void()> ReadNextBufferCallback = [this]() { ReadNextBuffer(); };
+		// The ReadNextBufferCallback life cycle is tied to this object. It is ultimately bound to a delegate
+		// in the render stream object which will be unbound in TeardownHardware, prior to 'this' being deallocated.
+		TFunction<void()> ReadNextBufferCallback = [this](){ ReadNextBuffer(); };
 		if (!ensure(DeviceManager && DeviceManager->InitializeHardware(StreamParams, ReadNextBufferCallback)))
 		{
 			AUDIO_PLATFORM_LOG_ONCE(TEXT("FAudioMixerWasapi::InitializeHardware InitializeHardware() failed"), Warning);
 			return false;
 		}
 
+		if (IAudioMixer::ShouldRecycleThreads())
+		{
+			// Pre-create the null render device thread so we can simply wake it up when needed.
+			// Give it nothing to do, with a slow tick as the default, but ask it to wait for a signal to wake up.
+			CreateNullDeviceThread([] {}, 1.0f, true);
+		}
+
 		return true;
 	}
 
@@ -57,10 +75,12 @@ namespace Audio
 			return false;
 		}
 
+		// Lock prior to changing state to avoid race condition if there happens to be an in-flight device swap
+		FScopeLock Lock(&DeviceSwapCriticalSection);
+
 		if (DeviceManager.IsValid() && !DeviceManager->TeardownHardware())
 		{
 			AUDIO_PLATFORM_LOG_ONCE(TEXT("FAudioMixerWasapi::TeardownHardware DeviceManager->TeardownHardware() failed."), Warning);
-			return false;
 		}
 
 		bIsInitialized = false;
@@ -70,17 +90,17 @@ namespace Audio
 
 	bool FAudioMixerWasapi::IsInitialized() const
 	{
-		return bIsInitialized && DeviceManager.IsValid() && DeviceManager->IsInitialized();
+		return bIsInitialized;
 	}
 
-	int32 FAudioMixerWasapi::GetNumFrames(const int32 InNumReqestedFrames)
+	int32 FAudioMixerWasapi::GetNumFrames(const int32 InNumRequestedFrames)
 	{
 		if (DeviceManager.IsValid())
 		{
-			return DeviceManager->GetNumFrames(InNumReqestedFrames);
+			return DeviceManager->GetNumFrames(InNumRequestedFrames);
 		}
 
-		return InNumReqestedFrames;
+		return InNumRequestedFrames;
 	}
 
 	bool FAudioMixerWasapi::GetNumOutputDevices(uint32& OutNumOutputDevices)
@@ -89,7 +109,7 @@ namespace Audio
 
 		OutNumOutputDevices = 0;
 
-		if (IAudioPlatformDeviceInfoCache* Cache = GetDeviceInfoCache())
+		if (const IAudioPlatformDeviceInfoCache* Cache = GetDeviceInfoCache())
 		{
 			OutNumOutputDevices = Cache->GetAllActiveOutputDevices().Num();
 			return true;
@@ -105,7 +125,7 @@ namespace Audio
 	{
 		SCOPED_NAMED_EVENT(FAudioMixerWasapi_GetOutputDeviceInfo, FColor::Blue);
 
-		if (IAudioPlatformDeviceInfoCache* Cache = GetDeviceInfoCache())
+		if (const IAudioPlatformDeviceInfoCache* Cache = GetDeviceInfoCache())
 		{
 			if (InDeviceIndex == AUDIO_MIXER_DEFAULT_DEVICE_INDEX)
 			{
@@ -151,11 +171,18 @@ namespace Audio
 			return false;
 		}
 
+		return InitStreamParams(DeviceInfo, InNumBufferFrames, InNumBuffers, InSampleRate, OutParams);
+	}
+	
+	bool FAudioMixerWasapi::InitStreamParams(const FAudioPlatformDeviceInfo& InDeviceInfo, const int32 InNumBufferFrames, const int32 InNumBuffers, const int32 InSampleRate, TArray<FWasapiRenderStreamParams>& OutParams) const
+	{
+		SCOPED_NAMED_EVENT(FAudioMixerWasapi_InitStreamParams, FColor::Blue);
+
 		if (bUseWasapiAggregateDevice)
 		{
-			if (IAudioPlatformDeviceInfoCache* Cache = GetDeviceInfoCache())
+			if (const IAudioPlatformDeviceInfoCache* Cache = GetDeviceInfoCache())
 			{
-				FName DeviceId = *DeviceInfo.DeviceId;
+				const FName DeviceId = *InDeviceInfo.DeviceId;
 				TArray<FAudioPlatformDeviceInfo> AggregateDevice = Cache->GetAggregateDeviceInfo(DeviceId);
 
 				for (const FAudioPlatformDeviceInfo& AggregateDeviceInfo : AggregateDevice)
@@ -173,14 +200,14 @@ namespace Audio
 		}
 		else
 		{
-			TComPtr<IMMDevice> MMDevice = GetMMDevice(DeviceInfo.DeviceId);
+			const TComPtr<IMMDevice> MMDevice = GetMMDevice(InDeviceInfo.DeviceId);
 			if (!MMDevice)
 			{
 				UE_LOG(LogAudioMixer, Error, TEXT("FAudioMixerWasapi::InitStreamParams null MMDevice"));
 				return false;
 			}
 
-			OutParams.Emplace(FWasapiRenderStreamParams(MMDevice, DeviceInfo, InNumBufferFrames, InNumBuffers, InSampleRate));
+			OutParams.Emplace(FWasapiRenderStreamParams(MMDevice, InDeviceInfo, InNumBufferFrames, InNumBuffers, InSampleRate));
 		}
 
 		return true;
@@ -219,6 +246,12 @@ namespace Audio
 			return false;
 		}
 
+		// Store the device ID here in case it is removed. We can switch back if the device comes back.
+		if (Params.bRestoreIfRemoved)
+		{
+			SetOriginalAudioDeviceId(AudioStreamInfo.DeviceInfo.DeviceId);
+		}
+
 		bIsInitialized = true;
 
 		UE_LOG(LogAudioMixer, Display, TEXT("FAudioMixerWasapi initialized: SampleRate=%d NumChannels=%d NumDirectOutChannels=%d AggregateMode=%d"), 
@@ -234,9 +267,16 @@ namespace Audio
 			return false;
 		}
 
-		if (!DeviceManager->CloseAudioStream())
+		// Lock prior to changing state to avoid race condition if there happens to be an in-flight device swap
+		FScopeLock Lock(&DeviceSwapCriticalSection);
+
+		// If we're closing the stream, we're not interested in the results of the device swap. 
+		// Reset the handle to the future.
+		ResetActiveDeviceSwapFuture();
+
+		if (DeviceManager.IsValid() && !DeviceManager->CloseAudioStream())
 		{
-			return false;
+			UE_LOG(LogAudioMixer, Warning, TEXT("FAudioMixerWasapi::CloseAudioStream CloseAudioStream failed"));
 		}
 
 		AudioStreamInfo.StreamState = EAudioOutputStreamState::Closed;
@@ -246,16 +286,33 @@ namespace Audio
 
 	bool FAudioMixerWasapi::StartAudioStream()
 	{
-		if (IsInitialized() && DeviceManager->StartAudioStream())
+		bool bDidStartAudioStream = false;
+		
+		if (bIsInitialized)
 		{
-			BeginGeneratingAudio();
-
-			AudioStreamInfo.StreamState = EAudioOutputStreamState::Running;
-
-			return true;
+			if (DeviceManager.IsValid() && DeviceManager->IsInitialized())
+			{
+				bDidStartAudioStream = DeviceManager->StartAudioStream();
+			}
+			else
+			{
+				check(!bIsUsingNullDevice);
+				StartRunningNullDevice();
+			}
+			
+			// Can be called during device swap when AudioRenderEvent can be null
+			if (nullptr == AudioRenderEvent)
+			{
+				// This will set AudioStreamInfo.StreamState to EAudioOutputStreamState::Running
+				BeginGeneratingAudio();
+			}
+			else
+			{
+				AudioStreamInfo.StreamState = EAudioOutputStreamState::Running;
+			}
 		}
 
-		return false;
+		return bDidStartAudioStream;
 	}
 
 	bool FAudioMixerWasapi::StopAudioStream()
@@ -266,6 +323,9 @@ namespace Audio
 			return false;
 		}
 
+		// Lock prior to changing state to avoid race condition if there happens to be an in-flight device swap
+		FScopeLock Lock(&DeviceSwapCriticalSection);
+
 		if (AudioStreamInfo.StreamState != EAudioOutputStreamState::Stopped && AudioStreamInfo.StreamState != EAudioOutputStreamState::Closed)
 		{
 			if (AudioStreamInfo.StreamState == EAudioOutputStreamState::Running)
@@ -280,6 +340,11 @@ namespace Audio
 
 			check(AudioStreamInfo.StreamState == EAudioOutputStreamState::Stopped);
 		}
+		
+		if (bIsUsingNullDevice)
+		{
+			StopRunningNullDevice();
+		}
 
 		return true;
 	}
@@ -325,6 +390,318 @@ namespace Audio
 
 	IAudioPlatformDeviceInfoCache* FAudioMixerWasapi::GetDeviceInfoCache() const
 	{
-		return DeviceInfoCache.Get();
+		if (ShouldUseDeviceInfoCache())
+		{
+			return DeviceInfoCache.Get();
+		}
+
+		return nullptr;
+	}
+	
+	bool FAudioMixerWasapi::IsDeviceInfoValid(const FAudioPlatformDeviceInfo& InDeviceInfo) const
+	{
+		// Device enumeration will not return invalid devices. This
+		// is more of a sanity check.
+		if (InDeviceInfo.NumChannels > 0 && InDeviceInfo.SampleRate > 0)
+		{
+			return true;
+		}
+
+		return false;
+	}
+
+	void FAudioMixerWasapi::OnSessionDisconnect(IAudioMixerDeviceChangedListener::EDisconnectReason InReason)
+	{
+		// Device has disconnected from current session.
+		if (InReason == IAudioMixerDeviceChangedListener::EDisconnectReason::FormatChanged)
+		{
+			// OnFormatChanged, retry again same device.
+			RequestDeviceSwap(GetDeviceId(), /*force*/ true, TEXT("FAudioMixerWasapi::OnSessionDisconnect() - FormatChanged"));		
+		}
+		else if (InReason == IAudioMixerDeviceChangedListener::EDisconnectReason::DeviceRemoval)
+		{
+			// Ignore Device Removal, as this is handle by the Device Removal logic in the Notification Client.
+		}
+		else
+		{
+			// ServerShutdown, SessionLogoff, SessionDisconnected, ExclusiveModeOverride
+			// Attempt a default swap, will likely fail, but then we'll switch to a null device.
+			RequestDeviceSwap(TEXT(""), /*force*/ true, TEXT("FAudioMixerWasapi::OnSessionDisconnect() - Other"));
+		}
+	}
+	
+	bool FAudioMixerWasapi::CheckThreadedDeviceSwap()
+	{
+		bool bDidStopGeneratingAudio = false;
+#if PLATFORM_WINDOWS
+		bDidStopGeneratingAudio = FAudioMixerPlatformSwappable::CheckThreadedDeviceSwap();
+#endif //PLATFORM_WINDOWS
+		return bDidStopGeneratingAudio;
+	}
+
+	bool FAudioMixerWasapi::InitializeDeviceSwapContext(const FString& InRequestedDeviceID, const TCHAR* InReason)
+	{
+		check(GetDeviceInfoCache());
+
+		// Look up device. Blank name looks up current default.
+		const FName NewDeviceName = *InRequestedDeviceID;
+		TOptional<FAudioPlatformDeviceInfo> DeviceInfo;
+	
+		if (TOptional<FAudioPlatformDeviceInfo> TempDeviceInfo = GetDeviceInfoCache()->FindActiveOutputDevice(NewDeviceName))
+		{
+			if (TempDeviceInfo.IsSet())
+			{
+				if (IsDeviceInfoValid(*TempDeviceInfo))
+				{
+					DeviceInfo = MoveTemp(TempDeviceInfo);
+				}
+				else
+				{
+					UE_LOG(LogAudioMixer, Warning, TEXT("FAudioMixerWasapi::InitializeDeviceSwapContext - Ignoring attempt to switch to device with unsupported params: Channels=%u, SampleRate=%u, Id=%s, Name=%s"),
+						(uint32)TempDeviceInfo->NumChannels, (uint32)TempDeviceInfo->SampleRate, *TempDeviceInfo->DeviceId, *TempDeviceInfo->Name);
+
+					return false;
+				}
+			}
+		}
+		
+		return InitDeviceSwapContextInternal(InRequestedDeviceID, InReason, DeviceInfo);
+	}
+
+	bool FAudioMixerWasapi::InitDeviceSwapContextInternal(const FString& InRequestedDeviceID, const TCHAR* InReason, const TOptional<FAudioPlatformDeviceInfo>& InDeviceInfo)
+	{
+		// Access to device swap context must be protected by DeviceSwapCriticalSection
+		FScopeLock Lock(&DeviceSwapCriticalSection);
+		
+		if (DeviceSwapContext.IsValid())
+		{
+			UE_LOG(LogAudioMixer, Display, TEXT("FAudioMixerWasapi::InitDeviceSwapContextInternal - DeviceSwapContext in-flight, ignoring"));
+			return false;
+		}
+
+		// Create the device swap context which will valid over the course of the swap
+		DeviceSwapContext = MakeUnique<FWasapiDeviceSwapContext>(InRequestedDeviceID, InReason);
+		if (!DeviceSwapContext.IsValid())
+		{
+			UE_LOG(LogAudioMixer, Warning, TEXT("FMixerPlatformWasapi::CreateDeviceSwapContext - failed to create DeviceSwapContext"));
+			return false;
+		}
+
+		DeviceSwapContext->NewDevice = InDeviceInfo;
+		
+		const FAudioPlatformSettings EngineSettings = GetPlatformSettings();
+		TArray<FWasapiRenderStreamParams> StreamParams;
+
+		if (DeviceSwapContext->NewDevice.IsSet())
+		{
+			if (!InitStreamParams(*DeviceSwapContext->NewDevice, EngineSettings.CallbackBufferFrameSize, EngineSettings.NumBuffers, EngineSettings.SampleRate, StreamParams))
+			{
+				UE_LOG(LogAudioMixer, Warning, TEXT("FAudioMixerWasapi::InitializeDeviceSwapContext - InitStreamParams() failed"));
+				DeviceSwapContext->NewDevice.Reset();
+
+				return false;
+			}
+		}
+
+		// Initialize remaining fields except for OldDeviceManager which will 
+		// happen later in CheckThreadedDeviceSwap from the Game thread.
+		DeviceSwapContext->ReadNextBufferCallback = [this](){ ReadNextBuffer(); };
+		DeviceSwapContext->StreamParams = StreamParams;
+		DeviceSwapContext->bUseAggregateDevice = bUseWasapiAggregateDevice;
+		DeviceSwapContext->PlatformSettings = EngineSettings;
+
+		return true;
+	}
+	
+	void FAudioMixerWasapi::EnqueueAsyncDeviceSwap()
+	{
+		FScopeLock Lock(&DeviceSwapCriticalSection);
+		UE_LOG(LogAudioMixer, Verbose, TEXT("FAudioMixerWasapi::EnqueueAsyncDeviceSwap - enqueuing async device swap"));
+		
+		TFunction<TUniquePtr<FDeviceSwapResult>()> AsyncDeviceSwap = [this]() mutable -> TUniquePtr<FDeviceSwapResult>
+		{
+			// Transfer ownership of DeviceSwapContext to the async task.
+			TUniquePtr<FWasapiDeviceSwapContext> TempContext;
+			{
+				FScopeLock Lock(&DeviceSwapCriticalSection);
+
+				if (AudioStreamInfo.StreamState == EAudioOutputStreamState::SwappingDevice)
+				{
+					TempContext = MoveTemp(DeviceSwapContext);
+				}
+			}
+			
+			return PerformDeviceSwap(MoveTemp(TempContext));
+		};
+		SetActiveDeviceSwapFuture(Async(EAsyncExecution::TaskGraph, MoveTemp(AsyncDeviceSwap)));
+	}
+
+	void FAudioMixerWasapi::SynchronousDeviceSwap()
+	{
+		FScopeLock Lock(&DeviceSwapCriticalSection);
+
+		// Transfer ownership of DeviceSwapContext memory to the device swap routine
+		TUniquePtr<FDeviceSwapResult> DeviceSwapResult = PerformDeviceSwap(TUniquePtr<FWasapiDeviceSwapContext>(MoveTemp(DeviceSwapContext)));
+
+		// Set the promise and future result to replicate what the async task does
+		TPromise<TUniquePtr<FDeviceSwapResult>> Promise;
+		
+		// It's ok if DeviceSwapResult is null here. It indicates an invalid device which will be handled.
+		Promise.SetValue(MoveTemp(DeviceSwapResult));
+		SetActiveDeviceSwapFuture(Promise.GetFuture());
+	}
+	
+	TUniquePtr<FDeviceSwapResult> FAudioMixerWasapi::PerformDeviceSwap(TUniquePtr<FWasapiDeviceSwapContext>&& InDeviceContext)
+	{
+		// Static method enforces no other state sharing occurs with the object that called it. InDeviceContext
+		// should have no other references outside of this method so that the device swap operation is isolated.
+		SCOPED_NAMED_EVENT(FAudioMixerWasapi_PerformDeviceSwap, FColor::Blue);
+
+		const uint64 StartTimeCycles = FPlatformTime::Cycles64();
+		// This runs in an async task whose thread may not have initialized com
+		FScopedCoInitialize CoInitialize;
+
+		// No need to lock critical section here since this call has sole ownership of the context
+		if (InDeviceContext.IsValid())
+		{
+			UE_LOG(LogAudioMixer, Verbose, TEXT("FAudioMixerWasapi::PerformDeviceSwap - AsyncTask Start. Because=%s"), *InDeviceContext->DeviceSwapReason);
+
+			if (InDeviceContext->OldDeviceManager.IsValid())
+			{
+				// Shutdown the current device manager
+				InDeviceContext->OldDeviceManager->StopAudioStream();
+				InDeviceContext->OldDeviceManager->CloseAudioStream();
+				InDeviceContext->OldDeviceManager->TeardownHardware();
+				InDeviceContext->OldDeviceManager.Reset();
+				UE_LOG(LogAudioMixer, Verbose, TEXT("FAudioMixerWasapi::PerformDeviceSwap - successfully shut down previous device manager"));
+			}
+			else
+			{
+				UE_LOG(LogAudioMixer, Verbose, TEXT("FAudioMixerWasapi::PerformDeviceSwap - no device manager running, null renderer must be active"));
+			}
+
+			// Don't attempt to create a new setup if there's no devices available.
+			if (!InDeviceContext->NewDevice.IsSet())
+			{
+				UE_LOG(LogAudioMixer, Verbose, TEXT("FAudioMixerWasapi::PerformDeviceSwap - InDeviceContext->bSwitchingToValidDevice: false"));
+				return {};
+			}
+
+			TUniquePtr<FWasapiDeviceSwapResult> DeviceSwapResult = MakeUnique<FWasapiDeviceSwapResult>();
+			CreateDeviceManager(InDeviceContext->bUseAggregateDevice, DeviceSwapResult->NewDeviceManager);
+			
+			if (!DeviceSwapResult->NewDeviceManager.IsValid())
+			{
+				UE_LOG(LogAudioMixer, Warning, TEXT("FAudioMixerWasapi::PerformDeviceSwap - InitializeHardware failed to create new device manager"));
+				return {};
+			}
+			
+			if (!DeviceSwapResult->NewDeviceManager->InitializeHardware(InDeviceContext->StreamParams, InDeviceContext->ReadNextBufferCallback))
+			{
+				UE_LOG(LogAudioMixer, Warning, TEXT("FAudioMixerWasapi::PerformDeviceSwap - InitializeHardware failed while attempting to device swap"));
+				return {};
+			}
+			
+			if (!DeviceSwapResult->NewDeviceManager->OpenAudioStream(InDeviceContext->StreamParams))
+			{
+				UE_LOG(LogAudioMixer, Warning, TEXT("FAudioMixerWasapi::PerformDeviceSwap - OpenAudioStream failed while attempting to device swap"));
+				return {};
+			}
+
+			DeviceSwapResult->SuccessfulDurationMs = FPlatformTime::ToMilliseconds64(FPlatformTime::Cycles64() - StartTimeCycles);
+			DeviceSwapResult->DeviceInfo = InDeviceContext->NewDevice.Get(FAudioPlatformDeviceInfo());
+			DeviceSwapResult->SwapReason = InDeviceContext->DeviceSwapReason;
+			UE_LOG(LogAudioMixer, Verbose, TEXT("FAudioMixerWasapi::PerformDeviceSwap - successfully completed device swap"));
+
+			return DeviceSwapResult;
+		}
+		else
+		{
+			UE_LOG(LogAudioMixer, Error, TEXT("FAudioMixerWasapi::PerformDeviceSwap - failed due to invalid DeviceSwapContext"));
+		}
+
+		return {};
+	}
+	
+	bool FAudioMixerWasapi::PreDeviceSwap()
+	{
+		if (DeviceManager.IsValid())
+		{
+			// Access to device swap context must be protected by DeviceSwapCriticalSection
+			FScopeLock Lock(&DeviceSwapCriticalSection);
+
+			if (DeviceSwapContext.IsValid())
+			{
+				// Finish initializing the device swap context
+				check(!DeviceSwapContext->OldDeviceManager);
+				DeviceSwapContext->OldDeviceManager = MoveTemp(DeviceManager);
+				UE_LOG(LogAudioMixer, Verbose, TEXT("FAudioMixerWasapi::PreDeviceSwap - Starting swap to [%s]"), DeviceSwapContext->RequestedDeviceId.IsEmpty() ? TEXT("[System Default]") : *DeviceSwapContext->RequestedDeviceId);
+			}
+			else
+			{
+				UE_LOG(LogAudioMixer, Warning, TEXT("FAudioMixerWasapi::PreDeviceSwap - null device swap context"));
+				return false;
+			}
+		}
+		else
+		{
+			// This is not an error because the null renderer could be running
+			UE_LOG(LogAudioMixer, Display, TEXT("FAudioMixerWasapi::PreDeviceSwap - no device manager (null renderer must be running)"));
+		}
+
+		return true;
+	}
+	
+	bool FAudioMixerWasapi::PostDeviceSwap()
+	{
+		// Once the context is handed off to the device swap routine (either async or synchronous),
+		// it should no longer be valid.
+		check(!DeviceSwapContext.IsValid());
+		bool bDidSucceed = false;
+
+		FWasapiDeviceSwapResult* DeviceSwapResult = StaticCast<FWasapiDeviceSwapResult*>(GetDeviceSwapResult());
+			
+		if (DeviceSwapResult && DeviceSwapResult->IsNewDeviceReady())
+		{
+			SCOPED_NAMED_EVENT(FAudioMixerWasapiFAudioMixerWasapi_CheckThreadedDeviceSwap_EndSwap, FColor::Blue);
+
+			FScopeLock Lock(&DeviceSwapCriticalSection);
+
+			// Copy our new Device Info into our active one.
+			AudioStreamInfo.DeviceInfo = DeviceSwapResult->DeviceInfo;
+
+			// Display our new XAudio2 Mastering voice details.
+			UE_LOG(LogAudioMixer, Display, TEXT("FAudioMixerWasapi::PostDeviceSwap - successful Swap new Device is (NumChannels=%u, SampleRate=%u, DeviceID=%s, Name=%s), Reason=%s, InstanceID=%d, DurationMS=%.2f"),
+				(uint32)AudioStreamInfo.DeviceInfo.NumChannels, (uint32)AudioStreamInfo.DeviceInfo.SampleRate, *AudioStreamInfo.DeviceInfo.DeviceId, *AudioStreamInfo.DeviceInfo.Name,
+				*DeviceSwapResult->SwapReason, InstanceID, DeviceSwapResult->SuccessfulDurationMs);
+
+			// Reinitialize the output circular buffer to match the buffer math of the new audio device.
+			const int32 NumOutputSamples = AudioStreamInfo.NumOutputFrames * AudioStreamInfo.DeviceInfo.NumChannels;
+			if (ensure(NumOutputSamples > 0))
+			{
+				OutputBuffer.Init(AudioStreamInfo.AudioMixer, NumOutputSamples, NumOutputBuffers, AudioStreamInfo.DeviceInfo.Format);
+			}
+
+			check(!DeviceManager);
+			DeviceManager = MoveTemp(DeviceSwapResult->NewDeviceManager);
+
+			bDidSucceed = true;
+		}
+		else
+		{
+			if (!DeviceSwapResult)
+			{
+				UE_LOG(LogAudioMixer, Error, TEXT("FAudioMixerWasapi::PostDeviceSwap - null device swap result!)"));
+			}
+			else
+			{
+				UE_LOG(LogAudioMixer, Error, TEXT("FAudioMixerWasapi::PostDeviceSwap - DeviceSwapResult->IsNewDeviceReady() = %d)"), DeviceSwapResult->IsNewDeviceReady());
+			}
+		}
+		
+		ResetActiveDeviceSwapFuture();
+
+		return bDidSucceed;
 	}
+	
 }
diff --git a/Engine/Source/Runtime/Windows/AudioMixerWasapi/Private/AudioMixerWasapi.h b/Engine/Source/Runtime/Windows/AudioMixerWasapi/Private/AudioMixerWasapi.h
--- a/Engine/Source/Runtime/Windows/AudioMixerWasapi/Private/AudioMixerWasapi.h
+++ b/Engine/Source/Runtime/Windows/AudioMixerWasapi/Private/AudioMixerWasapi.h
@@ -5,7 +5,8 @@
 #include "AudioMixer.h"	
 #include "AudioMixerWasapiRenderStream.h"
 #include "IAudioMixerWasapiDeviceManager.h"
-#include "Misc/ScopeRWLock.h"
+#include "Templates/SharedPointer.h"
+#include "Templates/UniquePtr.h"
 
 #include "Windows/AllowWindowsPlatformTypes.h"
 #include "Windows/AllowWindowsPlatformAtomics.h"
@@ -22,23 +23,49 @@ THIRD_PARTY_INCLUDES_END
 
 namespace Audio
 {
+	class FWasapiDefaultRenderStream;
+
+	struct FWasapiDeviceSwapContext : public FDeviceSwapContext
+	{
+		FWasapiDeviceSwapContext() = delete;
+		FWasapiDeviceSwapContext(const FString& InRequestedDeviceID, const FString& InReason) :
+			FDeviceSwapContext(InRequestedDeviceID, InReason)
+		{}
+		
+		bool bUseAggregateDevice = false;
+		FAudioPlatformSettings PlatformSettings;
+		TArray<FWasapiRenderStreamParams> StreamParams;
+		TFunction<void()> ReadNextBufferCallback;
+		TUniquePtr<IAudioMixerWasapiDeviceManager> OldDeviceManager;
+	};
+
+	struct FWasapiDeviceSwapResult : public FDeviceSwapResult
+	{
+		virtual bool IsNewDeviceReady() const override
+		{
+			return NewDeviceManager.IsValid();
+		}
+
+		TUniquePtr<IAudioMixerWasapiDeviceManager> NewDeviceManager;
+	};
+
 	/**
 	 * FAudioMixerWasapi - Wasapi audio backend for Windows and Xbox
 	 */
-	class FAudioMixerWasapi : public IAudioMixerPlatformInterface
+	class FAudioMixerWasapi : public FAudioMixerPlatformSwappable
 	{
 	public:
 
 		FAudioMixerWasapi() = default;
 		explicit FAudioMixerWasapi(const bool bInUseAggregateDevice);
-		virtual ~FAudioMixerWasapi();
+		virtual ~FAudioMixerWasapi() override;
 
 		//~ Begin IAudioMixerPlatformInterface
 		virtual FString GetPlatformApi() const override { return TEXT("WASAPIMixer"); }
 		virtual bool InitializeHardware() override;
 		virtual bool TeardownHardware() override;
 		virtual bool IsInitialized() const override;
-		virtual int32 GetNumFrames(const int32 InNumReqestedFrames) override;
+		virtual int32 GetNumFrames(const int32 InNumRequestedFrames) override;
 		virtual bool GetNumOutputDevices(uint32& OutNumOutputDevices) override;
 		virtual bool GetOutputDeviceInfo(const uint32 InDeviceIndex, FAudioPlatformDeviceInfo& OutInfo) override;
 		virtual FString GetCurrentDeviceName() const override;
@@ -54,25 +81,42 @@ namespace Audio
 		virtual FString GetDefaultDeviceName() override;
 		virtual FAudioPlatformSettings GetPlatformSettings() const override;
 		virtual IAudioPlatformDeviceInfoCache* GetDeviceInfoCache() const override;
+		virtual bool IsDeviceInfoValid(const FAudioPlatformDeviceInfo& InDeviceInfo) const override;
+		virtual bool ShouldUseDeviceInfoCache() const override { return true; }
 		//~ End IAudioMixerPlatformInterface
 
 		//~ Begin IAudioMixerDeviceChangedListener
 		virtual void RegisterDeviceChangedListener() override;
 		virtual void UnregisterDeviceChangedListener() override;
+		virtual void OnDefaultCaptureDeviceChanged(const EAudioDeviceRole InAudioDeviceRole, const FString& DeviceId) override;
+		virtual void OnDefaultRenderDeviceChanged(const EAudioDeviceRole InAudioDeviceRole, const FString& DeviceId) override;
+		virtual void OnDeviceAdded(const FString& DeviceId, bool bIsRenderDevice) override;
+		virtual void OnDeviceRemoved(const FString& DeviceId, bool bIsRenderDevice) override;
+		virtual void OnDeviceStateChanged(const FString& DeviceId, const EAudioDeviceState InState, bool bIsRenderDevice) override;
+		virtual void OnSessionDisconnect(Audio::IAudioMixerDeviceChangedListener::EDisconnectReason InReason) override;
 		virtual FString GetDeviceId() const override;
 		//~ End IAudioMixerDeviceChangedListener
 
+		//~ Begin FAudioMixerPlatformSwappable
+		virtual bool InitializeDeviceSwapContext(const FString& InRequestedDeviceID, const TCHAR* InReason) override;
+		virtual bool CheckThreadedDeviceSwap() override;
+		virtual bool PreDeviceSwap() override;
+		virtual void EnqueueAsyncDeviceSwap() override;
+		virtual void SynchronousDeviceSwap() override;
+		virtual bool PostDeviceSwap() override;
+		//~ End FAudioMixerPlatformSwappable
+
+	protected:
+		/** Can be used by subclasses to initialize a device swap context by supplying a specific 
+		 *  FAudioPlatformDeviceInfo rather than looking it up via the requested device Id.
+		 */
+		bool InitDeviceSwapContextInternal(const FString& InRequestedDeviceID, const TCHAR* InReason, const TOptional<FAudioPlatformDeviceInfo>& InDeviceInfo);
+		
 	private:
 
 		/** Cache for holding information about MM audio devices (IMMDevice).  */
 		TUniquePtr<IAudioPlatformDeviceInfoCache> DeviceInfoCache;
-
-		/** Mutex for protecting shared resources during a device swap. */
-		FCriticalSection AudioDeviceSwapCriticalSection;
-
-		/** Used for determining if a device swap has been requested. */
-		FString OriginalAudioDeviceId;
-
+		
 		/** Manages either a single, default device or an aggregate of several devices belonging to the same hardware */
 		TUniquePtr<IAudioMixerWasapiDeviceManager> DeviceManager;
 
@@ -82,10 +126,21 @@ namespace Audio
 		/** Indicates if this object has been successfully initialized. */
 		bool bIsInitialized = false;
 
+		/** Device swap context which holds necessary data required to perform a device wap. */
+		TUniquePtr<FWasapiDeviceSwapContext> DeviceSwapContext;
+		
 		/** Fetches an IMMDevice with the given ID. */
 		TComPtr<IMMDevice> GetMMDevice(const FString& InDeviceID) const;
 
-		/** Initializes a Wasapi stream paramters struct with the give values. */
+		/** Device manager factory */
+		static void CreateDeviceManager(const bool bInUseAggregateDevice, TUniquePtr<IAudioMixerWasapiDeviceManager>& InDeviceManager);
+		
+		/** Initializes a Wasapi stream parameters struct with the give values. */
 		bool InitStreamParams(const uint32 InDeviceIndex, const int32 InNumBufferFrames, const int32 InNumBuffers, const int32 InSampleRate, TArray<FWasapiRenderStreamParams>& OutParams);
+		bool InitStreamParams(const FAudioPlatformDeviceInfo& InDeviceInfo, const int32 InNumBufferFrames, const int32 InNumBuffers, const int32 InSampleRate, TArray<FWasapiRenderStreamParams>& OutParams) const;
+		
+		/** Performs a device swap with the given context. Static method enforces no other state sharing occurs. */
+		static TUniquePtr<FDeviceSwapResult> PerformDeviceSwap(TUniquePtr<FWasapiDeviceSwapContext>&& InDeviceContext);
 	};
+
  }
diff --git a/Engine/Source/Runtime/Windows/AudioMixerWasapi/Private/AudioMixerWasapiDeviceListener.cpp b/Engine/Source/Runtime/Windows/AudioMixerWasapi/Private/AudioMixerWasapiDeviceListener.cpp
--- a/Engine/Source/Runtime/Windows/AudioMixerWasapi/Private/AudioMixerWasapiDeviceListener.cpp
+++ b/Engine/Source/Runtime/Windows/AudioMixerWasapi/Private/AudioMixerWasapiDeviceListener.cpp
@@ -7,8 +7,6 @@
 
 #include <atomic>
 
-#if PLATFORM_WINDOWS
-
 #include "Microsoft/COMPointer.h"
 #include "ScopedCom.h"					// FScopedComString
 
@@ -31,7 +29,7 @@ namespace Audio
 	{
 		if (!WasapiWinNotificationClient.IsValid())
 		{
-			// Shared (This is a COM object, so we don't delete it, just derecement the ref counter).
+			// Shared (This is a COM object, so we don't delete it, just decrement the ref counter).
 			WasapiWinNotificationClient = TSharedPtr<FWindowsMMNotificationClient>(
 				new FWindowsMMNotificationClient, 
 				[](FWindowsMMNotificationClient* InPtr) { InPtr->ReleaseClient(); }
@@ -63,27 +61,97 @@ namespace Audio
 		}
 	}
 
-	FString FAudioMixerWasapi::GetDeviceId() const
+	void FAudioMixerWasapi::OnDefaultCaptureDeviceChanged(const EAudioDeviceRole InAudioDeviceRole, const FString& DeviceId)
 	{
-		return AudioStreamInfo.DeviceInfo.DeviceId;
+		if (UAudioDeviceNotificationSubsystem* AudioDeviceNotifySubsystem = UAudioDeviceNotificationSubsystem::Get())
+		{
+			AudioDeviceNotifySubsystem->OnDefaultCaptureDeviceChanged(InAudioDeviceRole, DeviceId);
+		}
 	}
 
-	TComPtr<IMMDevice> FAudioMixerWasapi::GetMMDevice(const FString& InDeviceID) const
+	void FAudioMixerWasapi::OnDefaultRenderDeviceChanged(const EAudioDeviceRole InAudioDeviceRole, const FString& DeviceId)
+	{		
+		// There's 3 defaults in windows (communications, console, multimedia). These technically can all be different devices.		
+		// However, the Windows UX only allows console+multimedia to be toggle as a pair. This means you get two notifications
+		// for default device changing typically. To prevent a trouble trigger we only listen to "Console" here. For more information on 
+		// device roles: https://docs.microsoft.com/en-us/windows/win32/coreaudio/device-roles
+		
+		if (InAudioDeviceRole == EAudioDeviceRole::Console)
+		{
+			UE_LOG(LogAudioMixer, Warning, TEXT("FAudioMixerWasapi: Changing default audio render device to new device: Role=%s, DeviceName=%s, InstanceID=%d"), 
+				Audio::ToString(InAudioDeviceRole), *WasapiWinNotificationClient->GetFriendlyName(DeviceId), InstanceID);
+
+			RequestDeviceSwap(DeviceId, /* force */true, TEXT("FAudioMixerWasapi::OnDefaultRenderDeviceChanged"));
+		}
+
+		if (UAudioDeviceNotificationSubsystem* AudioDeviceNotifSubsystem = UAudioDeviceNotificationSubsystem::Get())
+		{
+			AudioDeviceNotifSubsystem->OnDefaultRenderDeviceChanged(InAudioDeviceRole, DeviceId);
+		}
+	}
+
+	void FAudioMixerWasapi::OnDeviceAdded(const FString& DeviceId, bool bIsRenderDevice)
 	{
-		if (WasapiWinNotificationClient)
+		// Ignore changes in capture device.
+		if (!bIsRenderDevice)
 		{
-			return WasapiWinNotificationClient->GetDevice(InDeviceID);
+			return;
 		}
+		
+		// If the device that was added is our original device and our current device is NOT our original device, 
+		// move our audio stream to this newly added device.
+		const FString AudioDeviceId = GetOriginalAudioDeviceId();
+		if (AudioStreamInfo.DeviceInfo.DeviceId != AudioDeviceId && DeviceId == AudioDeviceId)
+		{
+			UE_LOG(LogAudioMixer, Warning, TEXT("FAudioMixerWasapi: Original audio device re-added. Moving audio back to original audio device: DeviceName=%s, bRenderDevice=%d, InstanceID=%d"), 
+				*WasapiWinNotificationClient->GetFriendlyName(*AudioDeviceId), (int32)bIsRenderDevice, InstanceID);
 
-		return TComPtr<IMMDevice>();
+			RequestDeviceSwap(AudioDeviceId, /*force */ true, TEXT("FAudioMixerWasapi::OnDeviceAdded"));
+		}
+
+		if (UAudioDeviceNotificationSubsystem* AudioDeviceNotifSubsystem = UAudioDeviceNotificationSubsystem::Get())
+		{
+			AudioDeviceNotifSubsystem->OnDeviceAdded(DeviceId, bIsRenderDevice);
+		}
 	}
-}
 
-#else 
-namespace Audio
-{
-	void FAudioMixerWasapi::RegisterDeviceChangedListener() {}
-	void FAudioMixerWasapi::UnregisterDeviceChangedListener() {}
+	void FAudioMixerWasapi::OnDeviceRemoved(const FString& DeviceId, bool bIsRenderDevice)
+	{
+		// Ignore changes in capture device.
+		if (!bIsRenderDevice)
+		{
+			return;
+		}
+		
+		// If the device we're currently using was removed... then switch to the new default audio device.
+		if (AudioStreamInfo.DeviceInfo.DeviceId == DeviceId)
+		{
+			UE_LOG(LogAudioMixer, Warning, TEXT("FAudioMixerWasapi: Audio device removed [%s], falling back to other windows default device. bIsRenderDevice=%d, InstanceID=%d"), 
+				*WasapiWinNotificationClient->GetFriendlyName(DeviceId), (int32)bIsRenderDevice, InstanceID);
+
+			RequestDeviceSwap(TEXT(""), /* force */ true, TEXT("FAudioMixerWasapi::OnDeviceRemoved"));
+		}
+
+		if (UAudioDeviceNotificationSubsystem* AudioDeviceNotifySubsystem = UAudioDeviceNotificationSubsystem::Get())
+		{
+			AudioDeviceNotifySubsystem->OnDeviceRemoved(DeviceId, bIsRenderDevice);
+		}
+	}
+
+	void FAudioMixerWasapi::OnDeviceStateChanged(const FString& DeviceId, const EAudioDeviceState InState, bool bIsRenderDevice)
+	{
+		// Ignore changes in capture device.
+		if (!bIsRenderDevice)
+		{
+			return;
+		}
+		
+		if (UAudioDeviceNotificationSubsystem* AudioDeviceNotifySubsystem = UAudioDeviceNotificationSubsystem::Get())
+		{
+			AudioDeviceNotifySubsystem->OnDeviceStateChanged(DeviceId, InState, bIsRenderDevice);
+		}
+	}
+	
 	FString FAudioMixerWasapi::GetDeviceId() const
 	{
 		return AudioStreamInfo.DeviceInfo.DeviceId;
@@ -91,8 +159,12 @@ namespace Audio
 
 	TComPtr<IMMDevice> FAudioMixerWasapi::GetMMDevice(const FString& InDeviceID) const
 	{
+		if (WasapiWinNotificationClient)
+		{
+			return WasapiWinNotificationClient->GetDevice(InDeviceID);
+		}
+
 		return TComPtr<IMMDevice>();
 	}
 }
-#endif
 
diff --git a/Engine/Source/Runtime/Windows/AudioMixerWasapi/Private/AudioMixerWasapiDeviceThread.cpp b/Engine/Source/Runtime/Windows/AudioMixerWasapi/Private/AudioMixerWasapiDeviceThread.cpp
--- a/Engine/Source/Runtime/Windows/AudioMixerWasapi/Private/AudioMixerWasapiDeviceThread.cpp
+++ b/Engine/Source/Runtime/Windows/AudioMixerWasapi/Private/AudioMixerWasapiDeviceThread.cpp
@@ -2,8 +2,8 @@
 
 #include "AudioMixerWasapiDeviceThread.h"
 #include "AudioMixer.h"	
-#include "AudioMixerWasapiLog.h"	
-
+#include "AudioMixerWasapiLog.h"
+#include "ScopedCom.h"
 
 namespace Audio
 {
@@ -59,7 +59,7 @@ namespace Audio
 	{
 		bIsRunning = true;
 
-		bool bCoInitialized = FWindowsPlatformMisc::CoInitialize(ECOMModel::Multithreaded);
+		FScopedCoInitialize ScopedCoInitialize;
 
 		const uint32 NumHandles = EventHandles.Num();
 		const HANDLE* Handles = EventHandles.GetData();
@@ -69,7 +69,7 @@ namespace Audio
 			static constexpr uint32 TimeoutInMs = 1000;
 
 			// WASAPI events for multiple audio devices will all be signaled at the
-			// same time when beloning to the same physical device.
+			// same time when belonging to the same physical device.
 			uint32 Result = ::WaitForMultipleObjects(NumHandles, Handles, true, TimeoutInMs);
 			if (Result == WAIT_TIMEOUT)
 			{
@@ -80,12 +80,7 @@ namespace Audio
 				DeviceRenderCallback();
 			}
 		}
-
-		if (bCoInitialized)
-		{
-			FWindowsPlatformMisc::CoUninitialize();
-		}
-
+		
 		return 0;
 	}
 
@@ -125,6 +120,7 @@ namespace Audio
 
 			bool bShouldWait = true;
 			DeviceRenderThread->Kill(bShouldWait);
+			DeviceRenderThread.Reset();
 		}
 	}
 
@@ -138,6 +134,7 @@ namespace Audio
 			// the stream is disposed of mid-callback.
 			bool bShouldWait = true;
 			DeviceRenderThread->Kill(bShouldWait);
+			DeviceRenderThread.Reset();
 		}
 	}
 }
diff --git a/Engine/Source/Runtime/Windows/AudioMixerWasapi/Private/AudioMixerWasapiDeviceThread.h b/Engine/Source/Runtime/Windows/AudioMixerWasapi/Private/AudioMixerWasapiDeviceThread.h
--- a/Engine/Source/Runtime/Windows/AudioMixerWasapi/Private/AudioMixerWasapiDeviceThread.h
+++ b/Engine/Source/Runtime/Windows/AudioMixerWasapi/Private/AudioMixerWasapiDeviceThread.h
@@ -6,6 +6,7 @@
 
 #include "HAL/Runnable.h"
 #include "HAL/RunnableThread.h"
+#include "IAudioMixerWasapiDeviceManager.h"
 #include "Templates/Function.h"
 #include "Templates/UniquePtr.h"
 
@@ -92,7 +93,7 @@ namespace Audio
 		/**
 		 * Start - Creates the FRunnableThread object which immediately begins running the FAudioMixerWasapiRunnable member.
 		 * 
-		 * @return - Boolean indicating of the thread was succesfully created.
+		 * @return - Boolean indicating of the thread was successfully created.
 		 */
 		bool Start();
 
@@ -100,9 +101,9 @@ namespace Audio
 		 * Stop - Gracefully shuts down the thread.
 		 */
 		void Stop();
-
+		
 		/**
-		 * Abort - Performs non-graceful shutdown of thread which will close the underyling thread handle 
+		 * Abort - Performs non-graceful shutdown of thread which will close the underlying thread handle 
 		 * without waiting for it to finish.
 		 */
 		void Abort();
diff --git a/Engine/Source/Runtime/Windows/AudioMixerWasapi/Private/AudioMixerWasapiModule.cpp b/Engine/Source/Runtime/Windows/AudioMixerWasapi/Private/AudioMixerWasapiModule.cpp
--- a/Engine/Source/Runtime/Windows/AudioMixerWasapi/Private/AudioMixerWasapiModule.cpp
+++ b/Engine/Source/Runtime/Windows/AudioMixerWasapi/Private/AudioMixerWasapiModule.cpp
@@ -1,5 +1,6 @@
 // Copyright Epic Games, Inc. All Rights Reserved.
 
+#include "AudioMixerWasapiModule.h"
 #include "AudioMixerWasapi.h"
 #include "AudioMixerWasapiLog.h"
 #include "Features/IModularFeatures.h"
@@ -8,30 +9,28 @@
 
 DEFINE_LOG_CATEGORY(LogAudioMixerWasapi);
 
-class FAudioMixerModuleWasapi : public IAudioDeviceModule
+void FAudioMixerModuleWasapi::StartupModule()
 {
-public:
-	virtual void StartupModule() override
-	{
-		IAudioDeviceModule::StartupModule();
+	IAudioDeviceModule::StartupModule();
 
-		FModuleManager::Get().LoadModuleChecked(TEXT("AudioMixer"));
-		FModuleManager::Get().LoadModuleChecked(TEXT("AudioMixerCore"));
-	}
+	FModuleManager::Get().LoadModuleChecked(TEXT("AudioMixer"));
+	FModuleManager::Get().LoadModuleChecked(TEXT("AudioMixerCore"));
+}
 
-	virtual bool IsAudioMixerModule() const override
-	{
-		return true;
-	}
+bool FAudioMixerModuleWasapi::IsAudioMixerModule() const
+{
+	return true;
+}
 
-	virtual Audio::IAudioMixerPlatformInterface* CreateAudioMixerPlatformInterface() override
-	{
-		bool bUseWasapiAggregateDevice = false;
+Audio::IAudioMixerPlatformInterface* FAudioMixerModuleWasapi::CreateAudioMixerPlatformInterface()
+{
+	bool bUseWasapiAggregateDevice = false;
 
-		GConfig->GetBool(TEXT("Audio"), TEXT("UseWasapiAggregateDevice"), bUseWasapiAggregateDevice, GEngineIni);
+	GConfig->GetBool(TEXT("Audio"), TEXT("UseWasapiAggregateDevice"), bUseWasapiAggregateDevice, GEngineIni);
 
-		return new Audio::FAudioMixerWasapi(bUseWasapiAggregateDevice);
-	}
-};
+	return new Audio::FAudioMixerWasapi(bUseWasapiAggregateDevice);
+}
 
+#if PLATFORM_WINDOWS
 IMPLEMENT_MODULE(FAudioMixerModuleWasapi, AudioMixerWasapi);
+#endif //PLATFORM_WINDOWS
diff --git a/Engine/Source/Runtime/Windows/AudioMixerWasapi/Private/AudioMixerWasapiModule.h b/Engine/Source/Runtime/Windows/AudioMixerWasapi/Private/AudioMixerWasapiModule.h
new file mode 100644
--- /dev/null
+++ b/Engine/Source/Runtime/Windows/AudioMixerWasapi/Private/AudioMixerWasapiModule.h
@@ -0,0 +1,14 @@
+// Copyright Epic Games, Inc. All Rights Reserved.
+
+#pragma once
+
+#include "AudioMixer.h"
+#include "Modules/ModuleManager.h"
+
+class FAudioMixerModuleWasapi : public IAudioDeviceModule
+{
+public:
+	virtual void StartupModule() override;
+	virtual bool IsAudioMixerModule() const override;
+	virtual Audio::IAudioMixerPlatformInterface* CreateAudioMixerPlatformInterface() override;
+};
diff --git a/Engine/Source/Runtime/Windows/AudioMixerWasapi/Private/AudioMixerWasapiRenderStream.cpp b/Engine/Source/Runtime/Windows/AudioMixerWasapi/Private/AudioMixerWasapiRenderStream.cpp
--- a/Engine/Source/Runtime/Windows/AudioMixerWasapi/Private/AudioMixerWasapiRenderStream.cpp
+++ b/Engine/Source/Runtime/Windows/AudioMixerWasapi/Private/AudioMixerWasapiRenderStream.cpp
@@ -7,6 +7,13 @@
 
 namespace Audio
 {
+	uint32 FAudioMixerWasapiRenderStream::GetMinimumBufferSize(const uint32 InSampleRate)
+	{
+		// Can be called prior to InitializeHardware
+		// Makes assumption about minimum buffer size which we verify in InitializeHardware
+		return InSampleRate / 100;
+	}
+
 	FAudioMixerWasapiRenderStream::FAudioMixerWasapiRenderStream()
 	{
 	}
@@ -64,6 +71,8 @@ namespace Audio
 		REFERENCE_TIME DevicePeriodRefTime = 0;
 		// The second param to GetDevicePeriod is only valid for exclusive mode
 		// Note that GetDevicePeriod returns ref time which is sample rate agnostic
+		// In testing, IAudioClient3::GetSharedModeEnginePeriod() appears to return the same value as
+		// IAudioClient::GetDevicePeriod() so we use the older API.
 		Result = TempAudioClient->GetDevicePeriod(&DevicePeriodRefTime, nullptr);
 		if (FAILED(Result))
 		{
@@ -72,14 +81,30 @@ namespace Audio
 		}
 
 		DefaultDevicePeriod = FWasapiAudioUtils::RefTimeToFrames(DevicePeriodRefTime, InParams.SampleRate);
+		if (DefaultDevicePeriod == 0)
+		{
+			UE_LOG(LogAudioMixerWasapi, Error, TEXT("OpenAudioStream failed DefaultDevicePeriod = %d"), DefaultDevicePeriod);
+			return false;
+		}
 
-		REFERENCE_TIME DesiredBufferDuration = FWasapiAudioUtils::FramesToRefTime(FMath::Max<uint32>(InParams.NumFrames, DefaultDevicePeriod), InParams.SampleRate);
-		// Add two times the device period to the buffer size to ensure buffer phasing isn't an issue with larger buffer sizes.
-		// For example, with an engine period of 1024 and a device period of 480, we'll have two device periods to fetch the next engine buffer before underruning. 
-		DesiredBufferDuration += (DevicePeriodRefTime * 2);
-
+		// Verify that our assumption about the minimum/default buffer size is correct
+		check(DefaultDevicePeriod == GetMinimumBufferSize(InParams.SampleRate));
+		
+		// Determine buffer size to use. 
+		uint32 BufferFramesToRequest = FMath::Max<uint32>(InParams.NumFrames, DefaultDevicePeriod);
+		
+		// If the engine buffer size is not an integral multiple of the device period then we must
+		// account for buffer phasing by padding the requested buffer size.
+		if (BufferFramesToRequest % DefaultDevicePeriod != 0)
+		{
+			// Round up to nearest multiple of the device period
+			uint32 Multiple = FMath::CeilToInt32(static_cast<float>(BufferFramesToRequest + DefaultDevicePeriod - 1) / DefaultDevicePeriod);
+			BufferFramesToRequest = DefaultDevicePeriod * Multiple;
+		}
+		REFERENCE_TIME DesiredBufferDuration = FWasapiAudioUtils::FramesToRefTime(BufferFramesToRequest, InParams.SampleRate);
+		
 		// For shared mode, this is required to be zero
-		const REFERENCE_TIME Periodicity = 0;
+		constexpr REFERENCE_TIME Periodicity = 0;
 
 		// Audio events will be delivered to us rather than needing to poll
 		uint32 Flags = AUDCLNT_STREAMFLAGS_EVENTCALLBACK;
@@ -139,9 +164,9 @@ namespace Audio
 		return bIsInitialized;
 	}
 
-	int32 FAudioMixerWasapiRenderStream::GetNumFrames(const int32 InNumReqestedFrames) const
+	int32 FAudioMixerWasapiRenderStream::GetNumFrames(const int32 InNumRequestedFrames) const
 	{
-		return FMath::Max<uint32>(InNumReqestedFrames, DefaultDevicePeriod);
+		return InNumRequestedFrames;
 	}
 
 	bool FAudioMixerWasapiRenderStream::OpenAudioStream(const FWasapiRenderStreamParams& InParams, HANDLE InEventHandle)
diff --git a/Engine/Source/Runtime/Windows/AudioMixerWasapi/Private/AudioMixerWasapiRenderStream.h b/Engine/Source/Runtime/Windows/AudioMixerWasapi/Private/AudioMixerWasapiRenderStream.h
--- a/Engine/Source/Runtime/Windows/AudioMixerWasapi/Private/AudioMixerWasapiRenderStream.h
+++ b/Engine/Source/Runtime/Windows/AudioMixerWasapi/Private/AudioMixerWasapiRenderStream.h
@@ -4,7 +4,6 @@
 
 #include "AudioMixer.h"	
 #include "IAudioMixerWasapiDeviceManager.h"
-#include "Misc/ScopeRWLock.h"
 #include "WasapiAudioFormat.h"
 
 #include "Windows/AllowWindowsPlatformTypes.h"
@@ -27,7 +26,7 @@ namespace Audio
 	/**
 	 * FAudioMixerWasapiRenderStream
 	 */
-	class FAudioMixerWasapiRenderStream
+	class FAudioMixerWasapiRenderStream : public IDeviceRenderCallback
 	{
 	public:
 
@@ -37,17 +36,18 @@ namespace Audio
 		virtual bool InitializeHardware(const FWasapiRenderStreamParams& InParams);
 		virtual bool TeardownHardware();
 		virtual bool IsInitialized() const;
-		virtual int32 GetNumFrames(const int32 InNumReqestedFrames) const;
+		virtual int32 GetNumFrames(const int32 InNumRequestedFrames) const;
 		virtual bool OpenAudioStream(const FWasapiRenderStreamParams& InParams, HANDLE InEventHandle);
 		virtual bool CloseAudioStream();
 		virtual bool StartAudioStream();
 		virtual bool StopAudioStream();
 		virtual void SubmitBuffer(const uint8* Buffer, const SIZE_T InNumFrames) { }
 		virtual void SubmitDirectOutBuffer(const int32 InChannelIndex, const FAlignedFloatBuffer& InBuffer) { }
-		virtual void DeviceRenderCallback() = 0;
 
-	protected:
+		static uint32 GetMinimumBufferSize(const uint32 InSampleRate);
 
+	protected:
+		
 		/** COM pointer to the WASAPI audio client object. */
 		TComPtr<IAudioClient3> AudioClient;
 
diff --git a/Engine/Source/Runtime/Windows/AudioMixerWasapi/Private/IAudioMixerWasapiDeviceManager.h b/Engine/Source/Runtime/Windows/AudioMixerWasapi/Private/IAudioMixerWasapiDeviceManager.h
--- a/Engine/Source/Runtime/Windows/AudioMixerWasapi/Private/IAudioMixerWasapiDeviceManager.h
+++ b/Engine/Source/Runtime/Windows/AudioMixerWasapi/Private/IAudioMixerWasapiDeviceManager.h
@@ -71,7 +71,7 @@ namespace Audio
 		virtual bool TeardownHardware() = 0;
 		virtual bool IsInitialized() const = 0;
 		virtual int32 GetNumDirectOutChannels() const = 0;
-		virtual int32 GetNumFrames(const int32 InNumReqestedFrames) const = 0;
+		virtual int32 GetNumFrames(const int32 InNumRequestedFrames) const = 0;
 		virtual bool OpenAudioStream(const  TArray<FWasapiRenderStreamParams>& InParams) = 0;
 		virtual bool CloseAudioStream() = 0;
 		virtual bool StartAudioStream() = 0;
@@ -79,4 +79,17 @@ namespace Audio
 		virtual void SubmitBuffer(const uint8* InBuffer, const SIZE_T InNumFrames) = 0;
 		virtual void SubmitDirectOutBuffer(const int32 InDirectOutIndex, const FAlignedFloatBuffer& InBuffer) = 0;
 	};
+
+	/**
+	 * IDeviceRenderCallback - Interface for providing a callback from the device render thread.
+	 */
+	
+	class IDeviceRenderCallback
+	{
+	public:
+		virtual ~IDeviceRenderCallback() = default;
+		
+		virtual void DeviceRenderCallback() = 0;
+	};
+
 }
diff --git a/Engine/Source/Runtime/Windows/AudioMixerWasapi/Private/WasapiAggregateDeviceMgr.cpp b/Engine/Source/Runtime/Windows/AudioMixerWasapi/Private/WasapiAggregateDeviceMgr.cpp
--- a/Engine/Source/Runtime/Windows/AudioMixerWasapi/Private/WasapiAggregateDeviceMgr.cpp
+++ b/Engine/Source/Runtime/Windows/AudioMixerWasapi/Private/WasapiAggregateDeviceMgr.cpp
@@ -91,14 +91,14 @@ namespace Audio
 		return bIsInitialized;
 	}
 
-	int32 FWasapiAggregateDeviceMgr::GetNumFrames(const int32 InNumReqestedFrames) const
+	int32 FWasapiAggregateDeviceMgr::GetNumFrames(const int32 InNumRequestedFrames) const
 	{
 		if (!RenderStreamDevices.IsEmpty())
 		{
-			return RenderStreamDevices[0]->GetNumFrames(InNumReqestedFrames);
+			return RenderStreamDevices[0]->GetNumFrames(InNumRequestedFrames);
 		}
 
-		return InNumReqestedFrames;
+		return InNumRequestedFrames;
 	}
 
 	bool FWasapiAggregateDeviceMgr::OpenAudioStream(const TArray<FWasapiRenderStreamParams>& InParams)
@@ -126,7 +126,7 @@ namespace Audio
 			UE_LOG(LogAudioMixer, Error, TEXT("OpenAudioStream error creating EventHandles"));
 			return false;
 		}
-
+		
 		for (uint32 Index = 0; Index < NumRenderDevices; ++Index)
 		{
 			if (!RenderStreamDevices[Index]->OpenAudioStream(InParams[Index], EventHandles[Index]))
@@ -141,6 +141,8 @@ namespace Audio
 
 	bool FWasapiAggregateDeviceMgr::CloseAudioStream()
 	{
+		RenderDeviceThread.Reset();
+		
 		if (!RenderStreamDevices.IsEmpty())
 		{
 			bool bDidAllClose = true;
diff --git a/Engine/Source/Runtime/Windows/AudioMixerWasapi/Private/WasapiAggregateDeviceMgr.h b/Engine/Source/Runtime/Windows/AudioMixerWasapi/Private/WasapiAggregateDeviceMgr.h
--- a/Engine/Source/Runtime/Windows/AudioMixerWasapi/Private/WasapiAggregateDeviceMgr.h
+++ b/Engine/Source/Runtime/Windows/AudioMixerWasapi/Private/WasapiAggregateDeviceMgr.h
@@ -15,7 +15,7 @@ namespace Audio
 	 * This provides the ability to address more than 8 channels of a physical device using the
 	 * Windows Audio Streaming API (WASAPI).
 	 */
-	class FWasapiAggregateDeviceMgr : public IAudioMixerWasapiDeviceManager
+	class FWasapiAggregateDeviceMgr : public IAudioMixerWasapiDeviceManager, public IDeviceRenderCallback
 	{
 	public:
 
@@ -27,7 +27,7 @@ namespace Audio
 		virtual bool TeardownHardware() override;
 		virtual bool IsInitialized() const override;
 		virtual int32 GetNumDirectOutChannels() const override { return NumDirectOutChannels; }
-		virtual int32 GetNumFrames(const int32 InNumReqestedFrames) const override;
+		virtual int32 GetNumFrames(const int32 InNumRequestedFrames) const override;
 		virtual bool OpenAudioStream(const TArray<FWasapiRenderStreamParams>& InParams) override;
 		virtual bool CloseAudioStream() override;
 		virtual bool StartAudioStream() override;
@@ -36,6 +36,10 @@ namespace Audio
 		virtual void SubmitDirectOutBuffer(const int32 InDirectOutIndex, const FAlignedFloatBuffer& InBuffer) override;
 		//~ End IAudioMixerWasapiDeviceManager
 
+		//~ Begin IDeviceRenderCallback
+		virtual void DeviceRenderCallback() override;
+		//~ End IDeviceRenderCallback
+
 	private:
 
 		bool bIsInitialized = false;
@@ -54,8 +58,5 @@ namespace Audio
 
 		/** The thread which provides an execution context during audio playback. */
 		TUniquePtr<FAudioMixerWasapiDeviceThread> RenderDeviceThread;
-
-		void DeviceRenderCallback();
-
 	};
  }
diff --git a/Engine/Source/Runtime/Windows/AudioMixerWasapi/Private/WasapiAggregateRenderStream.cpp b/Engine/Source/Runtime/Windows/AudioMixerWasapi/Private/WasapiAggregateRenderStream.cpp
--- a/Engine/Source/Runtime/Windows/AudioMixerWasapi/Private/WasapiAggregateRenderStream.cpp
+++ b/Engine/Source/Runtime/Windows/AudioMixerWasapi/Private/WasapiAggregateRenderStream.cpp
@@ -10,26 +10,34 @@ namespace Audio
 {
 	bool FWasapiAggregateRenderStream::InitializeHardware(const FWasapiRenderStreamParams& InParams)
 	{
-		DirectOutBuffers.Reset();
-		DirectOutBuffers.SetNum(InParams.HardwareDeviceInfo.NumChannels);
-
-		for (TCircularAudioBuffer<float>& Buffer : DirectOutBuffers)
+		if (FAudioMixerWasapiRenderStream::InitializeHardware(InParams))
 		{
-			int32 NumOutputBuffers = FMath::Max(InParams.NumBuffers, 2);
-			int32 BufferCapacity = InParams.NumFrames * NumOutputBuffers;
+			DirectOutBuffers.Reset();
+			DirectOutBuffers.SetNum(InParams.HardwareDeviceInfo.NumChannels);
+		
+			for (TCircularAudioBuffer<float>& Buffer : DirectOutBuffers)
+			{
+				int32 NumOutputBuffers = FMath::Max(InParams.NumBuffers, 2);
+				int32 BufferCapacity = InParams.NumFrames * NumOutputBuffers;
 
-			Buffer.SetCapacity(BufferCapacity);
-		}
+				Buffer.SetCapacity(BufferCapacity);
+			}
 
-		InterleaveBuffers.Reset();
-		InterleaveBuffers.SetNum(InParams.HardwareDeviceInfo.NumChannels);
+			InterleaveBuffers.Reset();
+			InterleaveBuffers.SetNum(InParams.HardwareDeviceInfo.NumChannels);
 
-		for (FAlignedFloatBuffer& Buffer : InterleaveBuffers)
-		{
-			Buffer.SetNumZeroed(InParams.NumFrames);
+			const uint32 MinBufferSize = GetMinimumBufferSize(InParams.SampleRate);			
+			WriteNumFrames = FMath::Max(MinBufferSize, InParams.NumFrames);
+
+			for (FAlignedFloatBuffer& Buffer : InterleaveBuffers)
+			{
+				Buffer.SetNumZeroed(WriteNumFrames);
+			}
+
+			return true;
 		}
 
-		return FAudioMixerWasapiRenderStream::InitializeHardware(InParams);
+		return false;
 	}
 
 	bool FWasapiAggregateRenderStream::StartAudioStream()
@@ -61,20 +69,21 @@ namespace Audio
 			uint32 NumFramesPadding = 0;
 			AudioClient->GetCurrentPadding(&NumFramesPadding);
 
-			// NumFramesPerDeviceBuffer is the buffer size WASAPI allocated. It is gauranteed to 
+			// NumFramesPerDeviceBuffer is the buffer size WASAPI allocated. It is guaranteed to 
 			// be at least the amount requested. For example, if we request a 1024 frame buffer, WASAPI
 			// might allocate a 1056 frame buffer. The padding is subtracted from the allocated amount
 			// to determine how much space is available currently in the buffer.
 			const int32 NumFramesAvailable = NumFramesPerDeviceBuffer - NumFramesPadding;
-			if (NumFramesAvailable >= static_cast<int32>(RenderStreamParams.NumFrames))
+			
+			if (NumFramesAvailable >= WriteNumFrames)
 			{
 				uint8* RenderBufferPtr = nullptr;
 
-				if (SUCCEEDED(RenderClient->GetBuffer(RenderStreamParams.NumFrames, &RenderBufferPtr)))
+				if (SUCCEEDED(RenderClient->GetBuffer(WriteNumFrames, &RenderBufferPtr)))
 				{
-					InterleaveOutput((float*)RenderBufferPtr, RenderStreamParams.NumFrames);
+					InterleaveOutput((float*)RenderBufferPtr, WriteNumFrames);
 					
-					HRESULT Result = RenderClient->ReleaseBuffer(RenderStreamParams.NumFrames, 0 /* flags */);
+					HRESULT Result = RenderClient->ReleaseBuffer(WriteNumFrames, 0 /* flags */);
 					if (FAILED(Result))
 					{
 						++CallbackBufferErrors;
@@ -103,7 +112,7 @@ namespace Audio
 
 				if (DirectOutBuffers[Index].Num() >= InNumFrames)
 				{
-					const int32 NumFramesPopped = DirectOutBuffers[Index].Pop(InterleaveBuffers[Index].GetData(), RenderStreamParams.NumFrames);
+					const int32 NumFramesPopped = DirectOutBuffers[Index].Pop(InterleaveBuffers[Index].GetData(), InNumFrames);
 				}
 			}
 
diff --git a/Engine/Source/Runtime/Windows/AudioMixerWasapi/Private/WasapiAggregateRenderStream.h b/Engine/Source/Runtime/Windows/AudioMixerWasapi/Private/WasapiAggregateRenderStream.h
--- a/Engine/Source/Runtime/Windows/AudioMixerWasapi/Private/WasapiAggregateRenderStream.h
+++ b/Engine/Source/Runtime/Windows/AudioMixerWasapi/Private/WasapiAggregateRenderStream.h
@@ -22,11 +22,15 @@ namespace Audio
 		virtual bool InitializeHardware(const FWasapiRenderStreamParams& InParams) override;
 		virtual bool StartAudioStream() override;
 		virtual void SubmitDirectOutBuffer(const int32 InChannelIndex, const FAlignedFloatBuffer& InBuffer) override;
-		virtual void DeviceRenderCallback() override;
 		//~ End FAudioMixerWasapiRenderStream
 
+		//~ Begin IDeviceRenderCallback
+		virtual void DeviceRenderCallback() override;
+		//~ End IDeviceRenderCallback
+
 	private:
 
+		int32 WriteNumFrames = 0;
 		TArray<FAlignedFloatBuffer> InterleaveBuffers;
 		TArray<TCircularAudioBuffer<float>> DirectOutBuffers;
 
diff --git a/Engine/Source/Runtime/Windows/AudioMixerWasapi/Private/WasapiDefaultDeviceMgr.cpp b/Engine/Source/Runtime/Windows/AudioMixerWasapi/Private/WasapiDefaultDeviceMgr.cpp
--- a/Engine/Source/Runtime/Windows/AudioMixerWasapi/Private/WasapiDefaultDeviceMgr.cpp
+++ b/Engine/Source/Runtime/Windows/AudioMixerWasapi/Private/WasapiDefaultDeviceMgr.cpp
@@ -36,14 +36,14 @@ namespace Audio
 		return MainRenderStreamDevice.IsValid() && MainRenderStreamDevice->IsInitialized();
 	}
 
-	int32 FWasapiDefaultDeviceMgr::GetNumFrames(const int32 InNumReqestedFrames) const
+	int32 FWasapiDefaultDeviceMgr::GetNumFrames(const int32 InNumRequestedFrames) const
 	{ 
 		if (MainRenderStreamDevice.IsValid())
 		{
-			return MainRenderStreamDevice->GetNumFrames(InNumReqestedFrames);
+			return MainRenderStreamDevice->GetNumFrames(InNumRequestedFrames);
 		}
 
-		return InNumReqestedFrames;
+		return InNumRequestedFrames;
 	}
 
 	bool FWasapiDefaultDeviceMgr::OpenAudioStream(const TArray<FWasapiRenderStreamParams>& InParams)
@@ -63,7 +63,7 @@ namespace Audio
 			UE_LOG(LogAudioMixer, Error, TEXT("OpenAudioStream null EventHandle"));
 			return false;
 		}
-
+		
 		if (!MainRenderStreamDevice->OpenAudioStream(InParams[0], EventHandle))
 		{
 			UE_LOG(LogAudioMixer, Error, TEXT("OpenAudioStream failed to open main audio device"));
@@ -75,9 +75,14 @@ namespace Audio
 
 	bool FWasapiDefaultDeviceMgr::CloseAudioStream()
 	{
-		if (RenderDeviceThread.IsValid())
+		RenderDeviceThread.Reset();
+		
+		if (MainRenderStreamDevice.IsValid())
 		{
-			return MainRenderStreamDevice->CloseAudioStream();
+			bool bCloseSucceeded = MainRenderStreamDevice->CloseAudioStream();
+			MainRenderStreamDevice.Reset();
+
+			return bCloseSucceeded;
 		}
 
 		return false;
diff --git a/Engine/Source/Runtime/Windows/AudioMixerWasapi/Private/WasapiDefaultDeviceMgr.h b/Engine/Source/Runtime/Windows/AudioMixerWasapi/Private/WasapiDefaultDeviceMgr.h
--- a/Engine/Source/Runtime/Windows/AudioMixerWasapi/Private/WasapiDefaultDeviceMgr.h
+++ b/Engine/Source/Runtime/Windows/AudioMixerWasapi/Private/WasapiDefaultDeviceMgr.h
@@ -24,13 +24,13 @@ namespace Audio
 		virtual bool TeardownHardware() override;
 		virtual bool IsInitialized() const override;
 		virtual int32 GetNumDirectOutChannels() const override { return 0; }
-		virtual int32 GetNumFrames(const int32 InNumReqestedFrames) const override;
+		virtual int32 GetNumFrames(const int32 InNumRequestedFrames) const override;
 		virtual bool OpenAudioStream(const TArray<FWasapiRenderStreamParams>& InParams) override;
 		virtual bool CloseAudioStream() override;
 		virtual bool StartAudioStream() override;
 		virtual bool StopAudioStream() override;
 		virtual void SubmitBuffer(const uint8* InBuffer, const SIZE_T InNumFrames) override;
-		virtual void SubmitDirectOutBuffer(const int32 InDirectOutIndex, const FAlignedFloatBuffer& InBuffer) { }
+		virtual void SubmitDirectOutBuffer(const int32 InDirectOutIndex, const FAlignedFloatBuffer& InBuffer) override { }
 		//~ End IAudioMixerWasapiDeviceManager
 
 	private:
diff --git a/Engine/Source/Runtime/Windows/AudioMixerWasapi/Private/WasapiDefaultRenderStream.cpp b/Engine/Source/Runtime/Windows/AudioMixerWasapi/Private/WasapiDefaultRenderStream.cpp
--- a/Engine/Source/Runtime/Windows/AudioMixerWasapi/Private/WasapiDefaultRenderStream.cpp
+++ b/Engine/Source/Runtime/Windows/AudioMixerWasapi/Private/WasapiDefaultRenderStream.cpp
@@ -7,6 +7,27 @@
 
 namespace Audio
 {
+	bool FWasapiDefaultRenderStream::InitializeHardware(const FWasapiRenderStreamParams& InParams)
+	{
+		if (FAudioMixerWasapiRenderStream::InitializeHardware(InParams))
+		{
+			NumPartialBuffersToWrite = 1;
+			const uint32 MinBufferSize = GetMinimumBufferSize(InParams.SampleRate);			
+			if (MinBufferSize > InParams.NumFrames)
+			{
+				NumPartialBuffersToWrite = FMath::DivideAndRoundUp(MinBufferSize, InParams.NumFrames);
+			}
+
+			WriteNumFrames = NumPartialBuffersToWrite * InParams.NumFrames;
+			AudioBufferTotalBytes = WriteNumFrames * AudioFormat.GetFrameSizeInBytes();
+			PartialBufferNumBytes = InParams.NumFrames * AudioFormat.GetFrameSizeInBytes();
+
+			return true;
+		}
+		
+		return false;
+	}
+	
 	bool FWasapiDefaultRenderStream::TeardownHardware()
 	{
 		ReadNextBufferDelegate.Unbind();
@@ -23,26 +44,39 @@ namespace Audio
 			uint32 NumFramesPadding = 0;
 			AudioClient->GetCurrentPadding(&NumFramesPadding);
 
-			// NumFramesPerDeviceBuffer is the buffer size WASAPI allocated. It is gauranteed to 
+			// NumFramesPerDeviceBuffer is the buffer size WASAPI allocated. It is guaranteed to 
 			// be at least the amount requested. For example, if we request a 1024 frame buffer, WASAPI
 			// might allocate a 1056 frame buffer. The padding is subtracted from the allocated amount
 			// to determine how much space is available currently in the buffer.
 			const int32 NumFramesAvailable = NumFramesPerDeviceBuffer - NumFramesPadding;
-			if (NumFramesAvailable >= static_cast<int32>(RenderStreamParams.NumFrames))
+			
+			if (NumFramesAvailable >= WriteNumFrames)
 			{
-				check(!RenderBufferPtr);
+				check(RenderBufferView.IsEmpty());
 
-				if (SUCCEEDED(RenderClient->GetBuffer(RenderStreamParams.NumFrames, &RenderBufferPtr)))
+				uint8* BufferStartPtr = nullptr;
+				if (SUCCEEDED(RenderClient->GetBuffer(WriteNumFrames, &BufferStartPtr)))
 				{
-					ReadNextBufferDelegate.ExecuteIfBound();
+					TArrayView<uint8> BufferStartView(BufferStartPtr, AudioBufferTotalBytes);
+					
+					for (int32 i = 0; i < NumPartialBuffersToWrite; ++i)
+					{
+						const SIZE_T ByteOffset = i * PartialBufferNumBytes;
+						RenderBufferView = BufferStartView.Slice(ByteOffset, AudioBufferTotalBytes - ByteOffset);
+						
+						if (!ReadNextBufferDelegate.ExecuteIfBound())
+						{
+							++CallbackBufferErrors;
+						}
+					}
 
-					HRESULT Result = RenderClient->ReleaseBuffer(RenderStreamParams.NumFrames, 0 /* flags */);
+					HRESULT Result = RenderClient->ReleaseBuffer(WriteNumFrames, 0 /* flags */);
 					if (FAILED(Result))
 					{
 						++CallbackBufferErrors;
 					}
 
-					RenderBufferPtr = nullptr;
+					RenderBufferView = TArrayView<uint8>();
 				}
 				else
 				{
@@ -54,9 +88,12 @@ namespace Audio
 
 	void FWasapiDefaultRenderStream::SubmitBuffer(const uint8* InBuffer, const SIZE_T InNumFrames)
 	{
-		if (RenderBufferPtr)
+		if (RenderBufferView.Num() >= InNumFrames)
 		{
-			FMemory::Memcpy(RenderBufferPtr, InBuffer, InNumFrames * AudioFormat.GetFrameSizeInBytes());
+			check(InNumFrames == RenderStreamParams.NumFrames);
+			const SIZE_T NumBytes = InNumFrames * AudioFormat.GetFrameSizeInBytes();
+
+			FMemory::Memcpy(RenderBufferView.GetData(), InBuffer, NumBytes);
 		}
 	}
 }
diff --git a/Engine/Source/Runtime/Windows/AudioMixerWasapi/Private/WasapiDefaultRenderStream.h b/Engine/Source/Runtime/Windows/AudioMixerWasapi/Private/WasapiDefaultRenderStream.h
--- a/Engine/Source/Runtime/Windows/AudioMixerWasapi/Private/WasapiDefaultRenderStream.h
+++ b/Engine/Source/Runtime/Windows/AudioMixerWasapi/Private/WasapiDefaultRenderStream.h
@@ -20,19 +20,36 @@ namespace Audio
 		virtual ~FWasapiDefaultRenderStream() = default;
 
 		//~ Begin FAudioMixerWasapiRenderStream
+		virtual bool InitializeHardware(const FWasapiRenderStreamParams& InParams) override;
 		virtual bool TeardownHardware() override;
 		virtual void SubmitBuffer(const uint8* Buffer, const SIZE_T InNumFrames) override;
-		virtual void DeviceRenderCallback() override;
 		//~ End FAudioMixerWasapiRenderStream
 
+		//~ Begin IDeviceRenderCallback
+		virtual void DeviceRenderCallback() override;
+		//~ End IDeviceRenderCallback
+		
 		FAudioMixerReadNextBufferDelegate& OnReadNextBuffer() { return ReadNextBufferDelegate; }
 
 	private:
 
+		/** Number of frames to write every Wasapi period (must be at least the minimum required for current driver) */
+		int32 WriteNumFrames = 0;
+
+		/** Total number of bytes per audio buffer copied to Wasapi each buffer period. */
+		SIZE_T AudioBufferTotalBytes = 0;
+
+		/** Number of bytes per partial buffer where a partial buffer is used when the engine buffer size
+		 *  is less than the Wasapi drivers buffer size. */
+		SIZE_T PartialBufferNumBytes = 0;
+
+		/** Number of partial buffers needed to create a full Wasapi buffer. Used for buffer sizes smaller than the minimum for the current driver. */
+		int32 NumPartialBuffersToWrite = 0;
+		
 		/** Delegate called each buffer callback to signal the mixer to process the next buffer. */
 		FAudioMixerReadNextBufferDelegate ReadNextBufferDelegate;
 
 		/** Pointer to WASAPI render audio buffer filled in each callback by the mixer in SubmitBuffer(). */
-		uint8* RenderBufferPtr = nullptr;
+		TArrayView<uint8> RenderBufferView;
 	};
 }
diff --git a/Engine/Source/Runtime/Windows/AudioMixerXAudio2/Private/AudioMixerPlatformWindows.cpp b/Engine/Source/Runtime/Windows/AudioMixerXAudio2/Private/AudioMixerPlatformWindows.cpp
--- a/Engine/Source/Runtime/Windows/AudioMixerXAudio2/Private/AudioMixerPlatformWindows.cpp
+++ b/Engine/Source/Runtime/Windows/AudioMixerXAudio2/Private/AudioMixerPlatformWindows.cpp
@@ -121,19 +121,15 @@ namespace Audio
 			return;
 		}
 		
-		if (AudioDeviceSwapCriticalSection.TryLock())
+		// If the device that was added is our original device and our current device is NOT our original device, 
+		// move our audio stream to this newly added device.
+		const FString AudioDeviceId = GetOriginalAudioDeviceId();
+		if (AudioStreamInfo.DeviceInfo.DeviceId != AudioDeviceId && DeviceId == AudioDeviceId)
 		{
-			// If the device that was added is our original device and our current device is NOT our original device, 
-			// move our audio stream to this newly added device.
-			if (AudioStreamInfo.DeviceInfo.DeviceId != OriginalAudioDeviceId && DeviceId == OriginalAudioDeviceId)
-			{
-				UE_LOG(LogAudioMixer, Warning, TEXT("FMixerPlatformXAudio2: Original audio device re-added. Moving audio back to original audio device: DeviceName=%s, bRenderDevice=%d, InstanceID=%d"), 
-					*WindowsNotificationClient->GetFriendlyName(*OriginalAudioDeviceId), (int32)bIsRenderDevice, InstanceID);
-
-				RequestDeviceSwap(OriginalAudioDeviceId, /*force */ true, TEXT("FMixerPlatformXAudio2::OnDeviceAdded"));
-			}
+			UE_LOG(LogAudioMixer, Warning, TEXT("FMixerPlatformXAudio2: Original audio device re-added. Moving audio back to original audio device: DeviceName=%s, bRenderDevice=%d, InstanceID=%d"), 
+				*WindowsNotificationClient->GetFriendlyName(*AudioDeviceId), (int32)bIsRenderDevice, InstanceID);
 
-			AudioDeviceSwapCriticalSection.Unlock();
+			RequestDeviceSwap(AudioDeviceId, /*force */ true, TEXT("FMixerPlatformXAudio2::OnDeviceAdded"));
 		}
 
 		if (UAudioDeviceNotificationSubsystem* AudioDeviceNotifSubsystem = UAudioDeviceNotificationSubsystem::Get())
@@ -150,17 +146,13 @@ namespace Audio
 			return;
 		}
 		
-		if (AudioDeviceSwapCriticalSection.TryLock())
+		// If the device we're currently using was removed... then switch to the new default audio device.
+		if (AudioStreamInfo.DeviceInfo.DeviceId == DeviceId)
 		{
-			// If the device we're currently using was removed... then switch to the new default audio device.
-			if (AudioStreamInfo.DeviceInfo.DeviceId == DeviceId)
-			{
-				UE_LOG(LogAudioMixer, Warning, TEXT("FMixerPlatformXAudio2: Audio device removed [%s], falling back to other windows default device. bIsRenderDevice=%d, InstanceID=%d"), 
-					*WindowsNotificationClient->GetFriendlyName(DeviceId), (int32)bIsRenderDevice, InstanceID);
+			UE_LOG(LogAudioMixer, Warning, TEXT("FMixerPlatformXAudio2: Audio device removed [%s], falling back to other windows default device. bIsRenderDevice=%d, InstanceID=%d"), 
+				*WindowsNotificationClient->GetFriendlyName(DeviceId), (int32)bIsRenderDevice, InstanceID);
 
-				RequestDeviceSwap(TEXT(""), /* force */ true, TEXT("FMixerPlatformXAudio2::OnDeviceRemoved"));
-			}
-			AudioDeviceSwapCriticalSection.Unlock();
+			RequestDeviceSwap(TEXT(""), /* force */ true, TEXT("FMixerPlatformXAudio2::OnDeviceRemoved"));
 		}
 
 		if (UAudioDeviceNotificationSubsystem* AudioDeviceNotifSubsystem = UAudioDeviceNotificationSubsystem::Get())
diff --git a/Engine/Source/Runtime/Windows/AudioMixerXAudio2/Private/AudioMixerPlatformXAudio2.cpp b/Engine/Source/Runtime/Windows/AudioMixerXAudio2/Private/AudioMixerPlatformXAudio2.cpp
--- a/Engine/Source/Runtime/Windows/AudioMixerXAudio2/Private/AudioMixerPlatformXAudio2.cpp
+++ b/Engine/Source/Runtime/Windows/AudioMixerXAudio2/Private/AudioMixerPlatformXAudio2.cpp
@@ -22,6 +22,7 @@
 #include "CoreMinimal.h"
 #include "Logging/LogMacros.h"
 #include "ToStringHelpers.h"
+#include "ScopedCom.h"
 
 THIRD_PARTY_INCLUDES_START
 // Including initguid.h will define the PKEY symbols below which area used cross-platform
@@ -130,11 +131,9 @@ namespace Audio
 
 	FMixerPlatformXAudio2::FMixerPlatformXAudio2()
 		: XAudio2Dll(nullptr)
-		, bDeviceChanged(false)
 		, XAudio2System(nullptr)
 		, OutputAudioStreamMasteringVoice(nullptr)
 		, OutputAudioStreamSourceVoice(nullptr)
-		, LastDeviceSwapTime(0.0)
 		, TimeSinceNullDeviceWasLastChecked(0.0f)		
 		, bIsInitialized(false)
 		, bIsDeviceOpen(false)
@@ -160,328 +159,308 @@ namespace Audio
 	{
 		bool bDidStopGeneratingAudio = false;
 #if PLATFORM_WINDOWS
-		SCOPED_NAMED_EVENT(FMixerPlatformXAudio2_CheckThreadedDeviceSwap, FColor::Blue);
+		bDidStopGeneratingAudio = FAudioMixerPlatformSwappable::CheckThreadedDeviceSwap();
+#endif //PLATFORM_WINDOWS
+		return bDidStopGeneratingAudio;
+	}
 
-		static float ThreadedSwapDebugExtraTimeMsCVar = 0;
-		static FAutoConsoleVariableRef CVarOverrunTimeout(
-			TEXT("au.ThreadedSwapDebugExtraTime"),
-			ThreadedSwapDebugExtraTimeMsCVar,
-			TEXT("Simulate a slow device swap by adding addional time to the swap task"),
-			ECVF_Default);
+	bool FMixerPlatformXAudio2::PreDeviceSwap()
+	{
+		// Access to device swap context must be protected by DeviceSwapCriticalSection
+		FScopeLock Lock(&DeviceSwapCriticalSection);
+		
+		if (DeviceSwapContext.IsValid())
+		{
+			// Finish initializing the device swap context
+			DeviceSwapContext->PreviousSystem = XAudio2System;
+			DeviceSwapContext->PreviousMasteringVoice = OutputAudioStreamMasteringVoice;
+			DeviceSwapContext->PreviousSourceVoice = OutputAudioStreamSourceVoice;
+			DeviceSwapContext->Callbacks = &OutputVoiceCallback;
+			DeviceSwapContext->RenderingSampleRate = OpenStreamParams.SampleRate;
+
+			// NULL our system / master / source voices. (as they are about to be torn down async).
+			XAudio2System = nullptr;
+			OutputAudioStreamMasteringVoice = nullptr;
+			OutputAudioStreamSourceVoice = nullptr;
+
+			UE_LOG(LogAudioMixer, Verbose, TEXT("FMixerPlatformXAudio2::PreDeviceSwap - Starting swap to [%s]"), DeviceSwapContext->RequestedDeviceId.IsEmpty() ? TEXT("[System Default]") : *DeviceSwapContext->RequestedDeviceId);
+
+			return true;
+		}
+		else
+		{
+			UE_LOG(LogAudioMixer, Warning, TEXT("FMixerPlatformXAudio2::PreDeviceSwap - null device swap context"));
+			return false;
+		}
+
+		return false;
+	}
 	
-		if (bMoveAudioStreamToNewAudioDevice || ActiveDeviceSwap.IsValid())
+	void FMixerPlatformXAudio2::EnqueueAsyncDeviceSwap()
+	{
+		UE_LOG(LogAudioMixer, Verbose, TEXT("FMixerPlatformXAudio2::EnqueueAsyncDeviceSwap - enqueuing async device swap"));
+		FScopeLock Lock(&DeviceSwapCriticalSection);
+		
+		TFunction<TUniquePtr<FDeviceSwapResult>()> AsyncDeviceSwap = [this]() mutable -> TUniquePtr<FDeviceSwapResult>
 		{
-			// Start a job?
-			if (!ActiveDeviceSwap.IsValid())
+			// Transfer ownership of DeviceSwapContext to the async task.
+			TUniquePtr<FXAudio2DeviceSwapContext> TempContext;
 			{
-				SCOPED_NAMED_EVENT(FMixerPlatformXAudio2_CheckThreadedDeviceSwap_StartAsyncSwap, FColor::Blue);
-
-				UE_LOG(LogAudioMixer, Verbose, TEXT("FMixerPlatformXAudio2::CheckThreadedDeviceSwap() Starting swap to [%s]"), NewAudioDeviceId.IsEmpty() ? TEXT("[System Default]") : *NewAudioDeviceId);
+				FScopeLock Lock(&DeviceSwapCriticalSection);
 
-				// Toggle move to audio stream now we're spawning a new job.
-				bMoveAudioStreamToNewAudioDevice = false;
+				if (AudioStreamInfo.StreamState == EAudioOutputStreamState::SwappingDevice)
+				{
+					TempContext = MoveTemp(DeviceSwapContext);
+				}
+			}
+			
+			return PerformDeviceSwap(MoveTemp(TempContext));
+		};
+		SetActiveDeviceSwapFuture(Async(EAsyncExecution::TaskGraph, MoveTemp(AsyncDeviceSwap)));
+	}
+	
+	bool FMixerPlatformXAudio2::PostDeviceSwap()
+	{
+		// Once the context is handed off to the device swap routine (either async or synchronous),
+		// it should no longer be valid.
+		check(!DeviceSwapContext.IsValid());
+		bool bDidSucceed = false;
+		const FXAudio2DeviceSwapResult* DeviceSwapResult = static_cast<const FXAudio2DeviceSwapResult*>(GetDeviceSwapResult());
+		
+		if (DeviceSwapResult && DeviceSwapResult->IsNewDeviceReady())
+		{
+			FScopeLock Lock(&DeviceSwapCriticalSection);
 
-				// Look up device. Blank name looks up current default.
-				FName NewDeviceName = *NewAudioDeviceId;
-				FAudioPlatformDeviceInfo NewDevice;
-				bool bSwitchingToValidDevice = false;
-				check(GetDeviceInfoCache());
-				if (TOptional<FAudioPlatformDeviceInfo> DeviceInfo = GetDeviceInfoCache()->FindActiveOutputDevice(NewDeviceName))
+			XAudio2System = DeviceSwapResult->NewSystem;
+			OutputAudioStreamMasteringVoice = DeviceSwapResult->NewMasteringVoice;
+			OutputAudioStreamSourceVoice = DeviceSwapResult->NewSourceVoice;
+		
+			// Success?
+			if (XAudio2System && OutputAudioStreamSourceVoice && OutputAudioStreamMasteringVoice)
+			{
+				HRESULT Result;
+				{
+					SCOPED_NAMED_EVENT(FMixerPlatformXAudio2_PostDeviceSwap_StartEngine, FColor::Blue);
+					Result = XAudio2System->StartEngine();
+				}
+				if (SUCCEEDED(Result))
 				{
-					if (DeviceInfo->NumChannels <= XAUDIO2_MAX_AUDIO_CHANNELS &&
-						DeviceInfo->SampleRate >= XAUDIO2_MIN_SAMPLE_RATE &&
-						DeviceInfo->SampleRate <= XAUDIO2_MAX_SAMPLE_RATE)
+					// Copy our new Device Info into our active one.
+					AudioStreamInfo.DeviceInfo = DeviceSwapResult->DeviceInfo;
+
+					// Display our new XAudio2 Mastering voice details.
+					UE_LOG(LogAudioMixer, Display, TEXT("FMixerPlatformXAudio2::PostDeviceSwap - successful Swap new Device is (NumChannels=%u, SampleRate=%u, DeviceID=%s, Name=%s), Reason=%s, InstanceID=%d, DurationMS=%.2f"),
+						(uint32)AudioStreamInfo.DeviceInfo.NumChannels, (uint32)AudioStreamInfo.DeviceInfo.SampleRate, *AudioStreamInfo.DeviceInfo.DeviceId, *AudioStreamInfo.DeviceInfo.Name,
+						*DeviceSwapResult->SwapReason, InstanceID, DeviceSwapResult->SuccessfulDurationMs);
+
+					// Reinitialize the output circular buffer to match the buffer math of the new audio device.
+					const int32 NumOutputSamples = AudioStreamInfo.NumOutputFrames * AudioStreamInfo.DeviceInfo.NumChannels;
+					if (ensure(NumOutputSamples > 0))
 					{
-						bSwitchingToValidDevice = true;
-						NewDevice = *DeviceInfo;
+						OutputBuffer.Init(AudioStreamInfo.AudioMixer, NumOutputSamples, NumOutputBuffers, AudioStreamInfo.DeviceInfo.Format);
 					}
-					else
-					{
-						UE_LOG(LogAudioMixer, Warning, TEXT("Ignoring attempt to switch to device with unsupported params: Channels=%u, SampleRate=%u, Id=%s, Name=%s"),
-							(uint32)DeviceInfo->NumChannels, (uint32)DeviceInfo->SampleRate, *DeviceInfo->DeviceId, *DeviceInfo->Name);
 
-						NewDevice.Reset();
-					}
+					bDidSucceed = true;
 				}
 				else
 				{
-					// Failed to find any device/default or otherwise. Reset our current device and do a switch to a null device.
-					NewDevice.Reset();
+					XAUDIO2_LOG_AND_HANDLE_ON_FAIL("XAudio2System->StartEngine()", Result, {/* nop */});
 				}
-				
-				TFunction<FXAudio2AsyncCreateResult()> AsyncDeleteCreate;
+			}
+			else // We either failed to init or deliberately switched to null device.
+			{
+				// Null renderer doesn't/shouldn't care about the format, so leave the format as it was before.
+			}
+		}
 
-				{
-					// Lock for creation of lambda as we are copying state.
-					FScopeLock Lock(&AudioDeviceSwapCriticalSection);
-
-					// Build Delete/Create Async lambda
-					AsyncDeleteCreate = [
-						System = XAudio2System,
-						SourceVoice = OutputAudioStreamSourceVoice,
-						MasterVoice = OutputAudioStreamMasteringVoice,
-						NewDevice,
-						&Callbacks = OutputVoiceCallback,
-						RenderingSampleRate = OpenStreamParams.SampleRate,
-						bSwitchingToValidDevice,
-						SwapReason = DeviceSwapReason
-					]() mutable->FXAudio2AsyncCreateResult
-						{
-							uint64_t StartTimeCycles = FPlatformTime::Cycles64();
-							
-							SCOPED_NAMED_EVENT(FMixerPlatformXAudio2_AsyncDeleteCreate, FColor::Blue);
-							UE_LOG(LogAudioMixer, Verbose, TEXT("FMixerPlatformXAudio2::CheckThreadedDeviceSwap() - AsyncTask Start. Because=%s"), *SwapReason);
-
-							// New thread might not have COM setup.
-							FPlatformMisc::CoInitialize();
-
-							{
-								// Stop old engine running.
-								if (System)
-								{
-									SCOPED_NAMED_EVENT(FMixerPlatformXAudio2_AsyncDeleteCreate_StopEngine, FColor::Blue);
-									System->StopEngine();
-								}
-
-								// Kill source voice.
-								if (SourceVoice)
-								{
-									{
-										SCOPED_NAMED_EVENT(FMixerPlatformXAudio2_AsyncDeleteCreate_FlushSourceBuffers, FColor::Blue);
-										XAUDIO2_CALL_AND_HANDLE_ERROR(SourceVoice->FlushSourceBuffers(), /* nop */ );
-									}
-									SCOPED_NAMED_EVENT(FMixerPlatformXAudio2_AsyncDeleteCreate_DestroySourceVoice, FColor::Blue);
-									SourceVoice->DestroyVoice();
-									SourceVoice = nullptr;
-								}
-
-								// Now destroy the mastering voice
-								if (MasterVoice)
-								{
-									SCOPED_NAMED_EVENT(FMixerPlatformXAudio2_AsyncDeleteCreate_DestroyMasterVoice, FColor::Blue);
-									MasterVoice->DestroyVoice();
-									MasterVoice = nullptr;
-								}
-
-
-								// Destroy System
-								{
-									SCOPED_NAMED_EVENT(FMixerPlatformXAudio2_AsyncDeleteCreate_DestroySystem, FColor::Blue);
-									SAFE_RELEASE(System);
-								}
-							}
-
-							FXAudio2AsyncCreateResult Results;
-							Results.SwapReason = SwapReason;
-
-							// Don't attempt to create a new setup if there's no devices available.
-							if (!bSwitchingToValidDevice)
-							{
-								return {};
-							}
-
-							// Create System.
-							{
-								SCOPED_NAMED_EVENT(FMixerPlatformXAudio2_AsyncDeleteCreate_CreateSystem, FColor::Blue);
-								XAUDIO2_CALL_AND_HANDLE_ERROR(XAudio2Create(&Results.XAudio2System, 0, GetXAudio2ProcessorsToUse()), return {});
-							}
-
-							// Create Master
-							{
-								check(NewDevice.NumChannels <= XAUDIO2_MAX_AUDIO_CHANNELS);
-								check(NewDevice.SampleRate >= XAUDIO2_MIN_SAMPLE_RATE);
-								check(NewDevice.SampleRate <= XAUDIO2_MAX_SAMPLE_RATE);
-
-								SCOPED_NAMED_EVENT(FMixerPlatformXAudio2_AsyncDeleteCreate_CreateMasterVoice, FColor::Blue);
-								HRESULT Result = Results.XAudio2System->CreateMasteringVoice(
-									&Results.OutputAudioStreamMasteringVoice,
-									(uint32)NewDevice.NumChannels,
-									(uint32)NewDevice.SampleRate,
-									0,
-									*NewDevice.DeviceId,
-									nullptr,
-									AudioCategory_GameEffects
-								);
-								if (FAILED(Result))
-								{
-									UE_LOG(LogAudioMixer, Error, TEXT("XAudio2System->CreateMasteringVoice() -> 0x%X: %s (line: %d) with Args (NumChannels=%u, SampleRate=%u, DeviceID=%s Name=%s)"),
-										Result, *Audio::ToErrorFString(Result), __LINE__, (uint32)NewDevice.NumChannels, (uint32)NewDevice.SampleRate, *NewDevice.DeviceId, *NewDevice.Name);
-									SAFE_RELEASE(Results.XAudio2System);
-									return {}; // FAIL.
-								}
-							}
-
-							// Create Source Voice.
-							{
-								SCOPED_NAMED_EVENT(FMixerPlatformXAudio2_AsyncDeleteCreate_CreateSourceVoice, FColor::Blue);
-
-								// Setup the format of the output source voice
-								WAVEFORMATEX Format = { 0 };
-								Format.nChannels = NewDevice.NumChannels;
-								Format.nSamplesPerSec = RenderingSampleRate;		// NOTE: We use the Rendering sample rate here.
-								Format.wFormatTag = WAVE_FORMAT_IEEE_FLOAT;
-								Format.nAvgBytesPerSec = Format.nSamplesPerSec * sizeof(float) * Format.nChannels;
-								Format.nBlockAlign = sizeof(float) * Format.nChannels;
-								Format.wBitsPerSample = sizeof(float) * 8;
-
-								// Create the output source voice
-								HRESULT Result = Results.XAudio2System->CreateSourceVoice(
-									&Results.OutputAudioStreamSourceVoice,
-									&Format,
-									XAUDIO2_VOICE_NOPITCH,
-									XAUDIO2_DEFAULT_FREQ_RATIO,
-									&Callbacks,
-									nullptr,
-									nullptr
-								);
-
-								if (FAILED(Result))
-								{
-									Results.OutputAudioStreamMasteringVoice->DestroyVoice();
-									Results.OutputAudioStreamMasteringVoice = nullptr;
-									SAFE_RELEASE(Results.XAudio2System);
-									XAUDIO2_LOG_AND_HANDLE_ON_FAIL("XAudio2System->CreateSourceVoice", Result, return {});
-								}
-							}
-
-							// Optionally for testing, sleep for some duration in order to help repro race conditions.
-							if (ThreadedSwapDebugExtraTimeMsCVar > 0.f)
-							{
-								FPlatformProcess::Sleep(ThreadedSwapDebugExtraTimeMsCVar/1000.f); 
-							}
-
-							// Listen session for changes to this device.
-							RegisterForSessionEvents(NewDevice.DeviceId);
-
-							FPlatformMisc::CoUninitialize();
-
-							// Success.
-							Results.SuccessfullDurationMs = FPlatformTime::ToMilliseconds64(FPlatformTime::Cycles64() - StartTimeCycles);
-							Results.DeviceInfo = NewDevice;
-							return Results;
-						}; // End of Lambda.
-
-						// Still inside CS here.
-						// NULL our system / master / source voices. (as they are about to be torn down async).
-						bIsInDeviceSwap = true;
-						XAudio2System = nullptr;
-						OutputAudioStreamMasteringVoice = nullptr;
-						OutputAudioStreamSourceVoice = nullptr;
-					} // End of Lambda Creation scope (with CS).
+		ResetActiveDeviceSwapFuture();
+		
+		return bDidSucceed;
+	}
+	
 
-				
-					{
-						SCOPED_NAMED_EVENT(FMixerPlatformXAudio2_AsyncDeleteCreate_StopSource, FColor::Blue);
+	void FMixerPlatformXAudio2::SynchronousDeviceSwap()
+	{
+		// Transfer ownership of DeviceSwapContext memory to the device swap routine
+		TUniquePtr<FDeviceSwapResult> DeviceSwapResult = PerformDeviceSwap(TUniquePtr<FXAudio2DeviceSwapContext>(MoveTemp(DeviceSwapContext)));
 
-						// Ask source voice to stop. 
-						if (OutputAudioStreamSourceVoice)
-						{
-							XAUDIO2_CALL_AND_HANDLE_ERROR(OutputAudioStreamSourceVoice->Stop(), /* nop */ );
-						}
-					}
+		// Set the promise and future result to replicate what the async task does
+		TPromise<TUniquePtr<FDeviceSwapResult>> Promise;
+		
+		// It's ok if DeviceSwapResult is null here. It indicates an invalid device which will be handled.
+		Promise.SetValue(MoveTemp(DeviceSwapResult));
+		SetActiveDeviceSwapFuture(Promise.GetFuture());
+	}
+	
+	TUniquePtr<FDeviceSwapResult> FMixerPlatformXAudio2::PerformDeviceSwap(TUniquePtr<FXAudio2DeviceSwapContext>&& InDeviceContext)
+	{
+		// Static method enforces no other state sharing occurs with the object that called it. InDeviceContext
+		// should have no other references outside of this method so that the device swap operation is isolated.
+		SCOPED_NAMED_EVENT(FMixerPlatformXAudio2_PerformDeviceSwap, FColor::Blue);
 
-					// Start a new device swap.
-					{
-						SCOPED_NAMED_EVENT(FMixerPlatformXAudio2_AsyncDeleteCreate_Async, FColor::Blue);
-						ActiveDeviceSwap = Async(EAsyncExecution::TaskGraph, MoveTemp(AsyncDeleteCreate));
-					}
+		const uint64 StartTimeCycles = FPlatformTime::Cycles64();
 
-					if (!bIsUsingNullDevice)
-					{
-						StartRunningNullDevice();
-					}
+		static float ThreadedSwapDebugExtraTimeMsCVar = 0;
+		static FAutoConsoleVariableRef CVarOverrunTimeout(
+			TEXT("au.ThreadedSwapDebugExtraTime"),
+			ThreadedSwapDebugExtraTimeMsCVar,
+			TEXT("Simulate a slow device swap by adding additional time to the swap task"),
+			ECVF_Default);
 
-					// Return FALSE, as null can run with whatever the same config we were using before.
-					return false;
-			}
-			else // Running...
+		// New thread might not have COM setup.
+		FScopedCoInitialize ScopedCoInitialize;
+
+		// Critical section not required here as this function is now sole owner of context
+		if (!InDeviceContext.IsValid())
+		{
+			UE_LOG(LogAudioMixer, Error, TEXT("FMixerPlatformXAudio2::PerformDeviceSwap - failed due to invalid DeviceSwapContext"));
+			return {};
+		}
+
+		UE_LOG(LogAudioMixer, Verbose, TEXT("FMixerPlatformXAudio2::PerformDeviceSwap - AsyncTask Start. Because=%s"), *InDeviceContext->DeviceSwapReason);
+
+		// Stop old engine running.
+		if (InDeviceContext->PreviousSystem)
+		{
+			SCOPED_NAMED_EVENT(FMixerPlatformXAudio2_AsyncDeleteCreate_StopEngine, FColor::Blue);
+			InDeviceContext->PreviousSystem->StopEngine();
+		}
+
+		// Kill source voice.
+		if (InDeviceContext->PreviousSourceVoice)
+		{
 			{
-				if (ActiveDeviceSwap.IsReady())  // Finished?
-				{
-					SCOPED_NAMED_EVENT(FMixerPlatformXAudio2_CheckThreadedDeviceSwap_FinishedJobSwitchOver, FColor::Blue);
+				SCOPED_NAMED_EVENT(FMixerPlatformXAudio2_AsyncDeleteCreate_FlushSourceBuffers, FColor::Blue);
+				XAUDIO2_CALL_AND_HANDLE_ERROR(InDeviceContext->PreviousSourceVoice->FlushSourceBuffers(), /* nop */ );
+			}
+			SCOPED_NAMED_EVENT(FMixerPlatformXAudio2_AsyncDeleteCreate_DestroySourceVoice, FColor::Blue);
+			InDeviceContext->PreviousSourceVoice->DestroyVoice();
+			InDeviceContext->PreviousSourceVoice = nullptr;
+		}
 
-					XAudio2System = ActiveDeviceSwap.Get().XAudio2System;
-					OutputAudioStreamMasteringVoice = ActiveDeviceSwap.Get().OutputAudioStreamMasteringVoice;
-					OutputAudioStreamSourceVoice = ActiveDeviceSwap.Get().OutputAudioStreamSourceVoice;
+		// Now destroy the mastering voice
+		if (InDeviceContext->PreviousMasteringVoice)
+		{
+			SCOPED_NAMED_EVENT(FMixerPlatformXAudio2_AsyncDeleteCreate_DestroyMasterVoice, FColor::Blue);
+			InDeviceContext->PreviousMasteringVoice->DestroyVoice();
+			InDeviceContext->PreviousMasteringVoice = nullptr;
+		}
 
-					// Success?
-					if (XAudio2System && OutputAudioStreamSourceVoice && OutputAudioStreamMasteringVoice)
-					{
-						HRESULT Result;
-						{
-							SCOPED_NAMED_EVENT(FMixerPlatformXAudio2_CheckThreadedDeviceSwap_StartEngine, FColor::Blue);
-							Result = XAudio2System->StartEngine();
-						}
-						if (SUCCEEDED(Result))
-						{
-							if (bIsUsingNullDevice)
-							{
-								StopRunningNullDevice();
-							}
-
-							StopGeneratingAudio();
-							bDidStopGeneratingAudio = true;
-
-							// Copy our new Device Info into our active one.
-							AudioStreamInfo.DeviceInfo = ActiveDeviceSwap.Get().DeviceInfo;
-
-							// Display our new XAudio2 Mastering voice details.
-							UE_LOG(LogAudioMixer, Display, TEXT("Successful Swap new Device is (NumChannels=%u, SampleRate=%u, DeviceID=%s, Name=%s), Reason=%s, InstanceID=%d, DurationMS=%.2f"),
-								(uint32)AudioStreamInfo.DeviceInfo.NumChannels, (uint32)AudioStreamInfo.DeviceInfo.SampleRate, *AudioStreamInfo.DeviceInfo.DeviceId, *AudioStreamInfo.DeviceInfo.Name, *ActiveDeviceSwap.Get().SwapReason, InstanceID, ActiveDeviceSwap.Get().SuccessfullDurationMs);
-
-							// Reinitialize the output circular buffer to match the buffer math of the new audio device.
-							const int32 NumOutputSamples = AudioStreamInfo.NumOutputFrames * AudioStreamInfo.DeviceInfo.NumChannels;
-							if (ensure(NumOutputSamples > 0))
-							{
-								OutputBuffer.Init(AudioStreamInfo.AudioMixer, NumOutputSamples, NumOutputBuffers, AudioStreamInfo.DeviceInfo.Format);
-							}
-						}
-						else
-						{
-							XAUDIO2_LOG_AND_HANDLE_ON_FAIL("XAudio2System->StartEngine()", Result, {/* nop */});
-						}
-					}
-					else // We either failed to init or deliberately switched to null device.
-					{
-						// Null renderer doesn't/shouldn't care about the format, so leave the format as it was before.
-					}
+		// Destroy System
+		{
+			SCOPED_NAMED_EVENT(FMixerPlatformXAudio2_AsyncDeleteCreate_DestroySystem, FColor::Blue);
+			SAFE_RELEASE(InDeviceContext->PreviousSystem);
+			InDeviceContext->PreviousSystem = nullptr;
+		}
 
-					{
-						SCOPED_NAMED_EVENT(FMixerPlatformXAudio2_CheckThreadedDeviceSwap_EndSwap, FColor::Blue);
-						bIsInDeviceSwap = false;
-						ActiveDeviceSwap.Reset();
-					}
+		// Don't attempt to create a new setup if there's no devices available.
+		if (!InDeviceContext->NewDevice.IsSet())
+		{
+			return {};
+		}
+
+		TUniquePtr<FXAudio2DeviceSwapResult> DeviceSwapResult = MakeUnique<FXAudio2DeviceSwapResult>();
+		
+		// Create System.
+		{
+			SCOPED_NAMED_EVENT(FMixerPlatformXAudio2_AsyncDeleteCreate_CreateSystem, FColor::Blue);
+			XAUDIO2_CALL_AND_HANDLE_ERROR(XAudio2Create(&DeviceSwapResult->NewSystem, 0, GetXAudio2ProcessorsToUse()), return {});
+		}
+
+		// Create Master
+		{
+			check(InDeviceContext->NewDevice->NumChannels <= XAUDIO2_MAX_AUDIO_CHANNELS);
+			check(InDeviceContext->NewDevice->SampleRate >= XAUDIO2_MIN_SAMPLE_RATE);
+			check(InDeviceContext->NewDevice->SampleRate <= XAUDIO2_MAX_SAMPLE_RATE);
+
+			SCOPED_NAMED_EVENT(FMixerPlatformXAudio2_AsyncDeleteCreate_CreateMasterVoice, FColor::Blue);
+			HRESULT Result = DeviceSwapResult->NewSystem->CreateMasteringVoice(
+				&DeviceSwapResult->NewMasteringVoice,
+				(uint32)InDeviceContext->NewDevice->NumChannels,
+				(uint32)InDeviceContext->NewDevice->SampleRate,
+				0,
+				*InDeviceContext->NewDevice->DeviceId,
+				nullptr,
+				AudioCategory_GameEffects
+			);
+			if (FAILED(Result))
+			{
+				UE_LOG(LogAudioMixer, Error, TEXT("XAudio2System->CreateMasteringVoice() -> 0x%X: %s (line: %d) with Args (NumChannels=%u, SampleRate=%u, DeviceID=%s Name=%s)"),
+					Result, *Audio::ToErrorFString(Result), __LINE__, (uint32)InDeviceContext->NewDevice->NumChannels,
+					(uint32)InDeviceContext->NewDevice->SampleRate, *InDeviceContext->NewDevice->DeviceId, *InDeviceContext->NewDevice->Name);
+
+				if (DeviceSwapResult->NewMasteringVoice)
+				{
+					DeviceSwapResult->NewMasteringVoice->DestroyVoice();
+					DeviceSwapResult->NewMasteringVoice = nullptr;
 				}
+				
+				SAFE_RELEASE(DeviceSwapResult->NewSystem);
+				return {}; // FAIL.
 			}
 		}
-#endif //PLATFORM_WINDOWS
 
-		return bDidStopGeneratingAudio;
-	}
+		// Create Source Voice.
+		{
+			SCOPED_NAMED_EVENT(FMixerPlatformXAudio2_AsyncDeleteCreate_CreateSourceVoice, FColor::Blue);
 
-	bool FMixerPlatformXAudio2::AllowDeviceSwap()
-	{
+			// Setup the format of the output source voice
+			WAVEFORMATEX Format = { 0 };
+			Format.nChannels = InDeviceContext->NewDevice->NumChannels;
+			Format.nSamplesPerSec = InDeviceContext->RenderingSampleRate;		// NOTE: We use the Rendering sample rate here.
+			Format.wFormatTag = WAVE_FORMAT_IEEE_FLOAT;
+			Format.nAvgBytesPerSec = Format.nSamplesPerSec * sizeof(float) * Format.nChannels;
+			Format.nBlockAlign = sizeof(float) * Format.nChannels;
+			Format.wBitsPerSample = sizeof(float) * 8;
 
-		double CurrentTime = FPlatformTime::Seconds();
+			// Create the output source voice
+			HRESULT Result = DeviceSwapResult->NewSystem->CreateSourceVoice(
+				&DeviceSwapResult->NewSourceVoice,
+				&Format,
+				XAUDIO2_VOICE_NOPITCH,
+				XAUDIO2_DEFAULT_FREQ_RATIO,
+				InDeviceContext->Callbacks,
+				nullptr,
+				nullptr
+			);
 
-		// If we're already in the process of swapping, we do not want to "double-trigger" a swap
-		if (bMoveAudioStreamToNewAudioDevice)
-		{
-			LastDeviceSwapTime = CurrentTime;
-			return false;
+			if (FAILED(Result))
+			{
+				if (DeviceSwapResult->NewSourceVoice)
+				{
+					DeviceSwapResult->NewSourceVoice->DestroyVoice();
+					DeviceSwapResult->NewSourceVoice = nullptr;
+				}
+
+				if (DeviceSwapResult->NewMasteringVoice)
+				{
+					DeviceSwapResult->NewMasteringVoice->DestroyVoice();
+					DeviceSwapResult->NewMasteringVoice = nullptr;
+				}
+				
+				SAFE_RELEASE(DeviceSwapResult->NewSystem);
+				XAUDIO2_LOG_AND_HANDLE_ON_FAIL("XAudio2System->CreateSourceVoice", Result, return {});
+			}
 		}
 
-		// Some devices spam device swap notifications, so we want to rate-limit them to prevent double/triple triggering.
-		static const int32 MinSwapTimeMs = 10;
-		if (CurrentTime - LastDeviceSwapTime > (double)MinSwapTimeMs / 1000.0)
+		// Optionally for testing, sleep for some duration in order to help repro race conditions.
+		if (ThreadedSwapDebugExtraTimeMsCVar > 0.f)
 		{
-			LastDeviceSwapTime = CurrentTime;
-			return true;
+			FPlatformProcess::Sleep(ThreadedSwapDebugExtraTimeMsCVar/1000.f); 
 		}
 
-		return false;
+		// Listen session for changes to this device.
+#if PLATFORM_WINDOWS
+		RegisterForSessionEvents(InDeviceContext->RequestedDeviceId);
+#endif
+		
+		DeviceSwapResult->SuccessfulDurationMs = FPlatformTime::ToMilliseconds64(FPlatformTime::Cycles64() - StartTimeCycles);
+		DeviceSwapResult->DeviceInfo = InDeviceContext->NewDevice.Get(FAudioPlatformDeviceInfo());
+		DeviceSwapResult->SwapReason = InDeviceContext->DeviceSwapReason;
+		
+		return DeviceSwapResult;
 	}
-
+	
 	bool FMixerPlatformXAudio2::ResetXAudio2System()
 	{
 		SAFE_RELEASE(XAudio2System);
@@ -562,6 +541,9 @@ namespace Audio
 			AUDIO_PLATFORM_LOG_ONCE(TEXT("XAudio2 was already tore down."), Warning);
 			return false;
 		}
+		
+		// Lock prior to changing state to avoid race condition if there happens to be an in-flight device swap
+		FScopeLock Lock(&DeviceSwapCriticalSection);
 
 		if (XAudio2System)
 		{
@@ -936,7 +918,7 @@ namespace Audio
 			// Store the device ID here in case it is removed. We can switch back if the device comes back.
 			if (Params.bRestoreIfRemoved)
 			{
-				OriginalAudioDeviceId = AudioStreamInfo.DeviceInfo.DeviceId;
+				SetOriginalAudioDeviceId(AudioStreamInfo.DeviceInfo.DeviceId);
 			}
 
 #if PLATFORM_WINDOWS			
@@ -1042,9 +1024,12 @@ namespace Audio
 			return false;
 		}
 
+		// Lock prior to changing state to avoid race condition if there happens to be an in-flight device swap
+		FScopeLock Lock(&DeviceSwapCriticalSection);
+
 		// If we're closing the stream, we're not interested in the results of the device swap. 
 		// Reset the handle to the future.
-		ActiveDeviceSwap.Reset();
+		ResetActiveDeviceSwapFuture();
 			
 		if (bIsDeviceOpen && !StopAudioStream())
 		{
@@ -1062,13 +1047,13 @@ namespace Audio
 			OutputAudioStreamSourceVoice = nullptr;
 		}
 
-		check(OutputAudioStreamMasteringVoice || bIsUsingNullDevice);
 		if (OutputAudioStreamMasteringVoice)
 		{
 			OutputAudioStreamMasteringVoice->DestroyVoice();
 			OutputAudioStreamMasteringVoice = nullptr;
 		}
-		else
+
+		if (bIsUsingNullDevice)
 		{
 			StopRunningNullDevice();
 		}
@@ -1083,13 +1068,10 @@ namespace Audio
 	bool FMixerPlatformXAudio2::StartAudioStream()
 	{
 		UE_LOG(LogAudioMixer, Log, TEXT("FMixerPlatformXAudio2::StartAudioStream() called. InstanceID=%d"), InstanceID);
-		// Start generating audio with our output source voice
-		BeginGeneratingAudio();
 
 		// If we already have a source voice, we can just restart it
 		if (OutputAudioStreamSourceVoice)
 		{
-			AudioStreamInfo.StreamState = EAudioOutputStreamState::Running;
 			OutputAudioStreamSourceVoice->Start();
 		}
 		else
@@ -1097,6 +1079,18 @@ namespace Audio
 			check(!bIsUsingNullDevice);
 			StartRunningNullDevice();
 		}
+		
+		// Start generating audio with our output source voice
+		// Can be called during device swap when AudioRenderEvent can be null
+		if (nullptr == AudioRenderEvent)
+		{
+			// This will set AudioStreamInfo.StreamState to EAudioOutputStreamState::Running
+			BeginGeneratingAudio();
+		}
+		else
+		{
+			AudioStreamInfo.StreamState = EAudioOutputStreamState::Running;
+		}
 
 		return true;
 	}
@@ -1111,6 +1105,9 @@ namespace Audio
 
 		UE_LOG(LogAudioMixer, Log, TEXT("FMixerPlatformXAudio2::StopAudioStream() called. InstanceID=%d"), InstanceID);
 
+		// Lock prior to changing state to avoid race condition if there happens to be an in-flight device swap
+		FScopeLock Lock(&DeviceSwapCriticalSection);
+
 		if (AudioStreamInfo.StreamState != EAudioOutputStreamState::Stopped && AudioStreamInfo.StreamState != EAudioOutputStreamState::Closed)
 		{
 			if (AudioStreamInfo.StreamState == EAudioOutputStreamState::Running)
@@ -1134,295 +1131,96 @@ namespace Audio
 	bool FMixerPlatformXAudio2::CheckAudioDeviceChange()
 	{
 #if XAUDIO_SUPPORTS_DEVICE_DETAILS
-
-		SCOPED_NAMED_EVENT(FMixerPlatformXAudio2_CheckAudioDeviceChange, FColor::Blue);
-
-		// Use threaded version? (It also requires the info cache).
-		if (IAudioMixer::ShouldUseThreadedDeviceSwap() && IAudioMixer::ShouldUseDeviceInfoCache())
-		{
-			return CheckThreadedDeviceSwap();
-		}
-
-		FScopeLock Lock(&AudioDeviceSwapCriticalSection);
-		if (bMoveAudioStreamToNewAudioDevice)
-		{
-			bMoveAudioStreamToNewAudioDevice = false;
-
-			return MoveAudioStreamToNewAudioDevice(NewAudioDeviceId);
-		}
-#endif
-		return false;
-	}
-
-	bool FMixerPlatformXAudio2::RequestDeviceSwap(const FString& DeviceID, bool bInForce, const TCHAR* InReason)
-	{
-		if (!AllowDeviceSwap() && !bInForce)
-		{
-			UE_LOG(LogAudioMixer, Verbose, TEXT("NOT-ALLOWING attempt to swap audio render device to new device: '%s', because: '%s', force=%d, InstanceID=%d"),
-				!DeviceID.IsEmpty() ? *DeviceID : TEXT("[System Default]"),
-				InReason ? InReason : TEXT("None specified"),
-				(int32)bInForce,
-				InstanceID
-			);
-			return false;
-		}
-
-		if (AudioDeviceSwapCriticalSection.TryLock())
-		{
-			UE_LOG(LogAudioMixer, Verbose, TEXT("Attempt to swap audio render device to new device: '%s', because: '%s', force=%d, InstanceID=%d"),
-				!DeviceID.IsEmpty() ? *DeviceID : TEXT("[System Default]"),
-				InReason ? InReason : TEXT("None specified"),
-				(int32)bInForce,
-				InstanceID
-			);
-
-			NewAudioDeviceId = DeviceID;
-			bMoveAudioStreamToNewAudioDevice = true;
-			DeviceSwapReason = InReason;
-			
-			AudioDeviceSwapCriticalSection.Unlock();
-
-			return true;
-		}
-
-		UE_CLOG(bInForce,LogAudioMixer, Warning, TEXT("FMixerPlatformXAudio2::RequestDeviceSwap(), Failed to acquire CS, device swap to '%s', will be ignored. InstanceID=%d"),
-			!DeviceID.IsEmpty() ? *DeviceID : TEXT("[System Default]"), InstanceID );
+		return FAudioMixerPlatformSwappable::CheckAudioDeviceChange();
+#else
 		return false;
+#endif
 	}
-		
-	bool FMixerPlatformXAudio2::MoveAudioStreamToNewAudioDevice(const FString& InNewDeviceId)
+	
+	bool FMixerPlatformXAudio2::MoveAudioStreamToNewAudioDevice()
 	{
 		bool bDidStopGeneratingAudio = false;
-
 #if XAUDIO_SUPPORTS_DEVICE_DETAILS
+		bDidStopGeneratingAudio = FAudioMixerPlatformSwappable::MoveAudioStreamToNewAudioDevice();
+#endif // XAUDIO_SUPPORTS_DEVICE_DETAILS
+		return bDidStopGeneratingAudio;
+	}
+	
+	void FMixerPlatformXAudio2::SubmitBuffer(const uint8* Buffer)
+	{
+		SCOPED_NAMED_EVENT(FMixerPlatformXAudio2_SubmitBuffer, FColor::Blue);
 
-		SCOPED_NAMED_EVENT(FMixerPlatformXAudio2_MoveAudioStreamToNewAudioDevice, FColor::Blue);
-
-		uint32 NumDevices = 0;
-		if (!GetNumOutputDevices(NumDevices))
-		{
-			return bDidStopGeneratingAudio;
-		}
-
-		// If we're running the null device, This function is called every second or so.
-		// Because of this, we early exit from this function if we're running the null device
-		// and there still are no devices.
-		const bool bTrySwitchToHardwareDevice = NumDevices > 0;
-		const bool bContinueUsingNullDevice = bIsUsingNullDevice && !bTrySwitchToHardwareDevice;
-		const bool bSwitchToNullDevice = !bIsUsingNullDevice && !bTrySwitchToHardwareDevice;
-
-		if (bContinueUsingNullDevice)
-		{
-			// Audio device was not changed. Return false to avoid downstream device change logic.
-			return bDidStopGeneratingAudio;
-		}
-
-		UE_LOG(LogAudioMixer, Log, TEXT("Resetting audio stream to device id '%s'"), 
-			!InNewDeviceId.IsEmpty() ? *InNewDeviceId : TEXT("[System Default]") );
-
-		// Device swaps require reinitialization of output buffers to handle
-		// different channel formats. Stop generating audio to protect against
-		// accessing the OutputBuffer.
-		StopGeneratingAudio();
-		bDidStopGeneratingAudio = true;
-
-		// Stop currently running device
-		if (bIsUsingNullDevice)
-		{
-			StopRunningNullDevice();
-		}
-		else
+		if (OutputAudioStreamSourceVoice)
 		{
-			// Not initialized!
-			if (!bIsInitialized)
-			{
-				return bDidStopGeneratingAudio;
-			}
-
-			SCOPED_NAMED_EVENT(FMixerPlatformXAudio2_MoveAudioStreamToNewAudioDevice_DestroyVoices, FColor::Blue);
-
-			// If an XAudio2 callback is in flight,
-			// we have to wait for it here.
-			FScopeLock ScopeLock(&DeviceSwapCriticalSection);
-
-			// Now that we've properly locked, raise the bIsInDeviceSwap flag
-			// in case FlushSourceBuffers() calls OnBufferEnd on this thread,
-			// and DeviceSwapCriticalSection.TryLock() is still returning true
-			bIsInDeviceSwap = true;
-
-			// Flush all buffers. Because we've locked DeviceSwapCriticalSection, ReadNextBuffer will early exit and we will not submit any additional buffers.
-			if (OutputAudioStreamSourceVoice)
-			{
-				OutputAudioStreamSourceVoice->FlushSourceBuffers();
-			}
+			// Create a new xaudio2 buffer submission
+			XAUDIO2_BUFFER XAudio2Buffer = { 0 };
+			XAudio2Buffer.AudioBytes = OpenStreamParams.NumFrames * AudioStreamInfo.DeviceInfo.NumChannels * sizeof(float);
+			XAudio2Buffer.pAudioData = (const BYTE*)Buffer;
+			XAudio2Buffer.pContext = this;
 
-			if (OutputAudioStreamSourceVoice)
-			{
-				// Then destroy the current audio stream source voice
-				OutputAudioStreamSourceVoice->DestroyVoice();
-				OutputAudioStreamSourceVoice = nullptr;
-			}
+			// Submit buffer to the output streaming voice
+			OutputAudioStreamSourceVoice->SubmitSourceBuffer(&XAudio2Buffer);
 
-			// Now destroy the mastering voice
-			if (OutputAudioStreamMasteringVoice)
+			if(!FirstBufferSubmitted)
 			{
-				OutputAudioStreamMasteringVoice->DestroyVoice();
-				OutputAudioStreamMasteringVoice = nullptr;
+				UE_LOG(LogAudioMixer, Display, TEXT("FMixerPlatformXAudio2::SubmitBuffer() called for the first time. InstanceID=%d"), InstanceID);
+				FirstBufferSubmitted = true;
 			}
-
-			bIsInDeviceSwap = false;
 		}
+	}
 
-		// In order to resume audio playback, this function must return true. 
-		// All code paths below return true, even it they encounter an error.
+	bool FMixerPlatformXAudio2::InitializeDeviceSwapContext(const FString& InRequestedDeviceID, const TCHAR* InReason)
+	{
+		check(GetDeviceInfoCache());
+
+		// Look up device. Blank name looks up current default.
+		const FName NewDeviceName = *InRequestedDeviceID;
+		TOptional<FAudioPlatformDeviceInfo> DeviceInfo;
 		
-		if (bTrySwitchToHardwareDevice)
+		if (TOptional<FAudioPlatformDeviceInfo> TempDeviceInfo = GetDeviceInfoCache()->FindActiveOutputDevice(NewDeviceName))
 		{
-			if (!ResetXAudio2System())
-			{
-				// Reinitializing the XAudio2System failed, so we have to exit here.
-				return bDidStopGeneratingAudio;
-			}
-
-			// Now get info on the new audio device we're trying to reset to
-			uint32 DeviceIndex = AUDIO_MIXER_DEFAULT_DEVICE_INDEX;
-			if (!InNewDeviceId.IsEmpty())
+			if (TempDeviceInfo.IsSet())
 			{
-				// XAudio2 for HoloLens doesn't have GetDeviceDetails, use local wrapper instead
-				FAudioPlatformDeviceInfo DeviceDetails;
-				for (uint32 i = 0; i < NumDevices; ++i)
+				if (IsDeviceInfoValid(*TempDeviceInfo))
 				{
-					GetOutputDeviceInfo(i, DeviceDetails);
-					if (DeviceDetails.DeviceId == InNewDeviceId)
-					{
-						DeviceIndex = i;
-						break;
-					}
+					DeviceInfo = MoveTemp(TempDeviceInfo);
+				}
+				else
+				{
+					UE_LOG(LogAudioMixer, Warning, TEXT("Ignoring attempt to switch to device with unsupported params: Channels=%u, SampleRate=%u, Id=%s, Name=%s"),
+						(uint32)TempDeviceInfo->NumChannels, (uint32)TempDeviceInfo->SampleRate, *TempDeviceInfo->DeviceId, *TempDeviceInfo->Name);
+					return false;
 				}
 			}
-
-			// Update the audio stream info to the new device info
-			AudioStreamInfo.OutputDeviceIndex = DeviceIndex;
-
-			// Get the output device info at this new index
-			GetOutputDeviceInfo(AudioStreamInfo.OutputDeviceIndex, AudioStreamInfo.DeviceInfo);
-
-			// Create a new master voice
-			{
-				SCOPED_NAMED_EVENT(FMixerPlatformXAudio2_MoveAudioStreamToNewAudioDevice_CreateMaster, FColor::Blue);
-
-				// open up on the default device
-				OutputAudioStreamMasteringVoice = nullptr;
-				XAUDIO2_CALL_AND_HANDLE_ERROR(
-					XAudio2System->CreateMasteringVoice(
-					&OutputAudioStreamMasteringVoice,
-					AudioStreamInfo.DeviceInfo.NumChannels,
-					AudioStreamInfo.DeviceInfo.SampleRate,
-					0,
-					*AudioStreamInfo.DeviceInfo.DeviceId,
-					nullptr,
-					AudioCategory_GameEffects),
-					/* nop */);
-							
-				UE_CLOG(OutputAudioStreamMasteringVoice, LogAudioMixer, Display, TEXT("XAudio2 CreateMasterVoice Channels=%u SampleRate=%u DeviceId=%s Name=%s"),
-					AudioStreamInfo.DeviceInfo.NumChannels, AudioStreamInfo.DeviceInfo.SampleRate, *AudioStreamInfo.DeviceInfo.DeviceId, *AudioStreamInfo.DeviceInfo.Name
-				);
-			}
-
-			// Setup the format of the output source voice
-			WAVEFORMATEX Format = { 0 };
-			Format.nChannels = AudioStreamInfo.DeviceInfo.NumChannels;
-			Format.nSamplesPerSec = OpenStreamParams.SampleRate;
-			Format.wFormatTag = WAVE_FORMAT_IEEE_FLOAT;
-			Format.nAvgBytesPerSec = Format.nSamplesPerSec * sizeof(float) * Format.nChannels;
-			Format.nBlockAlign = sizeof(float) * Format.nChannels;
-			Format.wBitsPerSample = sizeof(float) * 8;
-
-			{				
-				SCOPED_NAMED_EVENT(FMixerPlatformXAudio2_MoveAudioStreamToNewAudioDevice_CreateSource, FColor::Blue);
-
-				// Create the output source voice
-				OutputAudioStreamSourceVoice = nullptr;
-					XAUDIO2_CALL_AND_HANDLE_ERROR(
-						XAudio2System->CreateSourceVoice(&OutputAudioStreamSourceVoice, &Format, XAUDIO2_VOICE_NOPITCH, 2.0f, &OutputVoiceCallback),
-						/* nop */);
-			}
-
-			// Reinitialize the output circular buffer to match the buffer math of the new audio device.
-			const int32 NumOutputSamples = AudioStreamInfo.NumOutputFrames * AudioStreamInfo.DeviceInfo.NumChannels;
-			if (ensure(NumOutputSamples > 0))
-			{
-				SCOPED_NAMED_EVENT(FMixerPlatformXAudio2_MoveAudioStreamToNewAudioDevice_OutBufferInit, FColor::Blue);
-				OutputBuffer.Init(AudioStreamInfo.AudioMixer, NumOutputSamples, NumOutputBuffers, AudioStreamInfo.DeviceInfo.Format);
-			}
-		}
-		else
-		{	
-			check(bSwitchToNullDevice);
-
-			// If we don't have any hardware playback devices available, use the null device callback to render buffers.
-			// NullDevice is started when OutputAudioStreamSourceVoice is null
-			check(nullptr == OutputAudioStreamSourceVoice);
-
-			return bDidStopGeneratingAudio;
 		}
 		
-#endif // #if PLATFORM_WINDOWS
-
-		return bDidStopGeneratingAudio; 
+		return InitDeviceSwapContextInternal(InRequestedDeviceID, InReason, DeviceInfo);
 	}
 
-	void FMixerPlatformXAudio2::ResumePlaybackOnNewDevice()
+	bool FMixerPlatformXAudio2::InitDeviceSwapContextInternal(const FString& InRequestedDeviceID, const TCHAR* InReason, const TOptional<FAudioPlatformDeviceInfo>& InDeviceInfo)
 	{
-		SCOPED_NAMED_EVENT(FMixerPlatformXAudio2_ResumePlaybackOnNewDevice, FColor::Blue);
-		
-		int32 NumSamplesPopped = 0;
-		TArrayView<const uint8> PoppedAudio = OutputBuffer.PopBufferData(NumSamplesPopped);
-		SubmitBuffer(PoppedAudio.GetData());
-
-		check(OpenStreamParams.NumFrames * AudioStreamInfo.DeviceInfo.NumChannels == OutputBuffer.GetNumSamples());
-
-		if (nullptr == AudioRenderEvent)
-		{
-			BeginGeneratingAudio();
-		}
-
-		AudioRenderEvent->Trigger();
+		// Access to device swap context must be protected by DeviceSwapCriticalSection
+		FScopeLock Lock(&DeviceSwapCriticalSection);
 
-		if (OutputAudioStreamSourceVoice)
+		if (DeviceSwapContext.IsValid())
 		{
-			// Start the voice streaming
-			OutputAudioStreamSourceVoice->Start();
+			UE_LOG(LogAudioMixer, Display, TEXT("FMixerPlatformXAudio2::InitDeviceSwapContextInternal DeviceSwapContext in-flight, ignoring"));
+			return false;
 		}
-		else
+		
+		// Create the device swap context which will be valid over the course of the swap
+		DeviceSwapContext = MakeUnique<FXAudio2DeviceSwapContext>(InRequestedDeviceID, InReason);
+		if (!DeviceSwapContext.IsValid())
 		{
-			StartRunningNullDevice();
+			UE_LOG(LogAudioMixer, Warning, TEXT("FMixerPlatformXAudio2::InitDeviceSwapContextInternal failed to create DeviceSwapContext"));
+			return false;
 		}
-	}
 
-	void FMixerPlatformXAudio2::SubmitBuffer(const uint8* Buffer)
-	{
-		SCOPED_NAMED_EVENT(FMixerPlatformXAudio2_SubmitBuffer, FColor::Blue);
-
-		if (OutputAudioStreamSourceVoice)
-		{
-			// Create a new xaudio2 buffer submission
-			XAUDIO2_BUFFER XAudio2Buffer = { 0 };
-			XAudio2Buffer.AudioBytes = OpenStreamParams.NumFrames * AudioStreamInfo.DeviceInfo.NumChannels * sizeof(float);
-			XAudio2Buffer.pAudioData = (const BYTE*)Buffer;
-			XAudio2Buffer.pContext = this;
+		DeviceSwapContext->NewDevice = InDeviceInfo;
 
-			// Submit buffer to the output streaming voice
-			OutputAudioStreamSourceVoice->SubmitSourceBuffer(&XAudio2Buffer);
-
-			if(!FirstBufferSubmitted)
-			{
-				UE_LOG(LogAudioMixer, Display, TEXT("FMixerPlatformXAudio2::SubmitBuffer() called for the first time. InstanceID=%d"), InstanceID);
-				FirstBufferSubmitted = true;
-			}
-		}
+		return true;
 	}
-
+	
 	FString FMixerPlatformXAudio2::GetDefaultDeviceName()
 	{
 		return FString();
@@ -1443,7 +1241,7 @@ namespace Audio
 
 	Audio::IAudioPlatformDeviceInfoCache* FMixerPlatformXAudio2::GetDeviceInfoCache() const
 	{
-		if (IAudioMixer::ShouldUseDeviceInfoCache())
+		if (ShouldUseDeviceInfoCache())
 		{
 			return DeviceInfoCache.Get();
 		}
@@ -1451,6 +1249,18 @@ namespace Audio
 		return nullptr;
 	}
 
+	bool FMixerPlatformXAudio2::IsDeviceInfoValid(const FAudioPlatformDeviceInfo& InDeviceInfo) const
+	{
+		if (InDeviceInfo.NumChannels <= XAUDIO2_MAX_AUDIO_CHANNELS &&
+			InDeviceInfo.SampleRate >= XAUDIO2_MIN_SAMPLE_RATE &&
+			InDeviceInfo.SampleRate <= XAUDIO2_MAX_SAMPLE_RATE)
+		{
+			return true;
+		}
+
+		return false;
+	}
+
 	void FMixerPlatformXAudio2::OnSessionDisconnect(IAudioMixerDeviceChangedListener::EDisconnectReason InReason)
 	{
 		// Device has disconnected from current session.
@@ -1487,7 +1297,8 @@ namespace Audio
 
 	void FMixerPlatformXAudio2::OnCriticalError(HRESULT Error)
 	{
-		// Windows handles this via session events, so no need to handle that here...
+		// Windows should handle this via session events, log if we receive one.
+		UE_LOG(LogAudioMixer, Warning, TEXT("FMixerPlatformXAudio2::OnCriticalError: 0x%X: %s"), Error, *Audio::ToErrorFString(Error));
 	}
 	
 }
diff --git a/Engine/Source/Runtime/Windows/AudioMixerXAudio2/Private/AudioMixerPlatformXAudio2.h b/Engine/Source/Runtime/Windows/AudioMixerXAudio2/Private/AudioMixerPlatformXAudio2.h
--- a/Engine/Source/Runtime/Windows/AudioMixerXAudio2/Private/AudioMixerPlatformXAudio2.h
+++ b/Engine/Source/Runtime/Windows/AudioMixerXAudio2/Private/AudioMixerPlatformXAudio2.h
@@ -67,14 +67,40 @@ namespace Audio
 
 	};
 
-	class FMixerPlatformXAudio2 : public IAudioMixerPlatformInterface,
+	struct FXAudio2DeviceSwapContext : public FDeviceSwapContext
+	{
+		FXAudio2DeviceSwapContext() = delete;
+		FXAudio2DeviceSwapContext(const FString& InRequestedDeviceID, const FString& InReason) :
+			FDeviceSwapContext(InRequestedDeviceID, InReason)
+		{}
+		
+		IXAudio2* PreviousSystem = nullptr;
+		IXAudio2MasteringVoice* PreviousMasteringVoice = nullptr;
+		IXAudio2SourceVoice* PreviousSourceVoice = nullptr;
+		FXAudio2VoiceCallback* Callbacks = nullptr;
+		uint32 RenderingSampleRate = 0;
+	};
+
+	struct FXAudio2DeviceSwapResult : public FDeviceSwapResult
+	{
+		virtual bool IsNewDeviceReady() const override
+		{
+			return NewSystem && NewMasteringVoice && NewSourceVoice;
+		}
+
+		IXAudio2* NewSystem = nullptr;
+		IXAudio2MasteringVoice* NewMasteringVoice = nullptr;
+		IXAudio2SourceVoice* NewSourceVoice = nullptr;
+	};
+
+	class FMixerPlatformXAudio2 : public FAudioMixerPlatformSwappable,
 								  public IXAudio2EngineCallback
 	{
 
 	public:
 
 		FMixerPlatformXAudio2();
-		~FMixerPlatformXAudio2();
+		virtual ~FMixerPlatformXAudio2() override;
 
 		//~ Begin IAudioMixerPlatformInterface
 		virtual FString GetPlatformApi() const override { return TEXT("XAudio2"); }
@@ -90,9 +116,8 @@ namespace Audio
 		virtual bool CloseAudioStream() override;
 		virtual bool StartAudioStream() override;
 		virtual bool StopAudioStream() override;
-		virtual bool MoveAudioStreamToNewAudioDevice(const FString& InNewDeviceId) override;
-		virtual bool RequestDeviceSwap(const FString& DeviceID, bool bInForce, const TCHAR* InReason) override;
-		virtual void ResumePlaybackOnNewDevice() override;
+		virtual bool MoveAudioStreamToNewAudioDevice() override;
+		
 		virtual FAudioPlatformDeviceInfo GetPlatformDeviceInfo() const override;
 		virtual void SubmitBuffer(const uint8* Buffer) override;
 		virtual bool DisablePCMAudioCaching() const override;		
@@ -100,6 +125,8 @@ namespace Audio
 		virtual FAudioPlatformSettings GetPlatformSettings() const override;
 		virtual void OnHardwareUpdate() override;
 		virtual IAudioPlatformDeviceInfoCache* GetDeviceInfoCache() const override;
+		virtual bool IsDeviceInfoValid(const FAudioPlatformDeviceInfo& InDeviceInfo) const override;
+		virtual bool ShouldUseDeviceInfoCache() const override { return true; }
 		//~ End IAudioMixerPlatformInterface
 
 		//~ Begin IAudioMixerDeviceChangedListener
@@ -114,6 +141,15 @@ namespace Audio
 		virtual FString GetDeviceId() const override;		
 		//~ End IAudioMixerDeviceChangedListener
 
+		//~ Begin FAudioMixerPlatformSwappable
+		virtual bool InitializeDeviceSwapContext(const FString& InRequestedDeviceID, const TCHAR* InReason) override;
+		virtual bool CheckThreadedDeviceSwap() override;
+		virtual bool PreDeviceSwap() override;
+		virtual void EnqueueAsyncDeviceSwap() override;
+		virtual void SynchronousDeviceSwap() override;
+		virtual bool PostDeviceSwap() override;
+		//~ End FAudioMixerPlatformSwappable
+	
 	protected:
 		virtual uint32 GetCreateFlags() const { return 0; }
 
@@ -122,44 +158,29 @@ namespace Audio
 		virtual void OnProcessingPassStart() override;
 		virtual void OnProcessingPassEnd() override;
 		//~ End IXAudio2EngineCallback
-			
-		struct FXAudio2AsyncCreateResult
-		{
-			IXAudio2* XAudio2System = nullptr;
-			IXAudio2MasteringVoice* OutputAudioStreamMasteringVoice = nullptr;
-			IXAudio2SourceVoice* OutputAudioStreamSourceVoice = nullptr;
-			FAudioPlatformDeviceInfo DeviceInfo;
-			FString SwapReason;
-			float SuccessfullDurationMs = 0.f;
-		};
-		TFuture<FXAudio2AsyncCreateResult> ActiveDeviceSwap;
-		bool CheckThreadedDeviceSwap();
-	
-		bool AllowDeviceSwap();
-
+		
 		// Used to teardown and reinitialize XAudio2.
 		// This must be done to repopulate the playback device list in XAudio 2.7.
 		bool ResetXAudio2System();
 
-
+		/** Can be used by subclasses to initialize a device swap context by supplying a specific 
+		 *  FAudioPlatformDeviceInfo rather than looking it up via the requested device Id.
+		 */
+		bool InitDeviceSwapContextInternal(const FString& InRequestedDeviceID, const TCHAR* InReason, const TOptional<FAudioPlatformDeviceInfo>& InDeviceInfo);
+		
 		// Handle to XAudio2DLL
 		HMODULE XAudio2Dll;
-
-
+		
 		// Bool indicating that the default audio device changed
 		// And that we need to restart the audio device.
+		UE_DEPRECATED(5.6, "bDeviceChanged has been deprecated.")
 		FThreadSafeBool bDeviceChanged;
 
 		IXAudio2* XAudio2System;
 		IXAudio2MasteringVoice* OutputAudioStreamMasteringVoice;
 		IXAudio2SourceVoice* OutputAudioStreamSourceVoice;
 		FXAudio2VoiceCallback OutputVoiceCallback;
-		FCriticalSection AudioDeviceSwapCriticalSection;
-		FString OriginalAudioDeviceId;
-		FString NewAudioDeviceId;
-		double LastDeviceSwapTime;
-		FString DeviceSwapReason;
-
+		
 		// When we are running the null device,
 		// we check whether a new audio device was connected every second or so.
 		float TimeSinceNullDeviceWasLastChecked;
@@ -170,6 +191,14 @@ namespace Audio
 
 		uint32 bIsInitialized : 1;
 		uint32 bIsDeviceOpen : 1;
+
+	private:
+		/** Context object holding state used during device swap. */
+		TUniquePtr<FXAudio2DeviceSwapContext> DeviceSwapContext;
+
+		/** Performs a device swap with the given context. Static method enforces no other state sharing occurs. */
+		static TUniquePtr<FDeviceSwapResult> PerformDeviceSwap(TUniquePtr<FXAudio2DeviceSwapContext>&& InDeviceContext);
 	};
+
 }
 
diff --git a/Engine/Source/Runtime/Windows/AudioMixerXAudio2/Private/ToStringHelpers.cpp b/Engine/Source/Runtime/Windows/AudioMixerXAudio2/Private/ToStringHelpers.cpp
--- a/Engine/Source/Runtime/Windows/AudioMixerXAudio2/Private/ToStringHelpers.cpp
+++ b/Engine/Source/Runtime/Windows/AudioMixerXAudio2/Private/ToStringHelpers.cpp
@@ -1,17 +1,7 @@
 // Copyright Epic Games, Inc. All Rights Reserved.
 
 #include "ToStringHelpers.h"
-
-#include "Windows/AllowWindowsPlatformTypes.h"
-
-THIRD_PARTY_INCLUDES_START
-#include <winerror.h>
-#if PLATFORM_WINDOWS
-#include <audiopolicy.h>			// IAudioSessionEvents
-#endif
-THIRD_PARTY_INCLUDES_END
-
-#include "Windows/HideWindowsPlatformTypes.h"
+#include "WindowsMMStringUtils.h"
 
 namespace Audio
 {
@@ -25,70 +15,11 @@ namespace Audio
 		case HRESULT(XAUDIO2_E_XMA_DECODER_ERROR):		return TEXT("XAUDIO2_E_XMA_DECODER_ERROR");
 		case HRESULT(XAUDIO2_E_XAPO_CREATION_FAILED):	return TEXT("XAUDIO2_E_XAPO_CREATION_FAILED");
 		case HRESULT(XAUDIO2_E_DEVICE_INVALIDATED):		return TEXT("XAUDIO2_E_DEVICE_INVALIDATED");
-#if PLATFORM_WINDOWS
-		case HRESULT(REGDB_E_CLASSNOTREG):				return TEXT("REGDB_E_CLASSNOTREG");
-		case HRESULT(CLASS_E_NOAGGREGATION):			return TEXT("CLASS_E_NOAGGREGATION");
-		case HRESULT(E_NOINTERFACE):					return TEXT("E_NOINTERFACE");
-		case HRESULT(E_POINTER):						return TEXT("E_POINTER");
-		case HRESULT(E_INVALIDARG):						return TEXT("E_INVALIDARG");
-		case HRESULT(E_OUTOFMEMORY):					return TEXT("E_OUTOFMEMORY");
-
-		// AudioClient.h
-		CASE_AND_STRING(AUDCLNT_E_NOT_INITIALIZED);
-		CASE_AND_STRING(AUDCLNT_E_ALREADY_INITIALIZED);
-		CASE_AND_STRING(AUDCLNT_E_WRONG_ENDPOINT_TYPE);
-		CASE_AND_STRING(AUDCLNT_E_DEVICE_INVALIDATED);
-		CASE_AND_STRING(AUDCLNT_E_NOT_STOPPED);
-		CASE_AND_STRING(AUDCLNT_E_BUFFER_TOO_LARGE);
-		CASE_AND_STRING(AUDCLNT_E_OUT_OF_ORDER);
-		CASE_AND_STRING(AUDCLNT_E_UNSUPPORTED_FORMAT);
-		CASE_AND_STRING(AUDCLNT_E_INVALID_SIZE);
-		CASE_AND_STRING(AUDCLNT_E_DEVICE_IN_USE);
-		CASE_AND_STRING(AUDCLNT_E_BUFFER_OPERATION_PENDING);
-		CASE_AND_STRING(AUDCLNT_E_THREAD_NOT_REGISTERED);
-		CASE_AND_STRING(AUDCLNT_E_EXCLUSIVE_MODE_NOT_ALLOWED);
-		CASE_AND_STRING(AUDCLNT_E_ENDPOINT_CREATE_FAILED);
-		CASE_AND_STRING(AUDCLNT_E_SERVICE_NOT_RUNNING);
-		CASE_AND_STRING(AUDCLNT_E_EVENTHANDLE_NOT_EXPECTED);
-		CASE_AND_STRING(AUDCLNT_E_EXCLUSIVE_MODE_ONLY);
-		CASE_AND_STRING(AUDCLNT_E_BUFDURATION_PERIOD_NOT_EQUAL);
-		CASE_AND_STRING(AUDCLNT_E_EVENTHANDLE_NOT_SET);
-		CASE_AND_STRING(AUDCLNT_E_INCORRECT_BUFFER_SIZE);
-		CASE_AND_STRING(AUDCLNT_E_BUFFER_SIZE_ERROR);
-		CASE_AND_STRING(AUDCLNT_E_CPUUSAGE_EXCEEDED);
-		CASE_AND_STRING(AUDCLNT_E_BUFFER_ERROR);
-		CASE_AND_STRING(AUDCLNT_E_BUFFER_SIZE_NOT_ALIGNED);
-		CASE_AND_STRING(AUDCLNT_E_INVALID_DEVICE_PERIOD);
-		CASE_AND_STRING(AUDCLNT_E_INVALID_STREAM_FLAG);
-		CASE_AND_STRING(AUDCLNT_E_ENDPOINT_OFFLOAD_NOT_CAPABLE);
-		CASE_AND_STRING(AUDCLNT_E_OUT_OF_OFFLOAD_RESOURCES);
-		CASE_AND_STRING(AUDCLNT_E_OFFLOAD_MODE_ONLY);
-		CASE_AND_STRING(AUDCLNT_E_NONOFFLOAD_MODE_ONLY);
-		CASE_AND_STRING(AUDCLNT_E_RESOURCES_INVALIDATED);
-		CASE_AND_STRING(AUDCLNT_E_RAW_MODE_UNSUPPORTED);
-		CASE_AND_STRING(AUDCLNT_E_ENGINE_PERIODICITY_LOCKED);
-		CASE_AND_STRING(AUDCLNT_E_ENGINE_FORMAT_LOCKED);
-		CASE_AND_STRING(AUDCLNT_S_BUFFER_EMPTY);
-		CASE_AND_STRING(AUDCLNT_S_THREAD_ALREADY_REGISTERED);
-		CASE_AND_STRING(AUDCLNT_S_POSITION_STALLED);
-
-#undef CASE_AND_STRING
-
-#endif //PLATFORM_WINDOWS
-
-		case HRESULT(0xe000020b):						return TEXT("ERROR_NO_SUCH_DEVINST");
-		case HRESULT(0x80010108):						return TEXT("RPC_E_DISCONNECTED");
 
 		default:
 		{
-			// We don't know this error, ask this system if it does.
-			TCHAR Buffer[1024] = { 0 };
-			FString Msg(FPlatformMisc::GetSystemErrorMessage(Buffer, UE_ARRAY_COUNT(Buffer), InResult));
-
-			// Anything to return? Otherwise "UNKNOWN"
-			return !Msg.IsEmpty() ?
-				Msg :
-				TEXT("UNKNOWN");
+			// Not an XAudio2 error, check for Audio Client errors
+			return AudioClientErrorToFString(InResult);
 		}
 		}
 	}
