Speculative fix for crash and adjusting log levels related to device swap

#rnx

diff --git a/Engine/Source/Runtime/Windows/AudioMixerWasapi/Private/AudioMixerWasapi.cpp b/Engine/Source/Runtime/Windows/AudioMixerWasapi/Private/AudioMixerWasapi.cpp
--- a/Engine/Source/Runtime/Windows/AudioMixerWasapi/Private/AudioMixerWasapi.cpp
+++ b/Engine/Source/Runtime/Windows/AudioMixerWasapi/Private/AudioMixerWasapi.cpp
@@ -326,9 +326,14 @@ namespace Audio
 		// Lock prior to changing state to avoid race condition if there happens to be an in-flight device swap
 		FScopeLock Lock(&DeviceSwapCriticalSection);
 
+		UE_LOG(LogAudioMixer, Display, TEXT("FAudioMixerWasapi::StopAudioStream() InstanceID=%d, StreamState=%d"),
+			InstanceID, static_cast<int32>(AudioStreamInfo.StreamState));
+
 		if (AudioStreamInfo.StreamState != EAudioOutputStreamState::Stopped && AudioStreamInfo.StreamState != EAudioOutputStreamState::Closed)
 		{
-			if (AudioStreamInfo.StreamState == EAudioOutputStreamState::Running)
+			// Shutdown the AudioRenderThread if we're running or mid-device swap
+			if (AudioStreamInfo.StreamState == EAudioOutputStreamState::Running ||
+				AudioStreamInfo.StreamState == EAudioOutputStreamState::SwappingDevice)
 			{
 				StopGeneratingAudio();
 			}
@@ -516,7 +521,7 @@ namespace Audio
 	void FAudioMixerWasapi::EnqueueAsyncDeviceSwap()
 	{
 		FScopeLock Lock(&DeviceSwapCriticalSection);
-		UE_LOG(LogAudioMixer, Verbose, TEXT("FAudioMixerWasapi::EnqueueAsyncDeviceSwap - enqueuing async device swap"));
+		UE_LOG(LogAudioMixer, Display, TEXT("FAudioMixerWasapi::EnqueueAsyncDeviceSwap - enqueuing async device swap"));
 		
 		TFunction<TUniquePtr<FDeviceSwapResult>()> AsyncDeviceSwap = [this]() mutable -> TUniquePtr<FDeviceSwapResult>
 		{
@@ -564,7 +569,7 @@ namespace Audio
 		// No need to lock critical section here since this call has sole ownership of the context
 		if (InDeviceContext.IsValid())
 		{
-			UE_LOG(LogAudioMixer, Verbose, TEXT("FAudioMixerWasapi::PerformDeviceSwap - AsyncTask Start. Because=%s"), *InDeviceContext->DeviceSwapReason);
+			UE_LOG(LogAudioMixer, Display, TEXT("FAudioMixerWasapi::PerformDeviceSwap - AsyncTask Start. Because=%s"), *InDeviceContext->DeviceSwapReason);
 
 			if (InDeviceContext->OldDeviceManager.IsValid())
 			{
@@ -573,17 +578,17 @@ namespace Audio
 				InDeviceContext->OldDeviceManager->CloseAudioStream();
 				InDeviceContext->OldDeviceManager->TeardownHardware();
 				InDeviceContext->OldDeviceManager.Reset();
-				UE_LOG(LogAudioMixer, Verbose, TEXT("FAudioMixerWasapi::PerformDeviceSwap - successfully shut down previous device manager"));
+				UE_LOG(LogAudioMixer, Display, TEXT("FAudioMixerWasapi::PerformDeviceSwap - successfully shut down previous device manager"));
 			}
 			else
 			{
-				UE_LOG(LogAudioMixer, Verbose, TEXT("FAudioMixerWasapi::PerformDeviceSwap - no device manager running, null renderer must be active"));
+				UE_LOG(LogAudioMixer, Display, TEXT("FAudioMixerWasapi::PerformDeviceSwap - no device manager running, null renderer must be active"));
 			}
 
 			// Don't attempt to create a new setup if there's no devices available.
 			if (!InDeviceContext->NewDevice.IsSet())
 			{
-				UE_LOG(LogAudioMixer, Verbose, TEXT("FAudioMixerWasapi::PerformDeviceSwap - InDeviceContext->bSwitchingToValidDevice: false"));
+				UE_LOG(LogAudioMixer, Display, TEXT("FAudioMixerWasapi::PerformDeviceSwap - InDeviceContext->bSwitchingToValidDevice: false"));
 				return {};
 			}
 
@@ -611,7 +616,7 @@ namespace Audio
 			DeviceSwapResult->SuccessfulDurationMs = FPlatformTime::ToMilliseconds64(FPlatformTime::Cycles64() - StartTimeCycles);
 			DeviceSwapResult->DeviceInfo = InDeviceContext->NewDevice.Get(FAudioPlatformDeviceInfo());
 			DeviceSwapResult->SwapReason = InDeviceContext->DeviceSwapReason;
-			UE_LOG(LogAudioMixer, Verbose, TEXT("FAudioMixerWasapi::PerformDeviceSwap - successfully completed device swap"));
+			UE_LOG(LogAudioMixer, Display, TEXT("FAudioMixerWasapi::PerformDeviceSwap - successfully completed device swap"));
 
 			return DeviceSwapResult;
 		}
@@ -635,7 +640,7 @@ namespace Audio
 				// Finish initializing the device swap context
 				check(!DeviceSwapContext->OldDeviceManager);
 				DeviceSwapContext->OldDeviceManager = MoveTemp(DeviceManager);
-				UE_LOG(LogAudioMixer, Verbose, TEXT("FAudioMixerWasapi::PreDeviceSwap - Starting swap to [%s]"), DeviceSwapContext->RequestedDeviceId.IsEmpty() ? TEXT("[System Default]") : *DeviceSwapContext->RequestedDeviceId);
+				UE_LOG(LogAudioMixer, Display, TEXT("FAudioMixerWasapi::PreDeviceSwap - Starting swap to [%s]"), DeviceSwapContext->RequestedDeviceId.IsEmpty() ? TEXT("[System Default]") : *DeviceSwapContext->RequestedDeviceId);
 			}
 			else
 			{
diff --git a/Engine/Source/Runtime/Windows/AudioMixerXAudio2/Private/AudioMixerPlatformXAudio2.cpp b/Engine/Source/Runtime/Windows/AudioMixerXAudio2/Private/AudioMixerPlatformXAudio2.cpp
--- a/Engine/Source/Runtime/Windows/AudioMixerXAudio2/Private/AudioMixerPlatformXAudio2.cpp
+++ b/Engine/Source/Runtime/Windows/AudioMixerXAudio2/Private/AudioMixerPlatformXAudio2.cpp
@@ -183,7 +183,7 @@ namespace Audio
 			OutputAudioStreamMasteringVoice = nullptr;
 			OutputAudioStreamSourceVoice = nullptr;
 
-			UE_LOG(LogAudioMixer, Verbose, TEXT("FMixerPlatformXAudio2::PreDeviceSwap - Starting swap to [%s]"), DeviceSwapContext->RequestedDeviceId.IsEmpty() ? TEXT("[System Default]") : *DeviceSwapContext->RequestedDeviceId);
+			UE_LOG(LogAudioMixer, Display, TEXT("FMixerPlatformXAudio2::PreDeviceSwap - Starting swap to [%s]"), DeviceSwapContext->RequestedDeviceId.IsEmpty() ? TEXT("[System Default]") : *DeviceSwapContext->RequestedDeviceId);
 
 			return true;
 		}
@@ -198,7 +198,7 @@ namespace Audio
 	
 	void FMixerPlatformXAudio2::EnqueueAsyncDeviceSwap()
 	{
-		UE_LOG(LogAudioMixer, Verbose, TEXT("FMixerPlatformXAudio2::EnqueueAsyncDeviceSwap - enqueuing async device swap"));
+		UE_LOG(LogAudioMixer, Display, TEXT("FMixerPlatformXAudio2::EnqueueAsyncDeviceSwap - enqueuing async device swap"));
 		FScopeLock Lock(&DeviceSwapCriticalSection);
 		
 		TFunction<TUniquePtr<FDeviceSwapResult>()> AsyncDeviceSwap = [this]() mutable -> TUniquePtr<FDeviceSwapResult>
@@ -317,7 +317,7 @@ namespace Audio
 			return {};
 		}
 
-		UE_LOG(LogAudioMixer, Verbose, TEXT("FMixerPlatformXAudio2::PerformDeviceSwap - AsyncTask Start. Because=%s"), *InDeviceContext->DeviceSwapReason);
+		UE_LOG(LogAudioMixer, Display, TEXT("FMixerPlatformXAudio2::PerformDeviceSwap - AsyncTask Start. Because=%s"), *InDeviceContext->DeviceSwapReason);
 
 		// Stop old engine running.
 		if (InDeviceContext->PreviousSystem)
@@ -1103,14 +1103,17 @@ namespace Audio
 			return false;
 		}
 
-		UE_LOG(LogAudioMixer, Log, TEXT("FMixerPlatformXAudio2::StopAudioStream() called. InstanceID=%d"), InstanceID);
-
 		// Lock prior to changing state to avoid race condition if there happens to be an in-flight device swap
 		FScopeLock Lock(&DeviceSwapCriticalSection);
 
+		UE_LOG(LogAudioMixer, Display, TEXT("FMixerPlatformXAudio2::StopAudioStream() called. InstanceID=%d, StreamState=%d"),
+			InstanceID, static_cast<int32>(AudioStreamInfo.StreamState));
+
 		if (AudioStreamInfo.StreamState != EAudioOutputStreamState::Stopped && AudioStreamInfo.StreamState != EAudioOutputStreamState::Closed)
 		{
-			if (AudioStreamInfo.StreamState == EAudioOutputStreamState::Running)
+			// Shutdown the AudioRenderThread if we're running or mid-device swap
+			if (AudioStreamInfo.StreamState == EAudioOutputStreamState::Running ||
+				AudioStreamInfo.StreamState == EAudioOutputStreamState::SwappingDevice)
 			{
 				StopGeneratingAudio();
 			}
@@ -1119,7 +1122,6 @@ namespace Audio
 			if (OutputAudioStreamSourceVoice)
 			{
 				OutputAudioStreamSourceVoice->Stop(0, 0); // Don't wait for tails, stop as quick as you can.
-
 			}
 
 			check(AudioStreamInfo.StreamState == EAudioOutputStreamState::Stopped);
