Translate Windows Audio Client error codes to strings in logging

#rnx

diff --git a/Engine/Source/Runtime/Windows/AudioMixerWasapi/Private/AudioMixerWasapiRenderStream.cpp b/Engine/Source/Runtime/Windows/AudioMixerWasapi/Private/AudioMixerWasapiRenderStream.cpp
--- a/Engine/Source/Runtime/Windows/AudioMixerWasapi/Private/AudioMixerWasapiRenderStream.cpp
+++ b/Engine/Source/Runtime/Windows/AudioMixerWasapi/Private/AudioMixerWasapiRenderStream.cpp
@@ -4,6 +4,7 @@
 
 #include "AudioMixerWasapiLog.h"
 #include "WasapiAudioUtils.h"
+#include "WindowsMMStringUtils.h"
 
 namespace Audio
 {
@@ -36,7 +37,7 @@ namespace Audio
 		HRESULT Result = MMDevice->Activate(__uuidof(IAudioClient3), CLSCTX_INPROC_SERVER, nullptr, IID_PPV_ARGS_Helper(&TempAudioClient));
 		if (FAILED(Result))
 		{
-			UE_LOG(LogAudioMixerWasapi, Error, TEXT("OpenAudioStream failed IMMDevice::Activate 0x%x"), (uint32)Result);
+			UE_LOG(LogAudioMixerWasapi, Error, TEXT("OpenAudioStream failed IMMDevice::Activate %s"), *AudioClientErrorToFString(Result));
 			return false;
 		}
 
@@ -48,7 +49,7 @@ namespace Audio
 		Result = TempAudioClient->SetClientProperties(&AudioProps);
 		if (FAILED(Result))
 		{
-			UE_LOG(LogAudioMixerWasapi, Error, TEXT("OpenAudioStream failed IAudioClient3::SetClientProperties 0x%x"), (uint32)Result);
+			UE_LOG(LogAudioMixerWasapi, Error, TEXT("OpenAudioStream failed IAudioClient3::SetClientProperties %s"), *AudioClientErrorToFString(Result));
 			return false;
 		}
 
@@ -56,7 +57,7 @@ namespace Audio
 		Result = TempAudioClient->GetMixFormat(&MixFormat);
 		if (FAILED(Result) || !MixFormat)
 		{
-			UE_LOG(LogAudioMixerWasapi, Error, TEXT("OpenAudioStream failed IAudioClient3::GetMixFormat MixFormat: 0x%llx Result: 0x%x"), (uint64)MixFormat, (uint32)Result);
+			UE_LOG(LogAudioMixerWasapi, Error, TEXT("OpenAudioStream failed IAudioClient3::GetMixFormat MixFormat: 0x%llx Result: %s"), (uint64)MixFormat, *AudioClientErrorToFString(Result));
 			return false;
 		}
 
@@ -76,7 +77,7 @@ namespace Audio
 		Result = TempAudioClient->GetDevicePeriod(&DevicePeriodRefTime, nullptr);
 		if (FAILED(Result))
 		{
-			UE_LOG(LogAudioMixerWasapi, Error, TEXT("OpenAudioStream failed IAudioClient3::GetDevicePeriod 0x%x"), (uint32)Result);
+			UE_LOG(LogAudioMixerWasapi, Error, TEXT("OpenAudioStream failed IAudioClient3::GetDevicePeriod %s"), *AudioClientErrorToFString(Result));
 			return false;
 		}
 
@@ -119,14 +120,14 @@ namespace Audio
 		Result = TempAudioClient->Initialize(AUDCLNT_SHAREMODE_SHARED, Flags, DesiredBufferDuration, Periodicity, StreamFormat.GetWaveFormat(), nullptr);
 		if (FAILED(Result))
 		{
-			UE_LOG(LogAudioMixerWasapi, Error, TEXT("OpenAudioStream failed IAudioClient3::Initialize 0x%x"), (uint32)Result);
+			UE_LOG(LogAudioMixerWasapi, Error, TEXT("OpenAudioStream failed IAudioClient3::Initialize %s"), *AudioClientErrorToFString(Result));
 			return false;
 		}
 
 		Result = TempAudioClient->GetBufferSize(&NumFramesPerDeviceBuffer);
 		if (FAILED(Result))
 		{
-			UE_LOG(LogAudioMixerWasapi, Error, TEXT("OpenAudioStream failed IAudioClient3::GetBufferSize 0x%x"), (uint32)Result);
+			UE_LOG(LogAudioMixerWasapi, Error, TEXT("OpenAudioStream failed IAudioClient3::GetBufferSize %s"), *AudioClientErrorToFString(Result));
 			return false;
 		}
 
@@ -189,7 +190,7 @@ namespace Audio
 		HRESULT Result = AudioClient->SetEventHandle(InEventHandle);
 		if (FAILED(Result))
 		{
-			UE_LOG(LogAudioMixerWasapi, Error, TEXT("OpenAudioStream failed IAudioClient3::SetEventHandle 0x%x"), (uint32)Result);
+			UE_LOG(LogAudioMixerWasapi, Error, TEXT("OpenAudioStream failed IAudioClient3::SetEventHandle %s"), *AudioClientErrorToFString(Result));
 			return false;
 		}
 
@@ -197,7 +198,7 @@ namespace Audio
 		Result = AudioClient->GetService(__uuidof(IAudioRenderClient), IID_PPV_ARGS_Helper(&TempRenderClient));
 		if (FAILED(Result))
 		{
-			UE_LOG(LogAudioMixerWasapi, Error, TEXT("OpenAudioStream failed IAudioClient3::GetService IAudioRenderClient 0x%x"), (uint32)Result);
+			UE_LOG(LogAudioMixerWasapi, Error, TEXT("OpenAudioStream failed IAudioClient3::GetService IAudioRenderClient %s"), *AudioClientErrorToFString(Result));
 			return false;
 		}
 
