Reduce logging verbosity for audio device changes from Warning to Display

diff --git a/Engine/Source/Runtime/Windows/AudioMixerXAudio2/Private/AudioMixerPlatformWindows.cpp b/Engine/Source/Runtime/Windows/AudioMixerXAudio2/Private/AudioMixerPlatformWindows.cpp
--- a/Engine/Source/Runtime/Windows/AudioMixerXAudio2/Private/AudioMixerPlatformWindows.cpp
+++ b/Engine/Source/Runtime/Windows/AudioMixerXAudio2/Private/AudioMixerPlatformWindows.cpp
@@ -101,7 +101,7 @@ namespace Audio
 		
 		if (InAudioDeviceRole == EAudioDeviceRole::Console)
 		{
-			UE_LOG(LogAudioMixer, Warning, TEXT("FMixerPlatformXAudio2: Changing default audio render device to new device: Role=%s, DeviceName=%s, InstanceID=%d"), 
+			UE_LOG(LogAudioMixer, Display, TEXT("FMixerPlatformXAudio2: Changing default audio render device to new device: Role=%s, DeviceName=%s, InstanceID=%d"), 
 				Audio::ToString(InAudioDeviceRole), *WindowsNotificationClient->GetFriendlyName(DeviceId), InstanceID);
 
 			RequestDeviceSwap(DeviceId, /* force */true, TEXT("FMixerPlatformXAudio2::OnDefaultRenderDeviceChanged"));
@@ -126,7 +126,7 @@ namespace Audio
 		const FString AudioDeviceId = GetOriginalAudioDeviceId();
 		if (AudioStreamInfo.DeviceInfo.DeviceId != AudioDeviceId && DeviceId == AudioDeviceId)
 		{
-			UE_LOG(LogAudioMixer, Warning, TEXT("FMixerPlatformXAudio2: Original audio device re-added. Moving audio back to original audio device: DeviceName=%s, bRenderDevice=%d, InstanceID=%d"), 
+			UE_LOG(LogAudioMixer, Display, TEXT("FMixerPlatformXAudio2: Original audio device re-added. Moving audio back to original audio device: DeviceName=%s, bRenderDevice=%d, InstanceID=%d"), 
 				*WindowsNotificationClient->GetFriendlyName(*AudioDeviceId), (int32)bIsRenderDevice, InstanceID);
 
 			RequestDeviceSwap(AudioDeviceId, /*force */ true, TEXT("FMixerPlatformXAudio2::OnDeviceAdded"));
@@ -149,7 +149,7 @@ namespace Audio
 		// If the device we're currently using was removed... then switch to the new default audio device.
 		if (AudioStreamInfo.DeviceInfo.DeviceId == DeviceId)
 		{
-			UE_LOG(LogAudioMixer, Warning, TEXT("FMixerPlatformXAudio2: Audio device removed [%s], falling back to other windows default device. bIsRenderDevice=%d, InstanceID=%d"), 
+			UE_LOG(LogAudioMixer, Display, TEXT("FMixerPlatformXAudio2: Audio device removed [%s], falling back to other windows default device. bIsRenderDevice=%d, InstanceID=%d"), 
 				*WindowsNotificationClient->GetFriendlyName(DeviceId), (int32)bIsRenderDevice, InstanceID);
 
 			RequestDeviceSwap(TEXT(""), /* force */ true, TEXT("FMixerPlatformXAudio2::OnDeviceRemoved"));
