Removed #pragma comment(lib, "xaudio2_9redist.lib") since we use an external module to pull in the lib/dll

diff --git a/Engine/Plugins/Runtime/Windows/WinDualShock/Source/WinDualShock/Private/WinDualShock.cpp b/Engine/Plugins/Runtime/Windows/WinDualShock/Source/WinDualShock/Private/WinDualShock.cpp
--- a/Engine/Plugins/Runtime/Windows/WinDualShock/Source/WinDualShock/Private/WinDualShock.cpp
+++ b/Engine/Plugins/Runtime/Windows/WinDualShock/Source/WinDualShock/Private/WinDualShock.cpp
@@ -23,9 +23,6 @@ THIRD_PARTY_INCLUDES_START
 #include <winreg.h>
 #include <xaudio2redist.h>
 THIRD_PARTY_INCLUDES_END
-#if !PLATFORM_CPU_ARM_FAMILY
-#pragma comment(lib,"xaudio2_9redist.lib")
-#endif
 #endif
 
 #include <pad.h>
diff --git a/Engine/Source/Runtime/Windows/AudioMixerXAudio2/Private/AudioMixerPlatformXAudio2.h b/Engine/Source/Runtime/Windows/AudioMixerXAudio2/Private/AudioMixerPlatformXAudio2.h
--- a/Engine/Source/Runtime/Windows/AudioMixerXAudio2/Private/AudioMixerPlatformXAudio2.h
+++ b/Engine/Source/Runtime/Windows/AudioMixerXAudio2/Private/AudioMixerPlatformXAudio2.h
@@ -29,11 +29,6 @@
 
 #include "Async/Future.h"
 
-
-#if USE_REDIST_LIB
-#pragma comment(lib,"xaudio2_9redist.lib")
-#endif
-
 #ifndef XAUDIO_SUPPORTS_DEVICE_DETAILS
     #define XAUDIO_SUPPORTS_DEVICE_DETAILS		1
 #endif	//XAUDIO_SUPPORTS_DEVICE_DETAILS
