Resubmitting to main: added a virtual for whether XAudio2 should force usage of the default device ID, to allow platforms with only one audio device to use the virtual endpoint.
Also changed some platforms to avoid submitting buffers to XAudio2 while using the null device

diff --git a/Engine/Source/Runtime/Windows/AudioMixerXAudio2/Private/AudioMixerPlatformXAudio2.cpp b/Engine/Source/Runtime/Windows/AudioMixerXAudio2/Private/AudioMixerPlatformXAudio2.cpp
--- a/Engine/Source/Runtime/Windows/AudioMixerXAudio2/Private/AudioMixerPlatformXAudio2.cpp
+++ b/Engine/Source/Runtime/Windows/AudioMixerXAudio2/Private/AudioMixerPlatformXAudio2.cpp
@@ -34,7 +34,6 @@ THIRD_PARTY_INCLUDES_START
 #endif //PLATFORM_WINDOWS
 THIRD_PARTY_INCLUDES_END
 
-
 #include "Misc/CoreDelegates.h"
 #include "ProfilingDebugging/ScopedTimers.h"
 #include "Async/Async.h"
@@ -94,9 +93,15 @@ const FString& GetDllName(FName Current = NAME_None)
 */
 #define XAUDIO2_DEBUG_ENABLED 0
 
-
 namespace Audio
 {
+	static float GThreadedSwapDebugExtraTimeMs = 0;
+	static FAutoConsoleVariableRef GThreadedSwapDebugExtraTimeMsCVar(
+		TEXT("au.ThreadedSwapDebugExtraTime"),
+		GThreadedSwapDebugExtraTimeMs,
+		TEXT("Simulate a slow device swap by adding additional time to the swap task"),
+		ECVF_Default);
+
 	void FXAudio2VoiceCallback::OnBufferEnd(void* BufferContext)
 	{
 		SCOPED_NAMED_EVENT(FXAudio2VoiceCallback_OnBufferEnd, FColor::Blue);
@@ -300,13 +305,6 @@ namespace Audio
 
 		const uint64 StartTimeCycles = FPlatformTime::Cycles64();
 
-		static float ThreadedSwapDebugExtraTimeMsCVar = 0;
-		static FAutoConsoleVariableRef CVarOverrunTimeout(
-			TEXT("au.ThreadedSwapDebugExtraTime"),
-			ThreadedSwapDebugExtraTimeMsCVar,
-			TEXT("Simulate a slow device swap by adding additional time to the swap task"),
-			ECVF_Default);
-
 		// New thread might not have COM setup.
 		FScopedCoInitialize ScopedCoInitialize;
 
@@ -374,27 +372,9 @@ namespace Audio
 			check(InDeviceContext->NewDevice->SampleRate <= XAUDIO2_MAX_SAMPLE_RATE);
 
 			SCOPED_NAMED_EVENT(FMixerPlatformXAudio2_AsyncDeleteCreate_CreateMasterVoice, FColor::Blue);
-			HRESULT Result = DeviceSwapResult->NewSystem->CreateMasteringVoice(
-				&DeviceSwapResult->NewMasteringVoice,
-				(uint32)InDeviceContext->NewDevice->NumChannels,
-				(uint32)InDeviceContext->NewDevice->SampleRate,
-				0,
-				*InDeviceContext->NewDevice->DeviceId,
-				nullptr,
-				AudioCategory_GameEffects
-			);
-			if (FAILED(Result))
+			DeviceSwapResult->NewMasteringVoice = CreateMasteringVoice(*DeviceSwapResult->NewSystem, *InDeviceContext->NewDevice, InDeviceContext->bUseDefaultDevice);
+			if (!DeviceSwapResult->NewMasteringVoice)
 			{
-				UE_LOG(LogAudioMixer, Error, TEXT("XAudio2System->CreateMasteringVoice() -> 0x%X: %s (line: %d) with Args (NumChannels=%u, SampleRate=%u, DeviceID=%s Name=%s)"),
-					Result, *Audio::ToErrorFString(Result), __LINE__, (uint32)InDeviceContext->NewDevice->NumChannels,
-					(uint32)InDeviceContext->NewDevice->SampleRate, *InDeviceContext->NewDevice->DeviceId, *InDeviceContext->NewDevice->Name);
-
-				if (DeviceSwapResult->NewMasteringVoice)
-				{
-					DeviceSwapResult->NewMasteringVoice->DestroyVoice();
-					DeviceSwapResult->NewMasteringVoice = nullptr;
-				}
-				
 				SAFE_RELEASE(DeviceSwapResult->NewSystem);
 				return {}; // FAIL.
 			}
@@ -444,9 +424,9 @@ namespace Audio
 		}
 
 		// Optionally for testing, sleep for some duration in order to help repro race conditions.
-		if (ThreadedSwapDebugExtraTimeMsCVar > 0.f)
+		if (GThreadedSwapDebugExtraTimeMs > 0.f)
 		{
-			FPlatformProcess::Sleep(ThreadedSwapDebugExtraTimeMsCVar/1000.f); 
+			FPlatformProcess::Sleep(GThreadedSwapDebugExtraTimeMs / 1000.f); 
 		}
 
 		// Listen session for changes to this device.
@@ -921,32 +901,14 @@ namespace Audio
 				SetOriginalAudioDeviceId(AudioStreamInfo.DeviceInfo.DeviceId);
 			}
 
-#if PLATFORM_WINDOWS			
-			XAUDIO2_CALL_AND_HANDLE_ERROR(
-				XAudio2System->CreateMasteringVoice(
-					&OutputAudioStreamMasteringVoice,
-					AudioStreamInfo.DeviceInfo.NumChannels,
-					AudioStreamInfo.DeviceInfo.SampleRate,
-					0,
-					*AudioStreamInfo.DeviceInfo.DeviceId,
-					nullptr,
-					AudioCategory_GameEffects)
-					, goto Cleanup);
-#else //PLATFORM_WINDOWS
 			// Passing the device-id to CreateMasteringVoice on a non-windows platform, will prevent
 			// the creation of virtualized device which handles disconnection state for us, but
 			// if we need to handle these errors i.e. OnCriticalError callback, we need to pass the device here.
-			// This allows device swapping, so protect with a CVAR.
-			XAUDIO2_CALL_AND_HANDLE_ERROR(
-				 XAudio2System->CreateMasteringVoice(
-					&OutputAudioStreamMasteringVoice,
-					AudioStreamInfo.DeviceInfo.NumChannels,
-					AudioStreamInfo.DeviceInfo.SampleRate,
-					0,
-					*AudioStreamInfo.DeviceInfo.DeviceId, 
-					nullptr),
-				goto Cleanup);
-#endif //#else PLATFORM_WINDOWS
+			OutputAudioStreamMasteringVoice = CreateMasteringVoice(*XAudio2System, AudioStreamInfo.DeviceInfo, ShouldUseDefaultDevice());
+			if (!OutputAudioStreamMasteringVoice)
+			{
+				goto Cleanup;
+			}
 
 			// Start the xaudio2 engine running, which will now allow us to start feeding audio to it
 			XAUDIO2_CALL_AND_HANDLE_ERROR(XAudio2System->StartEngine(), goto Cleanup);
@@ -1219,6 +1181,7 @@ namespace Audio
 		}
 
 		DeviceSwapContext->NewDevice = InDeviceInfo;
+		DeviceSwapContext->bUseDefaultDevice = ShouldUseDefaultDevice();
 
 		return true;
 	}
@@ -1288,6 +1251,34 @@ namespace Audio
 		return true;
 	}
 
+	IXAudio2MasteringVoice* FMixerPlatformXAudio2::CreateMasteringVoice(IXAudio2& InXAudio2System, const FAudioPlatformDeviceInfo& NewDevice, bool bUseDefaultDevice)
+	{
+		IXAudio2MasteringVoice* MasteringVoice = nullptr;
+		const TCHAR* DeviceId = bUseDefaultDevice ? nullptr : *NewDevice.DeviceId;
+		const HRESULT Result = InXAudio2System.CreateMasteringVoice(
+			&MasteringVoice,
+			NewDevice.NumChannels,
+			NewDevice.SampleRate,
+			0,
+			DeviceId,
+			nullptr,
+			AudioCategory_GameEffects);
+		if (FAILED(Result))
+		{
+			if (MasteringVoice)
+			{
+				// Probably unreachable, but just to be safe...
+				MasteringVoice->DestroyVoice();
+				MasteringVoice = nullptr;
+			}
+
+			const TCHAR* DeviceIdDebug = DeviceId ? DeviceId : TEXT("(default)");
+			UE_LOG(LogAudioMixer, Error, TEXT("CreateMasteringVoice failed with result 0x%X: %s (line: %d) with Args (NumChannels=%u, SampleRate=%u, DeviceID=%s, Name=%s)"),
+				Result, *Audio::ToErrorFString(Result), __LINE__, NewDevice.NumChannels, NewDevice.SampleRate, DeviceIdDebug, *NewDevice.Name);
+		}
+
+		return MasteringVoice;
+	}
 
 	void FMixerPlatformXAudio2::OnProcessingPassStart()
 	{
diff --git a/Engine/Source/Runtime/Windows/AudioMixerXAudio2/Private/AudioMixerPlatformXAudio2.h b/Engine/Source/Runtime/Windows/AudioMixerXAudio2/Private/AudioMixerPlatformXAudio2.h
--- a/Engine/Source/Runtime/Windows/AudioMixerXAudio2/Private/AudioMixerPlatformXAudio2.h
+++ b/Engine/Source/Runtime/Windows/AudioMixerXAudio2/Private/AudioMixerPlatformXAudio2.h
@@ -69,6 +69,7 @@ namespace Audio
 			FDeviceSwapContext(InRequestedDeviceID, InReason)
 		{}
 		
+		bool bUseDefaultDevice = false;
 		IXAudio2* PreviousSystem = nullptr;
 		IXAudio2MasteringVoice* PreviousMasteringVoice = nullptr;
 		IXAudio2SourceVoice* PreviousSourceVoice = nullptr;
@@ -147,6 +148,9 @@ namespace Audio
 	
 	protected:
 		virtual uint32 GetCreateFlags() const { return 0; }
+		virtual bool ShouldUseDefaultDevice() const { return false; }
+
+		static IXAudio2MasteringVoice* CreateMasteringVoice(IXAudio2& InXAudio2System, const FAudioPlatformDeviceInfo& NewDevice, bool bUseDefaultDevice);
 
 		//~ Begin IXAudio2EngineCallback
 		virtual void OnCriticalError(HRESULT Error) override;
