Fix error on TeardownHardware and dialing back some non-fatal warning logs to Display.

#rnx

diff --git a/Engine/Source/Runtime/AudioDeviceEnumeration/Windows/WindowsMMDeviceEnumeration/Private/WindowsMMDeviceInfoCache.cpp b/Engine/Source/Runtime/AudioDeviceEnumeration/Windows/WindowsMMDeviceEnumeration/Private/WindowsMMDeviceInfoCache.cpp
--- a/Engine/Source/Runtime/AudioDeviceEnumeration/Windows/WindowsMMDeviceEnumeration/Private/WindowsMMDeviceInfoCache.cpp
+++ b/Engine/Source/Runtime/AudioDeviceEnumeration/Windows/WindowsMMDeviceEnumeration/Private/WindowsMMDeviceInfoCache.cpp
@@ -701,13 +701,16 @@ namespace Audio
 	{
 		FReadScopeLock ReadLock(CacheMutationLock);
 
-		TArray<FCacheEntry> ActiveDevices;
-		for (const auto& Entry : Cache)
+		if (!InDeviceID.IsNone())
 		{
-			// HardwareId is used for DeviceId with aggregate devices
-			if (Entry.Value.HardwareId == InDeviceID)
+			TArray<FCacheEntry> ActiveDevices;
+			for (const auto& Entry : Cache)
 			{
-				return true;
+				// HardwareId is used for DeviceId with aggregate devices
+				if (Entry.Value.HardwareId == InDeviceID)
+				{
+					return true;
+				}
 			}
 		}
 
diff --git a/Engine/Source/Runtime/AudioDeviceEnumeration/Windows/WindowsMMDeviceEnumeration/Private/WindowsMMNotificationClient.cpp b/Engine/Source/Runtime/AudioDeviceEnumeration/Windows/WindowsMMDeviceEnumeration/Private/WindowsMMNotificationClient.cpp
--- a/Engine/Source/Runtime/AudioDeviceEnumeration/Windows/WindowsMMDeviceEnumeration/Private/WindowsMMNotificationClient.cpp
+++ b/Engine/Source/Runtime/AudioDeviceEnumeration/Windows/WindowsMMDeviceEnumeration/Private/WindowsMMNotificationClient.cpp
@@ -385,7 +385,7 @@ namespace Audio
 
 	HRESULT STDMETHODCALLTYPE FWindowsMMNotificationClient::OnDefaultDeviceChanged(EDataFlow InFlow, ERole InRole, LPCWSTR pwstrDeviceId)
 	{
-		UE_CLOG(WindowsMMCvarUtils::ShouldLogDeviceSwaps(), LogAudioEnumeration, Warning,
+		UE_CLOG(WindowsMMCvarUtils::ShouldLogDeviceSwaps(), LogAudioEnumeration, Display,
 			TEXT("FWindowsMMNotificationClient: OnDefaultDeviceChanged: %s, %s, %s - %s"), ToString(InFlow), ToString(InRole), pwstrDeviceId, *GetFriendlyName(pwstrDeviceId));
 
 		Audio::EAudioDeviceRole AudioDeviceRole;
diff --git a/Engine/Source/Runtime/AudioMixerCore/Private/AudioMixer.cpp b/Engine/Source/Runtime/AudioMixerCore/Private/AudioMixer.cpp
--- a/Engine/Source/Runtime/AudioMixerCore/Private/AudioMixer.cpp
+++ b/Engine/Source/Runtime/AudioMixerCore/Private/AudioMixer.cpp
@@ -936,10 +936,11 @@ namespace Audio
 			);
 
 			// Set up the context for this device swap
-			InitializeDeviceSwapContext(DeviceID, InReason);
-
-			// Set the flag indicating we wish to begin a device swap
-			AudioStreamInfo.StreamState = EAudioOutputStreamState::SwappingDevice;
+			if (InitializeDeviceSwapContext(DeviceID, InReason))
+			{
+				// Set the flag indicating we wish to begin a device swap
+				AudioStreamInfo.StreamState = EAudioOutputStreamState::SwappingDevice;
+			}
 		
 			DeviceSwapCriticalSection.Unlock();
 
@@ -1032,7 +1033,7 @@ namespace Audio
 				
 				if (!PreDeviceSwap())
 				{
-					UE_LOG(LogAudioMixer, Warning, TEXT("FAudioMixerPlatformSwappable::CheckThreadedDeviceSwap PreDeviceSwap() failed"));
+					UE_LOG(LogAudioMixer, Display, TEXT("FAudioMixerPlatformSwappable::CheckThreadedDeviceSwap PreDeviceSwap() failed"));
 					return bDidStopGeneratingAudio;
 				}
 				
diff --git a/Engine/Source/Runtime/Windows/AudioMixerWasapi/Private/AudioMixerWasapi.cpp b/Engine/Source/Runtime/Windows/AudioMixerWasapi/Private/AudioMixerWasapi.cpp
--- a/Engine/Source/Runtime/Windows/AudioMixerWasapi/Private/AudioMixerWasapi.cpp
+++ b/Engine/Source/Runtime/Windows/AudioMixerWasapi/Private/AudioMixerWasapi.cpp
@@ -468,7 +468,7 @@ namespace Audio
 				}
 				else
 				{
-					UE_LOG(LogAudioMixer, Warning, TEXT("FAudioMixerWasapi::InitializeDeviceSwapContext - Ignoring attempt to switch to device with unsupported params: Channels=%u, SampleRate=%u, Id=%s, Name=%s"),
+					UE_LOG(LogAudioMixer, Display, TEXT("FAudioMixerWasapi::InitializeDeviceSwapContext - Ignoring attempt to switch to device with unsupported params: Channels=%u, SampleRate=%u, Id=%s, Name=%s"),
 						(uint32)TempDeviceInfo->NumChannels, (uint32)TempDeviceInfo->SampleRate, *TempDeviceInfo->DeviceId, *TempDeviceInfo->Name);
 
 					return false;
diff --git a/Engine/Source/Runtime/Windows/AudioMixerWasapi/Private/AudioMixerWasapiDeviceListener.cpp b/Engine/Source/Runtime/Windows/AudioMixerWasapi/Private/AudioMixerWasapiDeviceListener.cpp
--- a/Engine/Source/Runtime/Windows/AudioMixerWasapi/Private/AudioMixerWasapiDeviceListener.cpp
+++ b/Engine/Source/Runtime/Windows/AudioMixerWasapi/Private/AudioMixerWasapiDeviceListener.cpp
@@ -164,6 +164,22 @@ namespace Audio
 		{
 			return;
 		}
+
+		// If the device we're currently using was removed and it's not the system default... then switch to the new default audio device.
+		// If it is the system default device, then OnDefaultRenderDeviceChanged() will be called to handle this.
+		if (AudioStreamInfo.DeviceInfo.DeviceId == DeviceId && !AudioStreamInfo.DeviceInfo.bIsSystemDefault &&
+			(InState == EAudioDeviceState::Disabled || InState == EAudioDeviceState::NotPresent || InState == EAudioDeviceState::Unplugged))
+		{
+			UE_LOG(LogAudioMixer, Display, TEXT("FAudioMixerWasapi::OnDeviceStateChanged: Audio device not available [%s], falling back to other windows default device. InState=%d, bIsRenderDevice=%d, InstanceID=%d"), 
+				*WasapiWinNotificationClient->GetFriendlyName(DeviceId), (int32)InState, (int32)bIsRenderDevice, InstanceID);
+
+			// Ignore if not listening for events. This means the user has specified a specific device
+			// to use and is not interested in following the current system default device.
+			if (GetIsListeningForDeviceEvents())
+			{
+				RequestDeviceSwap(TEXT(""), /* force */ true, TEXT("FAudioMixerWasapi::OnDeviceStateChanged"));
+			}
+		}
 		
 		if (UAudioDeviceNotificationSubsystem* AudioDeviceNotifySubsystem = UAudioDeviceNotificationSubsystem::Get())
 		{
diff --git a/Engine/Source/Runtime/Windows/AudioMixerWasapi/Private/AudioMixerWasapiDeviceThread.cpp b/Engine/Source/Runtime/Windows/AudioMixerWasapi/Private/AudioMixerWasapiDeviceThread.cpp
--- a/Engine/Source/Runtime/Windows/AudioMixerWasapi/Private/AudioMixerWasapiDeviceThread.cpp
+++ b/Engine/Source/Runtime/Windows/AudioMixerWasapi/Private/AudioMixerWasapiDeviceThread.cpp
@@ -89,7 +89,7 @@ namespace Audio
 		bIsRunning = false;
 		if (OutputStreamTimeoutsDetected > 0)
 		{
-			UE_LOG(LogAudioMixerWasapi, Error, TEXT("FAudioMixerWasapiRunnable::Stop render stream reported %d timeouts"), OutputStreamTimeoutsDetected);
+			UE_LOG(LogAudioMixerWasapi, Display, TEXT("FAudioMixerWasapiRunnable::Stop render stream reported %d timeouts"), OutputStreamTimeoutsDetected);
 		}
 	}
 
diff --git a/Engine/Source/Runtime/Windows/AudioMixerWasapi/Private/AudioMixerWasapiRenderStream.cpp b/Engine/Source/Runtime/Windows/AudioMixerWasapi/Private/AudioMixerWasapiRenderStream.cpp
--- a/Engine/Source/Runtime/Windows/AudioMixerWasapi/Private/AudioMixerWasapiRenderStream.cpp
+++ b/Engine/Source/Runtime/Windows/AudioMixerWasapi/Private/AudioMixerWasapiRenderStream.cpp
@@ -3,9 +3,18 @@
 #include "AudioMixerWasapiRenderStream.h"
 
 #include "AudioMixerWasapiLog.h"
+#include "HAL/IConsoleManager.h"
 #include "WasapiAudioUtils.h"
 #include "WindowsMMStringUtils.h"
 
+static int32 UseDefaultQualitySRC_CVar = 0;
+FAutoConsoleVariableRef CVarUseDefaultQualitySRC(
+	TEXT("au.Wasapi.UseDefaultQualitySRC"),
+	UseDefaultQualitySRC_CVar,
+	TEXT("Enable Wasapi default SRC quality.\n")
+	TEXT("0: Not Enabled, 1: Enabled"),
+	ECVF_Default);
+
 namespace Audio
 {
 	uint32 FAudioMixerWasapiRenderStream::GetMinimumBufferSize(const uint32 InSampleRate)
@@ -40,19 +49,7 @@ namespace Audio
 			UE_LOG(LogAudioMixerWasapi, Error, TEXT("OpenAudioStream failed IMMDevice::Activate %s"), *AudioClientErrorToFString(Result));
 			return false;
 		}
-
-		AudioClientProperties AudioProps = { 0 };
-		AudioProps.cbSize = sizeof(AudioClientProperties);
-		AudioProps.bIsOffload = false;
-		AudioProps.eCategory = AudioCategory_Media;
-
-		Result = TempAudioClient->SetClientProperties(&AudioProps);
-		if (FAILED(Result))
-		{
-			UE_LOG(LogAudioMixerWasapi, Error, TEXT("OpenAudioStream failed IAudioClient3::SetClientProperties %s"), *AudioClientErrorToFString(Result));
-			return false;
-		}
-
+		
 		WAVEFORMATEX* MixFormat = nullptr;
 		Result = TempAudioClient->GetMixFormat(&MixFormat);
 		if (FAILED(Result) || !MixFormat)
@@ -113,6 +110,11 @@ namespace Audio
 		if (InParams.SampleRate != InParams.HardwareDeviceInfo.SampleRate)
 		{
 			Flags |= AUDCLNT_STREAMFLAGS_AUTOCONVERTPCM;
+			if (UseDefaultQualitySRC_CVar)
+			{
+				Flags |= AUDCLNT_STREAMFLAGS_SRC_DEFAULT_QUALITY;
+			}
+			
 			UE_LOG(LogAudioMixerWasapi, Warning, TEXT("Sample rate mismatch. Engine sample rate: %d Device sample rate: %d"), InParams.SampleRate, InParams.HardwareDeviceInfo.SampleRate);
 			UE_LOG(LogAudioMixerWasapi, Warning, TEXT("Device level sample rate conversion will be used."));
 		}
@@ -271,7 +273,7 @@ namespace Audio
 
 		if (CallbackBufferErrors > 0)
 		{
-			UE_LOG(LogAudioMixerWasapi, Error, TEXT("FAudioMixerWasapiRenderStream::StopAudioStream render stream reported %d callback buffer errors"), CallbackBufferErrors);
+			UE_LOG(LogAudioMixerWasapi, Display, TEXT("FAudioMixerWasapiRenderStream::StopAudioStream render stream reported %d callback buffer errors (can be normal if preceded by device swap)."), CallbackBufferErrors);
 			CallbackBufferErrors = 0;
 		}
 
diff --git a/Engine/Source/Runtime/Windows/AudioMixerWasapi/Private/WasapiDefaultDeviceMgr.cpp b/Engine/Source/Runtime/Windows/AudioMixerWasapi/Private/WasapiDefaultDeviceMgr.cpp
--- a/Engine/Source/Runtime/Windows/AudioMixerWasapi/Private/WasapiDefaultDeviceMgr.cpp
+++ b/Engine/Source/Runtime/Windows/AudioMixerWasapi/Private/WasapiDefaultDeviceMgr.cpp
@@ -79,10 +79,7 @@ namespace Audio
 		
 		if (MainRenderStreamDevice.IsValid())
 		{
-			bool bCloseSucceeded = MainRenderStreamDevice->CloseAudioStream();
-			MainRenderStreamDevice.Reset();
-
-			return bCloseSucceeded;
+			return MainRenderStreamDevice->CloseAudioStream();
 		}
 
 		return false;
diff --git a/Engine/Source/Runtime/Windows/AudioMixerXAudio2/Private/AudioMixerPlatformWindows.cpp b/Engine/Source/Runtime/Windows/AudioMixerXAudio2/Private/AudioMixerPlatformWindows.cpp
--- a/Engine/Source/Runtime/Windows/AudioMixerXAudio2/Private/AudioMixerPlatformWindows.cpp
+++ b/Engine/Source/Runtime/Windows/AudioMixerXAudio2/Private/AudioMixerPlatformWindows.cpp
@@ -184,6 +184,22 @@ namespace Audio
 			return;
 		}
 		
+		// If the device we're currently using was removed and it's not the system default... then switch to the new default audio device.
+		// If it is the system default device, then OnDefaultRenderDeviceChanged() will be called to handle this.
+		if (AudioStreamInfo.DeviceInfo.DeviceId == DeviceId && !AudioStreamInfo.DeviceInfo.bIsSystemDefault &&
+			(InState == EAudioDeviceState::Disabled || InState == EAudioDeviceState::NotPresent || InState == EAudioDeviceState::Unplugged))
+		{
+			UE_LOG(LogAudioMixer, Display, TEXT("FMixerPlatformXAudio2: Audio device not available [%s], falling back to other windows default device. InState=%d, bIsRenderDevice=%d, InstanceID=%d"), 
+				*WindowsNotificationClient->GetFriendlyName(DeviceId), (int32)InState, (int32)bIsRenderDevice, InstanceID);
+
+			// Ignore if not listening for events. This means the user has specified a specific device
+			// to use and is not interested in following the current system default device.
+			if (GetIsListeningForDeviceEvents())
+			{
+				RequestDeviceSwap(TEXT(""), /* force */ true, TEXT("FMixerPlatformXAudio2::OnDeviceStateChanged"));
+			}
+		}
+
 		if (UAudioDeviceNotificationSubsystem* AudioDeviceNotifSubsystem = UAudioDeviceNotificationSubsystem::Get())
 		{
 			AudioDeviceNotifSubsystem->OnDeviceStateChanged(DeviceId, InState, bIsRenderDevice);
