crash fix when resizing window resolution with -renderoffscreen

diff --git a/Engine/Source/Runtime/Windows/D3D11RHI/Private/Windows/WindowsD3D11Viewport.cpp b/Engine/Source/Runtime/Windows/D3D11RHI/Private/Windows/WindowsD3D11Viewport.cpp
--- a/Engine/Source/Runtime/Windows/D3D11RHI/Private/Windows/WindowsD3D11Viewport.cpp
+++ b/Engine/Source/Runtime/Windows/D3D11RHI/Private/Windows/WindowsD3D11Viewport.cpp
@@ -321,6 +321,11 @@ static const FString GetDXGIColorSpaceString(DXGI_COLOR_SPACE_TYPE ColorSpace)
 
 inline void EnsureColorSpace(IDXGISwapChain* SwapChain, EDisplayColorGamut DisplayGamut, EDisplayOutputFormat OutputDevice, EPixelFormat PixelFormat)
 {
+	if (SwapChain == nullptr)
+	{
+		return;
+	}
+
 	TRefCountPtr<IDXGISwapChain3> swapChain3;
 	if (FAILED(SwapChain->QueryInterface(IID_PPV_ARGS(swapChain3.GetInitReference()))))
 	{
