Add CPU marker to track D3D11 present time in Insights

Copied pattern from D3D12 RHI.

diff --git a/Engine/Source/Runtime/Windows/D3D11RHI/Private/D3D11Viewport.cpp b/Engine/Source/Runtime/Windows/D3D11RHI/Private/D3D11Viewport.cpp
--- a/Engine/Source/Runtime/Windows/D3D11RHI/Private/D3D11Viewport.cpp
+++ b/Engine/Source/Runtime/Windows/D3D11RHI/Private/D3D11Viewport.cpp
@@ -411,6 +411,7 @@ bool FD3D11Viewport::PresentChecked(IRHICommandContext& RHICmdContext, int32 Syn
 
 			{
 				FRenderThreadIdleScope IdleScope(ERenderThreadIdleTypes::WaitingForGPUPresent);
+				TRACE_CPUPROFILER_EVENT_SCOPE(D3D11_Present);
 				Result = SwapChain->Present(SyncInterval, Flags);
 			}
 
