Flag OOM asserts catched in D3D RHITerminateOnOutOfMemory as OOM for crash reporting purposes.

diff --git a/Engine/Source/Runtime/D3D12RHI/Private/D3D12Util.cpp b/Engine/Source/Runtime/D3D12RHI/Private/D3D12Util.cpp
--- a/Engine/Source/Runtime/D3D12RHI/Private/D3D12Util.cpp
+++ b/Engine/Source/Runtime/D3D12RHI/Private/D3D12Util.cpp
@@ -805,7 +805,8 @@ void FD3D12DynamicRHI::TerminateOnOutOfMemory(ID3D12Device* InDevice, HRESULT D3
 
 	// Also log Windows memory stats.
 	FPlatformMemory::DumpStats(*GLog);
-	
+	FPlatformMemory::bIsOOM = true;
+
 	UE_LOG(LogD3D12RHI, Fatal, TEXT("Out of video memory trying to allocate a rendering resource"));
 	if (!bGPUCrashOOM)
 	{
diff --git a/Engine/Source/Runtime/Windows/D3D11RHI/Private/D3D11Util.cpp b/Engine/Source/Runtime/Windows/D3D11RHI/Private/D3D11Util.cpp
--- a/Engine/Source/Runtime/Windows/D3D11RHI/Private/D3D11Util.cpp
+++ b/Engine/Source/Runtime/Windows/D3D11RHI/Private/D3D11Util.cpp
@@ -199,6 +199,8 @@ static void TerminateOnOutOfMemory(HRESULT D3DResult, bool bCreatingTextures)
 		uint64 VRAMBudgetBytes = 0, VRAMUsageBytes = 0;
 		GetAndLogMemoryStats(GD3D11RHI->GetAdapter(), VRAMBudgetBytes, VRAMUsageBytes);
 		FPlatformMemory::DumpStats(*GLog);
+		FPlatformMemory::bIsOOM = true;
+
 		FCoreDelegates::GetGPUOutOfMemoryDelegate().Broadcast(VRAMBudgetBytes, VRAMUsageBytes);
 
 		if (bCreatingTextures)
