Porting suspend/resume improvements over from XAudio2 implementation. Fixing bug where headphones only appear under EAudioDeviceRole::Communications.

#rnx

diff --git a/Engine/Source/Runtime/Windows/AudioMixerWasapi/Private/AudioMixerWasapiDeviceListener.cpp b/Engine/Source/Runtime/Windows/AudioMixerWasapi/Private/AudioMixerWasapiDeviceListener.cpp
--- a/Engine/Source/Runtime/Windows/AudioMixerWasapi/Private/AudioMixerWasapiDeviceListener.cpp
+++ b/Engine/Source/Runtime/Windows/AudioMixerWasapi/Private/AudioMixerWasapiDeviceListener.cpp
@@ -79,8 +79,9 @@ namespace Audio
 		// However, the Windows UX only allows console+multimedia to be toggle as a pair. This means you get two notifications
 		// for default device changing typically. To prevent a trouble trigger we only listen to "Console" here. For more information on 
 		// device roles: https://docs.microsoft.com/en-us/windows/win32/coreaudio/device-roles
+		// Note: On XBox, headphones plugged into an XBox controller will use the Communications role.
 		
-		if (InAudioDeviceRole == EAudioDeviceRole::Console)
+		if (InAudioDeviceRole == EAudioDeviceRole::Console || InAudioDeviceRole == EAudioDeviceRole::Communications)
 		{
 			UE_LOG(LogAudioMixer, Display, TEXT("FAudioMixerWasapi: Changing default audio render device to new device: Role=%s, DeviceName=%s, InstanceID=%d"), 
 				Audio::ToString(InAudioDeviceRole), *WasapiWinNotificationClient->GetFriendlyName(DeviceId), InstanceID);
