Just silence the in practice always true PVS warning instead of trying to put a check in place
#rnx

diff --git a/Engine/Source/Runtime/Windows/D3D11RHI/Private/D3D11Commands.cpp b/Engine/Source/Runtime/Windows/D3D11RHI/Private/D3D11Commands.cpp
--- a/Engine/Source/Runtime/Windows/D3D11RHI/Private/D3D11Commands.cpp
+++ b/Engine/Source/Runtime/Windows/D3D11RHI/Private/D3D11Commands.cpp
@@ -826,9 +826,8 @@ void FD3D11DynamicRHI::SetRenderTargets(
 	ID3D11RenderTargetView* NewRenderTargetViews[MaxSimultaneousRenderTargets];
 	for(uint32 RenderTargetIndex = 0;RenderTargetIndex < MaxSimultaneousRenderTargets;++RenderTargetIndex)
 	{
-		check(NewRenderTargetsRHI);
 		ID3D11RenderTargetView* RenderTargetView = NULL;
-		if(RenderTargetIndex < NewNumSimultaneousRenderTargets && NewRenderTargetsRHI[RenderTargetIndex].Texture != nullptr)
+		if(RenderTargetIndex < NewNumSimultaneousRenderTargets && NewRenderTargetsRHI[RenderTargetIndex].Texture != nullptr) //-V522
 		{
 			int32 RTMipIndex = NewRenderTargetsRHI[RenderTargetIndex].MipIndex;
 			int32 RTSliceIndex = NewRenderTargetsRHI[RenderTargetIndex].ArraySliceIndex;
