Fix Wasapi backend to use null renderer when no audio devices are available.

#rnx

diff --git a/Engine/Source/Runtime/Windows/AudioMixerWasapi/Private/AudioMixerWasapi.cpp b/Engine/Source/Runtime/Windows/AudioMixerWasapi/Private/AudioMixerWasapi.cpp
--- a/Engine/Source/Runtime/Windows/AudioMixerWasapi/Private/AudioMixerWasapi.cpp
+++ b/Engine/Source/Runtime/Windows/AudioMixerWasapi/Private/AudioMixerWasapi.cpp
@@ -142,7 +142,7 @@ namespace Audio
 		FAudioPlatformDeviceInfo DeviceInfo;
 		if (!GetOutputDeviceInfo(InDeviceIndex, DeviceInfo))
 		{
-			UE_LOG(LogAudioMixer, Error, TEXT("FAudioMixerWasapi::InitStreamParams unable to find default device"));
+			UE_LOG(LogAudioMixer, Display, TEXT("FAudioMixerWasapi::InitStreamParams unable to find default device"));
 			return false;
 		}
 
@@ -209,59 +209,69 @@ namespace Audio
 		// ignore device change events.
 		SetIsListeningForDeviceEvents(Params.bUseSystemAudioDevice);
 
+		bool bIsAggregateDevice = false;
 		TArray<FWasapiRenderStreamParams> StreamParams;
-		if (!InitStreamParams(OpenStreamParams.OutputDeviceIndex, OpenStreamParams.NumFrames, OpenStreamParams.NumBuffers, OpenStreamParams.SampleRate, StreamParams))
+
+		if (InitStreamParams(OpenStreamParams.OutputDeviceIndex, OpenStreamParams.NumFrames, OpenStreamParams.NumBuffers, OpenStreamParams.SampleRate, StreamParams))
 		{
-			return false;
-		}
+			// Adopt the first device info. In the case of an aggregate device, all of the sub-devices will
+			// be identical because they belong to the same physical device.
+			AudioStreamInfo.DeviceInfo = StreamParams[0].HardwareDeviceInfo;
 
-		// Adopt the first device info. In the case of an aggregate device, all of the sub-devices will
-		// be identical because they belong to the same physical device.
-		AudioStreamInfo.DeviceInfo = StreamParams[0].HardwareDeviceInfo;
+			// Set the current device name
+			bIsAggregateDevice = GetDeviceInfoCache()->IsAggregateHardwareDeviceId(*Params.AudioDeviceId);
+			if (bIsAggregateDevice)
+			{
+				CurrentDeviceName = ExtractAggregateDeviceName(AudioStreamInfo.DeviceInfo.Name);
+			}
+			else
+			{
+				CurrentDeviceName = AudioStreamInfo.DeviceInfo.Name;
+			}
 
-		// Set the current device name
-		const bool bIsAggregateDevice = GetDeviceInfoCache()->IsAggregateHardwareDeviceId(*Params.AudioDeviceId);
-		if (bIsAggregateDevice)
-		{
-			CurrentDeviceName = ExtractAggregateDeviceName(AudioStreamInfo.DeviceInfo.Name);
-		}
-		else
-		{
-			CurrentDeviceName = AudioStreamInfo.DeviceInfo.Name;
-		}
+			// Create and initialize the device manager
+			CreateDeviceManager(bIsAggregateDevice, DeviceManager);
 
-		// Create and initialize the device manager
-		CreateDeviceManager(bIsAggregateDevice, DeviceManager);
-		
-		// The ReadNextBufferCallback life cycle is tied to this object. It is ultimately bound to a delegate
-		// in the render stream object which will be unbound in TeardownHardware, prior to 'this' being deallocated.
-		TFunction<void()> ReadNextBufferCallback = [this](){ ReadNextBuffer(); };
-		if (!ensure(DeviceManager && DeviceManager->InitializeHardware(StreamParams, ReadNextBufferCallback)))
-		{
-			AUDIO_PLATFORM_LOG_ONCE(TEXT("FAudioMixerWasapi::OpenAudioStream DeviceManager->InitializeHardware() failed"), Warning);
-			return false;
-		}
+			// The ReadNextBufferCallback life cycle is tied to this object. It is ultimately bound to a delegate
+			// in the render stream object which will be unbound in TeardownHardware, prior to 'this' being deallocated.
+			TFunction<void()> ReadNextBufferCallback = [this]() { ReadNextBuffer(); };
+			if (!ensure(DeviceManager && DeviceManager->InitializeHardware(StreamParams, ReadNextBufferCallback)))
+			{
+				AUDIO_PLATFORM_LOG_ONCE(TEXT("FAudioMixerWasapi::OpenAudioStream DeviceManager->InitializeHardware() failed"), Warning);
+				return false;
+			}
 
-		// Assign the total number of direct out channels based on the device manager. For WasapiDefaultDeviceMgr this will be 0 
-		// and for WasapiAggregateDeviceMgr this will be the total of all the channels less the main outs (first 8 channels).
-		AudioStreamInfo.DeviceInfo.NumDirectOutChannels = DeviceManager->GetNumDirectOutChannels();
-		
-		if (!ensure(DeviceManager.IsValid() && DeviceManager->OpenAudioStream(StreamParams)))
-		{
-			UE_LOG(LogAudioMixer, Error, TEXT("FAudioMixerWasapi::OpenAudioStream DeviceManager->OpenAudioStream() failed"));
-			return false;
-		}
+			// Assign the total number of direct out channels based on the device manager. For WasapiDefaultDeviceMgr this will be 0 
+			// and for WasapiAggregateDeviceMgr this will be the total of all the channels less the main outs (first 8 channels).
+			AudioStreamInfo.DeviceInfo.NumDirectOutChannels = DeviceManager->GetNumDirectOutChannels();
 
-		// Store the device ID here in case it is removed. We can switch back if the device comes back.
-		if (Params.bRestoreIfRemoved)
+			if (!ensure(DeviceManager.IsValid() && DeviceManager->OpenAudioStream(StreamParams)))
+			{
+				UE_LOG(LogAudioMixer, Error, TEXT("FAudioMixerWasapi::OpenAudioStream DeviceManager->OpenAudioStream() failed"));
+				return false;
+			}
+
+			// Store the device ID here in case it is removed. We can switch back if the device comes back.
+			if (Params.bRestoreIfRemoved)
+			{
+				SetOriginalAudioDeviceId(AudioStreamInfo.DeviceInfo.DeviceId);
+			}
+		}
+		else
 		{
-			SetOriginalAudioDeviceId(AudioStreamInfo.DeviceInfo.DeviceId);
+			// Setup for running null device.
+			AudioStreamInfo.NumOutputFrames = OpenStreamParams.NumFrames;
+			AudioStreamInfo.DeviceInfo.OutputChannelArray = { EAudioMixerChannel::FrontLeft, EAudioMixerChannel::FrontRight };
+			AudioStreamInfo.DeviceInfo.NumChannels = 2;
+			AudioStreamInfo.DeviceInfo.SampleRate = OpenStreamParams.SampleRate;
+			AudioStreamInfo.DeviceInfo.Format = EAudioMixerStreamDataFormat::Float;
 		}
 
+		AudioStreamInfo.StreamState = EAudioOutputStreamState::Open;
 		bIsInitialized = true;
 
-		UE_LOG(LogAudioMixer, Display, TEXT("FAudioMixerWasapi initialized: SampleRate=%d NumChannels=%d NumDirectOutChannels=%d bIsAggregateDevice=%d"), 
-			OpenStreamParams.SampleRate, AudioStreamInfo.DeviceInfo.NumChannels, AudioStreamInfo.DeviceInfo.NumDirectOutChannels, bIsAggregateDevice);
+		UE_LOG(LogAudioMixer, Display, TEXT("FAudioMixerWasapi initialized: SampleRate=%d NumChannels=%d NumDirectOutChannels=%d bIsAggregateDevice=%d StreamParams.Num()=%d"), 
+			OpenStreamParams.SampleRate, AudioStreamInfo.DeviceInfo.NumChannels, AudioStreamInfo.DeviceInfo.NumDirectOutChannels, bIsAggregateDevice, StreamParams.Num());
 
 		return true;
 	}
@@ -304,6 +314,7 @@ namespace Audio
 			{
 				check(!bIsUsingNullDevice);
 				StartRunningNullDevice();
+				bDidStartAudioStream = true;
 			}
 			
 			// Can be called during device swap when AudioRenderEvent can be null
