Fix crash resulting from error when activating low level MMDevice during hardware init.

#rnx

diff --git a/Engine/Source/Runtime/AudioDeviceEnumeration/Windows/WindowsMMDeviceEnumeration/Private/WindowsMMStringUtils.cpp b/Engine/Source/Runtime/AudioDeviceEnumeration/Windows/WindowsMMDeviceEnumeration/Private/WindowsMMStringUtils.cpp
--- a/Engine/Source/Runtime/AudioDeviceEnumeration/Windows/WindowsMMDeviceEnumeration/Private/WindowsMMStringUtils.cpp
+++ b/Engine/Source/Runtime/AudioDeviceEnumeration/Windows/WindowsMMDeviceEnumeration/Private/WindowsMMStringUtils.cpp
@@ -211,9 +211,12 @@ namespace Audio
 			FString Msg(FPlatformMisc::GetSystemErrorMessage(Buffer, UE_ARRAY_COUNT(Buffer), InResult));
 
 			// Anything to return? Otherwise "UNKNOWN"
-			return !Msg.IsEmpty() ?
-				Msg :
-				TEXT("UNKNOWN");
+			if (Msg.IsEmpty())
+			{
+				Msg = FString::Printf(TEXT("0x%x (UNKNOWN)"), InResult);
+			}
+
+			return Msg;
 		}
 		}
 	}
diff --git a/Engine/Source/Runtime/Windows/AudioMixerWasapi/Private/AudioMixerWasapi.cpp b/Engine/Source/Runtime/Windows/AudioMixerWasapi/Private/AudioMixerWasapi.cpp
--- a/Engine/Source/Runtime/Windows/AudioMixerWasapi/Private/AudioMixerWasapi.cpp
+++ b/Engine/Source/Runtime/Windows/AudioMixerWasapi/Private/AudioMixerWasapi.cpp
@@ -51,8 +51,9 @@ namespace Audio
 	{
 		if (!bIsInitialized)
 		{
-			AUDIO_PLATFORM_LOG_ONCE(TEXT("FAudioMixerWasapi::TeardownHardware failed...not initialized."), Warning);
-			return false;
+			// Don't return prematurely here because TeardownHardware can be called when initialization has failed
+			// and we want to clean up any state that exists.
+			AUDIO_PLATFORM_LOG_ONCE(TEXT("FAudioMixerWasapi::TeardownHardware called when not fully initialized. Check for init errors earlier in the log."), Warning);
 		}
 
 		// Lock prior to changing state to avoid race condition if there happens to be an in-flight device swap
@@ -235,7 +236,7 @@ namespace Audio
 			// The ReadNextBufferCallback life cycle is tied to this object. It is ultimately bound to a delegate
 			// in the render stream object which will be unbound in TeardownHardware, prior to 'this' being deallocated.
 			TFunction<void()> ReadNextBufferCallback = [this]() { ReadNextBuffer(); };
-			if (!ensure(DeviceManager && DeviceManager->InitializeHardware(StreamParams, ReadNextBufferCallback)))
+			if (!ensure(DeviceManager.IsValid() && DeviceManager->InitializeHardware(StreamParams, ReadNextBufferCallback)))
 			{
 				AUDIO_PLATFORM_LOG_ONCE(TEXT("FAudioMixerWasapi::OpenAudioStream DeviceManager->InitializeHardware() failed"), Warning);
 				return false;
@@ -265,6 +266,7 @@ namespace Audio
 			AudioStreamInfo.DeviceInfo.NumChannels = 2;
 			AudioStreamInfo.DeviceInfo.SampleRate = OpenStreamParams.SampleRate;
 			AudioStreamInfo.DeviceInfo.Format = EAudioMixerStreamDataFormat::Float;
+			AudioStreamInfo.DeviceInfo.NumDirectOutChannels = 0;
 		}
 
 		AudioStreamInfo.StreamState = EAudioOutputStreamState::Open;
diff --git a/Engine/Source/Runtime/Windows/AudioMixerWasapi/Private/AudioMixerWasapiRenderStream.cpp b/Engine/Source/Runtime/Windows/AudioMixerWasapi/Private/AudioMixerWasapiRenderStream.cpp
--- a/Engine/Source/Runtime/Windows/AudioMixerWasapi/Private/AudioMixerWasapiRenderStream.cpp
+++ b/Engine/Source/Runtime/Windows/AudioMixerWasapi/Private/AudioMixerWasapiRenderStream.cpp
@@ -42,8 +42,8 @@ namespace Audio
 			return false;
 		}
 
-		TComPtr<IAudioClient3> TempAudioClient;
-		HRESULT Result = MMDevice->Activate(__uuidof(IAudioClient3), CLSCTX_INPROC_SERVER, nullptr, IID_PPV_ARGS_Helper(&TempAudioClient));
+		TComPtr<IAudioClient> TempAudioClient;
+		HRESULT Result = MMDevice->Activate(__uuidof(IAudioClient), CLSCTX_INPROC_SERVER, nullptr, IID_PPV_ARGS_Helper(&TempAudioClient));
 		if (FAILED(Result))
 		{
 			UE_LOG(LogAudioMixerWasapi, Error, TEXT("OpenAudioStream failed IMMDevice::Activate %s"), *AudioClientErrorToFString(Result));
@@ -54,7 +54,7 @@ namespace Audio
 		Result = TempAudioClient->GetMixFormat(&MixFormat);
 		if (FAILED(Result) || !MixFormat)
 		{
-			UE_LOG(LogAudioMixerWasapi, Error, TEXT("OpenAudioStream failed IAudioClient3::GetMixFormat MixFormat: 0x%llx Result: %s"), (uint64)MixFormat, *AudioClientErrorToFString(Result));
+			UE_LOG(LogAudioMixerWasapi, Error, TEXT("OpenAudioStream failed IAudioClient::GetMixFormat MixFormat: 0x%llx Result: %s"), (uint64)MixFormat, *AudioClientErrorToFString(Result));
 			return false;
 		}
 
@@ -74,7 +74,7 @@ namespace Audio
 		Result = TempAudioClient->GetDevicePeriod(&DevicePeriodRefTime, nullptr);
 		if (FAILED(Result))
 		{
-			UE_LOG(LogAudioMixerWasapi, Error, TEXT("OpenAudioStream failed IAudioClient3::GetDevicePeriod %s"), *AudioClientErrorToFString(Result));
+			UE_LOG(LogAudioMixerWasapi, Error, TEXT("OpenAudioStream failed IAudioClient::GetDevicePeriod %s"), *AudioClientErrorToFString(Result));
 			return false;
 		}
 
@@ -86,7 +86,12 @@ namespace Audio
 		}
 
 		// Verify that our assumption about the minimum/default buffer size is correct
-		check(DefaultDevicePeriod == GetMinimumBufferSize(InParams.SampleRate));
+		const uint32 MinBufferSize = GetMinimumBufferSize(InParams.SampleRate);
+		if (!ensure(DefaultDevicePeriod == MinBufferSize))
+		{
+			UE_LOG(LogAudioMixerWasapi, Error, TEXT("FAudioMixerWasapiRenderStream::InitializeHardware() DefaultDevicePeriod = %d MinimumBufferSize = %d"), DefaultDevicePeriod, MinBufferSize);
+			return false;
+		}
 		
 		// Determine buffer size to use. 
 		uint32 BufferFramesToRequest = FMath::Max<uint32>(InParams.NumFrames, DefaultDevicePeriod);
@@ -122,14 +127,14 @@ namespace Audio
 		Result = TempAudioClient->Initialize(AUDCLNT_SHAREMODE_SHARED, Flags, DesiredBufferDuration, Periodicity, StreamFormat.GetWaveFormat(), nullptr);
 		if (FAILED(Result))
 		{
-			UE_LOG(LogAudioMixerWasapi, Error, TEXT("OpenAudioStream failed IAudioClient3::Initialize %s"), *AudioClientErrorToFString(Result));
+			UE_LOG(LogAudioMixerWasapi, Error, TEXT("OpenAudioStream failed IAudioClient::Initialize %s"), *AudioClientErrorToFString(Result));
 			return false;
 		}
 
 		Result = TempAudioClient->GetBufferSize(&NumFramesPerDeviceBuffer);
 		if (FAILED(Result))
 		{
-			UE_LOG(LogAudioMixerWasapi, Error, TEXT("OpenAudioStream failed IAudioClient3::GetBufferSize %s"), *AudioClientErrorToFString(Result));
+			UE_LOG(LogAudioMixerWasapi, Error, TEXT("OpenAudioStream failed IAudioClient::GetBufferSize %s"), *AudioClientErrorToFString(Result));
 			return false;
 		}
 
@@ -192,7 +197,7 @@ namespace Audio
 		HRESULT Result = AudioClient->SetEventHandle(InEventHandle);
 		if (FAILED(Result))
 		{
-			UE_LOG(LogAudioMixerWasapi, Error, TEXT("OpenAudioStream failed IAudioClient3::SetEventHandle %s"), *AudioClientErrorToFString(Result));
+			UE_LOG(LogAudioMixerWasapi, Error, TEXT("OpenAudioStream failed IAudioClient::SetEventHandle %s"), *AudioClientErrorToFString(Result));
 			return false;
 		}
 
@@ -200,7 +205,7 @@ namespace Audio
 		Result = AudioClient->GetService(__uuidof(IAudioRenderClient), IID_PPV_ARGS_Helper(&TempRenderClient));
 		if (FAILED(Result))
 		{
-			UE_LOG(LogAudioMixerWasapi, Error, TEXT("OpenAudioStream failed IAudioClient3::GetService IAudioRenderClient %s"), *AudioClientErrorToFString(Result));
+			UE_LOG(LogAudioMixerWasapi, Error, TEXT("OpenAudioStream failed IAudioClient::GetService IAudioRenderClient %s"), *AudioClientErrorToFString(Result));
 			return false;
 		}
 
diff --git a/Engine/Source/Runtime/Windows/AudioMixerWasapi/Private/AudioMixerWasapiRenderStream.h b/Engine/Source/Runtime/Windows/AudioMixerWasapi/Private/AudioMixerWasapiRenderStream.h
--- a/Engine/Source/Runtime/Windows/AudioMixerWasapi/Private/AudioMixerWasapiRenderStream.h
+++ b/Engine/Source/Runtime/Windows/AudioMixerWasapi/Private/AudioMixerWasapiRenderStream.h
@@ -49,7 +49,7 @@ namespace Audio
 	protected:
 		
 		/** COM pointer to the WASAPI audio client object. */
-		TComPtr<IAudioClient3> AudioClient;
+		TComPtr<IAudioClient> AudioClient;
 
 		/** COM pointer to the WASAPI render client object. */
 		TComPtr<IAudioRenderClient> RenderClient;
