Cleanup GetMinBufferSize and replace with default value returned form device.

#rnx

diff --git a/Engine/Source/Runtime/Windows/AudioMixerWasapi/Private/AudioMixerWasapiRenderStream.cpp b/Engine/Source/Runtime/Windows/AudioMixerWasapi/Private/AudioMixerWasapiRenderStream.cpp
--- a/Engine/Source/Runtime/Windows/AudioMixerWasapi/Private/AudioMixerWasapiRenderStream.cpp
+++ b/Engine/Source/Runtime/Windows/AudioMixerWasapi/Private/AudioMixerWasapiRenderStream.cpp
@@ -85,14 +85,6 @@ namespace Audio
 			return false;
 		}
 
-		// Verify that our assumption about the minimum/default buffer size is correct
-		const uint32 MinBufferSize = GetMinimumBufferSize(InParams.SampleRate);
-		if (!ensure(DefaultDevicePeriod == MinBufferSize))
-		{
-			UE_LOG(LogAudioMixerWasapi, Error, TEXT("FAudioMixerWasapiRenderStream::InitializeHardware() DefaultDevicePeriod = %d MinimumBufferSize = %d"), DefaultDevicePeriod, MinBufferSize);
-			return false;
-		}
-		
 		// Determine buffer size to use. 
 		uint32 BufferFramesToRequest = FMath::Max<uint32>(InParams.NumFrames, DefaultDevicePeriod);
 		
@@ -106,6 +98,8 @@ namespace Audio
 		}
 		REFERENCE_TIME DesiredBufferDuration = FWasapiAudioUtils::FramesToRefTime(BufferFramesToRequest, InParams.SampleRate);
 		
+		UE_LOG(LogAudioMixerWasapi, Log, TEXT("FAudioMixerWasapiRenderStream::InitializeHardware DefaultDevicePeriod: %d requesting buffer size: %d"), DefaultDevicePeriod, BufferFramesToRequest);
+
 		// For shared mode, this is required to be zero
 		constexpr REFERENCE_TIME Periodicity = 0;
 
@@ -144,7 +138,7 @@ namespace Audio
 
 		bIsInitialized = true;
 
-		UE_LOG(LogAudioMixerWasapi, Verbose, TEXT("FAudioMixerWasapiRenderStream::InitializeHardware succeeded with sample rate: %d, buffer period: %d"), InParams.SampleRate, InParams.NumFrames);
+		UE_LOG(LogAudioMixerWasapi, Display, TEXT("FAudioMixerWasapiRenderStream::InitializeHardware succeeded with sample rate: %d, buffer period: %d"), InParams.SampleRate, InParams.NumFrames);
 
 		return true;
 	}
diff --git a/Engine/Source/Runtime/Windows/AudioMixerWasapi/Private/AudioMixerWasapiRenderStream.h b/Engine/Source/Runtime/Windows/AudioMixerWasapi/Private/AudioMixerWasapiRenderStream.h
--- a/Engine/Source/Runtime/Windows/AudioMixerWasapi/Private/AudioMixerWasapiRenderStream.h
+++ b/Engine/Source/Runtime/Windows/AudioMixerWasapi/Private/AudioMixerWasapiRenderStream.h
@@ -44,10 +44,14 @@ namespace Audio
 		virtual void SubmitBuffer(const uint8* Buffer, const SIZE_T InNumFrames) { }
 		virtual void SubmitDirectOutBuffer(const int32 InChannelIndex, const FAlignedFloatBuffer& InBuffer) { }
 
+		UE_DEPRECATED(5.7, "GetMinimumBufferSize() is deprecated. Please use GetDefaultDevicePeriod() in the future.")
 		static uint32 GetMinimumBufferSize(const uint32 InSampleRate);
 
 	protected:
 		
+		/** Returns the default callback period for as reported by the device. */
+		uint32 GetDefaultDevicePeriod() const { return DefaultDevicePeriod; }
+
 		/** COM pointer to the WASAPI audio client object. */
 		TComPtr<IAudioClient> AudioClient;
 
diff --git a/Engine/Source/Runtime/Windows/AudioMixerWasapi/Private/WasapiAggregateRenderStream.cpp b/Engine/Source/Runtime/Windows/AudioMixerWasapi/Private/WasapiAggregateRenderStream.cpp
--- a/Engine/Source/Runtime/Windows/AudioMixerWasapi/Private/WasapiAggregateRenderStream.cpp
+++ b/Engine/Source/Runtime/Windows/AudioMixerWasapi/Private/WasapiAggregateRenderStream.cpp
@@ -26,7 +26,7 @@ namespace Audio
 			InterleaveBuffers.Reset();
 			InterleaveBuffers.SetNum(InParams.HardwareDeviceInfo.NumChannels);
 
-			const uint32 MinBufferSize = GetMinimumBufferSize(InParams.SampleRate);			
+			const uint32 MinBufferSize = GetDefaultDevicePeriod();
 			WriteNumFrames = FMath::Max(MinBufferSize, InParams.NumFrames);
 
 			for (FAlignedFloatBuffer& Buffer : InterleaveBuffers)
diff --git a/Engine/Source/Runtime/Windows/AudioMixerWasapi/Private/WasapiDefaultRenderStream.cpp b/Engine/Source/Runtime/Windows/AudioMixerWasapi/Private/WasapiDefaultRenderStream.cpp
--- a/Engine/Source/Runtime/Windows/AudioMixerWasapi/Private/WasapiDefaultRenderStream.cpp
+++ b/Engine/Source/Runtime/Windows/AudioMixerWasapi/Private/WasapiDefaultRenderStream.cpp
@@ -12,7 +12,8 @@ namespace Audio
 		if (FAudioMixerWasapiRenderStream::InitializeHardware(InParams))
 		{
 			NumPartialBuffersToWrite = 1;
-			const uint32 MinBufferSize = GetMinimumBufferSize(InParams.SampleRate);			
+			const uint32 MinBufferSize = GetDefaultDevicePeriod();
+
 			if (MinBufferSize > InParams.NumFrames)
 			{
 				NumPartialBuffersToWrite = FMath::DivideAndRoundUp(MinBufferSize, InParams.NumFrames);
