Fix order of operations in StartAudioStream to avoid race condition

#rnx

diff --git a/Engine/Source/Runtime/Windows/AudioMixerWasapi/Private/AudioMixerWasapi.cpp b/Engine/Source/Runtime/Windows/AudioMixerWasapi/Private/AudioMixerWasapi.cpp
--- a/Engine/Source/Runtime/Windows/AudioMixerWasapi/Private/AudioMixerWasapi.cpp
+++ b/Engine/Source/Runtime/Windows/AudioMixerWasapi/Private/AudioMixerWasapi.cpp
@@ -308,20 +308,11 @@ namespace Audio
 		
 		if (bIsInitialized)
 		{
-			if (DeviceManager.IsValid() && DeviceManager->IsInitialized())
-			{
-				bDidStartAudioStream = DeviceManager->StartAudioStream();
-			}
-			else
-			{
-				check(!bIsUsingNullDevice);
-				StartRunningNullDevice();
-				bDidStartAudioStream = true;
-			}
-			
 			// Can be called during device swap when AudioRenderEvent can be null
 			if (nullptr == AudioRenderEvent)
 			{
+				// Call BeginGeneratingAudio before StartAudioStream so that the output buffer is set up
+				// prior to the device thread starting.
 				// This will set AudioStreamInfo.StreamState to EAudioOutputStreamState::Running
 				BeginGeneratingAudio();
 			}
@@ -329,6 +320,17 @@ namespace Audio
 			{
 				AudioStreamInfo.StreamState = EAudioOutputStreamState::Running;
 			}
+
+			if (DeviceManager.IsValid() && DeviceManager->IsInitialized())
+			{
+				bDidStartAudioStream = DeviceManager->StartAudioStream();
+			}
+			else
+			{
+				check(!bIsUsingNullDevice);
+				StartRunningNullDevice();
+				bDidStartAudioStream = true;
+			}
 		}
 
 		return bDidStartAudioStream;
diff --git a/Engine/Source/Runtime/Windows/AudioMixerWasapi/Private/WasapiDefaultRenderStream.cpp b/Engine/Source/Runtime/Windows/AudioMixerWasapi/Private/WasapiDefaultRenderStream.cpp
--- a/Engine/Source/Runtime/Windows/AudioMixerWasapi/Private/WasapiDefaultRenderStream.cpp
+++ b/Engine/Source/Runtime/Windows/AudioMixerWasapi/Private/WasapiDefaultRenderStream.cpp
@@ -89,12 +89,15 @@ namespace Audio
 
 	void FWasapiDefaultRenderStream::SubmitBuffer(const uint8* InBuffer, const SIZE_T InNumFrames)
 	{
-		if (RenderBufferView.Num() >= InNumFrames)
+		if (InBuffer && RenderBufferView.Num() >= InNumFrames)
 		{
 			check(InNumFrames == RenderStreamParams.NumFrames);
 			const SIZE_T NumBytes = InNumFrames * AudioFormat.GetFrameSizeInBytes();
 
-			FMemory::Memcpy(RenderBufferView.GetData(), InBuffer, NumBytes);
+			if (uint8* RenderBufferPtr = RenderBufferView.GetData())
+			{
+				FMemory::Memcpy(RenderBufferPtr, InBuffer, NumBytes);
+			}
 		}
 	}
 }
