Remove invalid [[nodiscard]] from void functions
#rnx

diff --git a/Engine/Source/Runtime/D3D12RHI/Private/D3D12WorkGraph.cpp b/Engine/Source/Runtime/D3D12RHI/Private/D3D12WorkGraph.cpp
--- a/Engine/Source/Runtime/D3D12RHI/Private/D3D12WorkGraph.cpp
+++ b/Engine/Source/Runtime/D3D12RHI/Private/D3D12WorkGraph.cpp
@@ -750,7 +750,7 @@ struct FD3D12BindlessConstantSetter
 		}
 	}
 
-	[[nodiscard]] void Finalize(FAllocatedConstantBuffer& OutConstantBuffer)
+	void Finalize(FAllocatedConstantBuffer& OutConstantBuffer)
 	{
 		OutConstantBuffer.ConstantBuffer = &ConstantBuffer;
 		ConstantBuffer.Version(OutConstantBuffer.ResourceLocation, false);
diff --git a/Engine/Source/Runtime/Windows/D3D11RHI/Private/D3D11RHIPrivate.h b/Engine/Source/Runtime/Windows/D3D11RHI/Private/D3D11RHIPrivate.h
--- a/Engine/Source/Runtime/Windows/D3D11RHI/Private/D3D11RHIPrivate.h
+++ b/Engine/Source/Runtime/Windows/D3D11RHI/Private/D3D11RHIPrivate.h
@@ -454,7 +454,7 @@ public:
 #endif
 
 	[[nodiscard]] D3D11RHI_API FD3D11Buffer* BeginCreateBufferInternal(const FRHIBufferCreateDesc& CreateDesc);
-	[[nodiscard]] D3D11RHI_API void FinalizeCreateBufferInternal(FD3D11Buffer* Buffer, TConstArrayView<uint8> InitialData);
+	D3D11RHI_API void FinalizeCreateBufferInternal(FD3D11Buffer* Buffer, TConstArrayView<uint8> InitialData);
 	[[nodiscard]] D3D11RHI_API FD3D11Buffer* CreateBufferInternal(const FRHIBufferCreateDesc& CreateDesc, TConstArrayView<uint8> InitialData);
 
 	[[nodiscard]] D3D11RHI_API virtual FRHIBufferInitializer RHICreateBufferInitializer(FRHICommandListBase& RHICmdList, const FRHIBufferCreateDesc& CreateDesc) final override;
