Fix deadlock when running with D3D11 RHI
- When the FD3D11GPUFence::PollFencesUntil function is passed a target FGraphEvent, it should block until a sync containing that event has been signalled. However, the function was returning early if there was an unrelated sync in the queue, which returned S_FALSE from GetData().
- Reworked the logic in that function to ensure it never returns until the required event has been signalled.

diff --git a/Engine/Source/Runtime/Windows/D3D11RHI/Private/D3D11Commands.cpp b/Engine/Source/Runtime/Windows/D3D11RHI/Private/D3D11Commands.cpp
--- a/Engine/Source/Runtime/Windows/D3D11RHI/Private/D3D11Commands.cpp
+++ b/Engine/Source/Runtime/Windows/D3D11RHI/Private/D3D11Commands.cpp
@@ -1783,13 +1783,12 @@ void FD3D11GPUFence::PollFencesUntil(FGraphEvent* Target)
 {
 	FD3D11DynamicRHI& RHI = FD3D11DynamicRHI::Get();
 
+	// Flush the queue if we have a specific fence to wait for, otherwise the commands that signal the
+	// fence may be stuck in the driver's command queue and never reach the GPU, causing an infinite loop.
+	UINT const Flags = Target ? 0 : D3D11_ASYNC_GETDATA_DONOTFLUSH;
+
 	while (FD3D11Sync* Sync = ActiveSyncs.Peek())
 	{
-		bool const bTarget = Sync->Event == Target;
-
-		// Wait forever if this fence is the one we're looking for specifically, otherwise just poll.
-		UINT const Flags = bTarget ? 0 : D3D11_ASYNC_GETDATA_DONOTFLUSH;
-
 	Retry:
 		BOOL Value = 0;
 		HRESULT Result = RHI.GetDeviceContext()->GetData(Sync->Query, &Value, sizeof(Value), Flags);
@@ -1797,7 +1796,7 @@ void FD3D11GPUFence::PollFencesUntil(FGraphEvent* Target)
 		if (Result == S_FALSE || (Result == S_OK && !Value))
 		{
 			// Fence is not done
-			if (bTarget)
+			if (Target)
 			{
 				// We're waiting for a specific fence. Spin until it passes.
 				goto Retry;
@@ -1810,11 +1809,13 @@ void FD3D11GPUFence::PollFencesUntil(FGraphEvent* Target)
 
 		VERIFYD3D11RESULT(Result);
 
+		bool bFoundTarget = Sync->Event == Target;
+
 		// The fence has completed. Signal the graph event and remove the node.
 		Sync->Event->DispatchSubsequents();
 		ActiveSyncs.Pop();
 
-		if (bTarget)
+		if (bFoundTarget)
 		{
 			// We found the fence we wanted. Stop polling.
 			return;
