Support for mono audio output devices in WASAPI backend

diff --git a/Engine/Source/Runtime/Windows/AudioMixerWasapi/Private/AudioMixerWasapiRenderStream.cpp b/Engine/Source/Runtime/Windows/AudioMixerWasapi/Private/AudioMixerWasapiRenderStream.cpp
--- a/Engine/Source/Runtime/Windows/AudioMixerWasapi/Private/AudioMixerWasapiRenderStream.cpp
+++ b/Engine/Source/Runtime/Windows/AudioMixerWasapi/Private/AudioMixerWasapiRenderStream.cpp
@@ -132,6 +132,11 @@ namespace Audio
 			return false;
 		}
 
+		if (InParams.HardwareDeviceInfo.NumChannels != StreamFormat.GetNumChannels())
+		{
+			InitializeDownmixBuffers(InParams.HardwareDeviceInfo.NumChannels, InParams.NumFrames, StreamFormat);
+		}
+		
 		AudioClient = MoveTemp(TempAudioClient);
 		AudioFormat = StreamFormat;
 		RenderStreamParams = InParams;
@@ -166,6 +171,24 @@ namespace Audio
 		return bIsInitialized;
 	}
 
+	void FAudioMixerWasapiRenderStream::InitializeDownmixBuffers(const int32 InNumInputChannels, const int32 InNumFrames, const FWasapiAudioFormat& InFormat)
+	{
+		// The downmixing code assumes float audio buffers
+		check(InFormat.GetEncoding() == EWasapiAudioEncoding::FLOATING_POINT_32);
+
+		MonoDownmixBuffer.Reset();
+		MonoDownmixBuffer.AddUninitialized(InNumFrames);
+
+		MixdownGainsMap.Reset();
+		const float MonoDownmixGain = 1.0f / InNumInputChannels;
+
+		// Gain values used when downmixing from input channel count to mono
+		for (int32 Index = 0; Index < InNumInputChannels; ++Index)
+		{
+			MixdownGainsMap.Add(MonoDownmixGain);
+		}
+	}
+
 	int32 FAudioMixerWasapiRenderStream::GetNumFrames(const int32 InNumRequestedFrames) const
 	{
 		return InNumRequestedFrames;
diff --git a/Engine/Source/Runtime/Windows/AudioMixerWasapi/Private/AudioMixerWasapiRenderStream.h b/Engine/Source/Runtime/Windows/AudioMixerWasapi/Private/AudioMixerWasapiRenderStream.h
--- a/Engine/Source/Runtime/Windows/AudioMixerWasapi/Private/AudioMixerWasapiRenderStream.h
+++ b/Engine/Source/Runtime/Windows/AudioMixerWasapi/Private/AudioMixerWasapiRenderStream.h
@@ -79,5 +79,16 @@ namespace Audio
 
 		/** Accumulates errors that occur in the audio callback. */
 		uint32 CallbackBufferErrors = 0;
+
+		/** Scratch buffer used for downmixing output for mono devices. */
+		FAlignedFloatBuffer MonoDownmixBuffer;
+		
+		/** Per-channel gains used when downmixing. */
+		TArray<float> MixdownGainsMap;
+
+	private:
+		
+		/** Initializes the necessary data structures needed for downmixing audio buffers. */
+		void InitializeDownmixBuffers(const int32 InNumInputChannels, const int32 InNumFrames, const FWasapiAudioFormat& InFormat);
 	};
 }
diff --git a/Engine/Source/Runtime/Windows/AudioMixerWasapi/Private/WasapiDefaultRenderStream.cpp b/Engine/Source/Runtime/Windows/AudioMixerWasapi/Private/WasapiDefaultRenderStream.cpp
--- a/Engine/Source/Runtime/Windows/AudioMixerWasapi/Private/WasapiDefaultRenderStream.cpp
+++ b/Engine/Source/Runtime/Windows/AudioMixerWasapi/Private/WasapiDefaultRenderStream.cpp
@@ -96,7 +96,19 @@ namespace Audio
 
 			if (uint8* RenderBufferPtr = RenderBufferView.GetData())
 			{
-				FMemory::Memcpy(RenderBufferPtr, InBuffer, NumBytes);
+				const int32 InputNumChannels = RenderStreamParams.HardwareDeviceInfo.NumChannels;
+				const int32 OutputNumChannels = AudioFormat.GetNumChannels();
+				
+				if (OutputNumChannels == 1)
+				{
+					DownmixBuffer(InputNumChannels, 1, reinterpret_cast<const float*>(InBuffer), MonoDownmixBuffer.GetData(), InNumFrames, MixdownGainsMap.GetData());
+					FMemory::Memcpy(RenderBufferPtr, MonoDownmixBuffer.GetData(), NumBytes);
+				}
+				else
+				{
+					check(OutputNumChannels == InputNumChannels);
+					FMemory::Memcpy(RenderBufferPtr, InBuffer, NumBytes);
+				}
 			}
 		}
 	}
