cleanup and add more logging

#rnx

diff --git a/Engine/Source/Runtime/Windows/AudioMixerWasapi/Private/AudioMixerWasapiRenderStream.cpp b/Engine/Source/Runtime/Windows/AudioMixerWasapi/Private/AudioMixerWasapiRenderStream.cpp
--- a/Engine/Source/Runtime/Windows/AudioMixerWasapi/Private/AudioMixerWasapiRenderStream.cpp
+++ b/Engine/Source/Runtime/Windows/AudioMixerWasapi/Private/AudioMixerWasapiRenderStream.cpp
@@ -132,9 +132,16 @@ namespace Audio
 			return false;
 		}
 
-		if (InParams.HardwareDeviceInfo.NumChannels != StreamFormat.GetNumChannels())
+		// HardwareDeviceInfo is somewhat of a misnomer, HardwareDeviceInfo.NumChannels is the number of channels
+		// that the AudioMixer is requesting. StreamFormat is fetched from the hardware and represents the 
+		// actual number of output channels on the hardware device.
+		const int32 AudioMixerNumChannels = InParams.HardwareDeviceInfo.NumChannels;
+		const int32 RenderStreamHardwareChannels = StreamFormat.GetNumChannels();
+		UE_LOG(LogAudioMixerWasapi, Display, TEXT("FAudioMixerWasapiRenderStream::InitializeHardware audio mixer requested num channels: %d, Wasapi render stream hardware supports num channels: %d"), AudioMixerNumChannels, RenderStreamHardwareChannels);
+
+		if (AudioMixerNumChannels != RenderStreamHardwareChannels)
 		{
-			InitializeDownmixBuffers(InParams.HardwareDeviceInfo.NumChannels, InParams.NumFrames, StreamFormat);
+			InitializeDownmixBuffers(AudioMixerNumChannels, InParams.NumFrames, StreamFormat);
 		}
 		
 		AudioClient = MoveTemp(TempAudioClient);
@@ -173,6 +180,8 @@ namespace Audio
 
 	void FAudioMixerWasapiRenderStream::InitializeDownmixBuffers(const int32 InNumInputChannels, const int32 InNumFrames, const FWasapiAudioFormat& InFormat)
 	{
+		UE_LOG(LogAudioMixerWasapi, Display, TEXT("FAudioMixerWasapiRenderStream: Initializing mono downmix buffer"));
+
 		// The downmixing code assumes float audio buffers
 		check(InFormat.GetEncoding() == EWasapiAudioEncoding::FLOATING_POINT_32);
 
diff --git a/Engine/Source/Runtime/Windows/AudioMixerWasapi/Private/WasapiDefaultRenderStream.cpp b/Engine/Source/Runtime/Windows/AudioMixerWasapi/Private/WasapiDefaultRenderStream.cpp
--- a/Engine/Source/Runtime/Windows/AudioMixerWasapi/Private/WasapiDefaultRenderStream.cpp
+++ b/Engine/Source/Runtime/Windows/AudioMixerWasapi/Private/WasapiDefaultRenderStream.cpp
@@ -91,11 +91,14 @@ namespace Audio
 	{
 		if (InBuffer && RenderBufferView.Num() >= InNumFrames)
 		{
-			check(InNumFrames == RenderStreamParams.NumFrames);
+			checkf(InNumFrames == RenderStreamParams.NumFrames, TEXT("InNumFrames = %d, RenderStreamParams.NumFrames = %d"), InNumFrames, RenderStreamParams.NumFrames);
 			const SIZE_T NumBytes = InNumFrames * AudioFormat.GetFrameSizeInBytes();
 
 			if (uint8* RenderBufferPtr = RenderBufferView.GetData())
 			{
+				// HardwareDeviceInfo is somewhat of a misnomer, HardwareDeviceInfo.NumChannels is the number of channels
+				// that the AudioMixer is requesting. AudioFormat is fetched from the hardware and represents the 
+				// actual number of output channels on the hardware device.
 				const int32 InputNumChannels = RenderStreamParams.HardwareDeviceInfo.NumChannels;
 				const int32 OutputNumChannels = AudioFormat.GetNumChannels();
 				
@@ -106,7 +109,7 @@ namespace Audio
 				}
 				else
 				{
-					check(OutputNumChannels == InputNumChannels);
+					checkf(OutputNumChannels == InputNumChannels, TEXT("OutputNumChannels = %d, InputNumChannels = %d"), OutputNumChannels, InputNumChannels);
 					FMemory::Memcpy(RenderBufferPtr, InBuffer, NumBytes);
 				}
 			}
