Fixing force-surround with wasapi backends.

Allows downmixing (or upmixing) to the hardware output channels if the audio mixer is not at the same channel count (e.g. is rendering to stereo but output is mono, or possibly forced to act as if it's in surround.

#rnx

diff --git a/Engine/Source/Runtime/Windows/AudioMixerWasapi/Private/AudioMixerWasapiRenderStream.cpp b/Engine/Source/Runtime/Windows/AudioMixerWasapi/Private/AudioMixerWasapiRenderStream.cpp
--- a/Engine/Source/Runtime/Windows/AudioMixerWasapi/Private/AudioMixerWasapiRenderStream.cpp
+++ b/Engine/Source/Runtime/Windows/AudioMixerWasapi/Private/AudioMixerWasapiRenderStream.cpp
@@ -185,16 +185,17 @@ namespace Audio
 		// The downmixing code assumes float audio buffers
 		check(InFormat.GetEncoding() == EWasapiAudioEncoding::FLOATING_POINT_32);
 
-		MonoDownmixBuffer.Reset();
-		MonoDownmixBuffer.AddUninitialized(InNumFrames);
+		DownmixScratchBuffer.Reset();
+		DownmixScratchBuffer.AddUninitialized(InNumFrames);
 
 		MixdownGainsMap.Reset();
-		const float MonoDownmixGain = 1.0f / InNumInputChannels;
+
+		const float DownmixGain = 1.0f / InNumInputChannels;
 
 		// Gain values used when downmixing from input channel count to mono
 		for (int32 Index = 0; Index < InNumInputChannels; ++Index)
 		{
-			MixdownGainsMap.Add(MonoDownmixGain);
+			MixdownGainsMap.Add(DownmixGain);
 		}
 	}
 
diff --git a/Engine/Source/Runtime/Windows/AudioMixerWasapi/Private/AudioMixerWasapiRenderStream.h b/Engine/Source/Runtime/Windows/AudioMixerWasapi/Private/AudioMixerWasapiRenderStream.h
--- a/Engine/Source/Runtime/Windows/AudioMixerWasapi/Private/AudioMixerWasapiRenderStream.h
+++ b/Engine/Source/Runtime/Windows/AudioMixerWasapi/Private/AudioMixerWasapiRenderStream.h
@@ -80,8 +80,8 @@ namespace Audio
 		/** Accumulates errors that occur in the audio callback. */
 		uint32 CallbackBufferErrors = 0;
 
-		/** Scratch buffer used for downmixing output for mono devices. */
-		FAlignedFloatBuffer MonoDownmixBuffer;
+		/** Scratch buffer used for downmixing output. */
+		FAlignedFloatBuffer DownmixScratchBuffer;
 		
 		/** Per-channel gains used when downmixing. */
 		TArray<float> MixdownGainsMap;
diff --git a/Engine/Source/Runtime/Windows/AudioMixerWasapi/Private/WasapiDefaultRenderStream.cpp b/Engine/Source/Runtime/Windows/AudioMixerWasapi/Private/WasapiDefaultRenderStream.cpp
--- a/Engine/Source/Runtime/Windows/AudioMixerWasapi/Private/WasapiDefaultRenderStream.cpp
+++ b/Engine/Source/Runtime/Windows/AudioMixerWasapi/Private/WasapiDefaultRenderStream.cpp
@@ -102,14 +102,13 @@ namespace Audio
 				const int32 InputNumChannels = RenderStreamParams.HardwareDeviceInfo.NumChannels;
 				const int32 OutputNumChannels = AudioFormat.GetNumChannels();
 				
-				if (OutputNumChannels == 1)
+				if (OutputNumChannels != InputNumChannels)
 				{
-					DownmixBuffer(InputNumChannels, 1, reinterpret_cast<const float*>(InBuffer), MonoDownmixBuffer.GetData(), InNumFrames, MixdownGainsMap.GetData());
-					FMemory::Memcpy(RenderBufferPtr, MonoDownmixBuffer.GetData(), NumBytes);
+					DownmixBuffer(InputNumChannels, OutputNumChannels, reinterpret_cast<const float*>(InBuffer), DownmixScratchBuffer.GetData(), InNumFrames, MixdownGainsMap.GetData());
+					FMemory::Memcpy(RenderBufferPtr, DownmixScratchBuffer.GetData(), NumBytes);
 				}
 				else
 				{
-					checkf(OutputNumChannels == InputNumChannels, TEXT("OutputNumChannels = %d, InputNumChannels = %d"), OutputNumChannels, InputNumChannels);
 					FMemory::Memcpy(RenderBufferPtr, InBuffer, NumBytes);
 				}
 			}
