Take into account output channel count when allocating downmix scratch buffer.

#rnx

diff --git a/Engine/Source/Runtime/Windows/AudioMixerWasapi/Private/AudioMixerWasapiRenderStream.cpp b/Engine/Source/Runtime/Windows/AudioMixerWasapi/Private/AudioMixerWasapiRenderStream.cpp
--- a/Engine/Source/Runtime/Windows/AudioMixerWasapi/Private/AudioMixerWasapiRenderStream.cpp
+++ b/Engine/Source/Runtime/Windows/AudioMixerWasapi/Private/AudioMixerWasapiRenderStream.cpp
@@ -178,15 +178,15 @@ namespace Audio
 		return bIsInitialized;
 	}
 
-	void FAudioMixerWasapiRenderStream::InitializeDownmixBuffers(const int32 InNumInputChannels, const int32 InNumFrames, const FWasapiAudioFormat& InFormat)
+	void FAudioMixerWasapiRenderStream::InitializeDownmixBuffers(const int32 InNumInputChannels, const int32 InNumFrames, const FWasapiAudioFormat& InAudioOutputFormat)
 	{
 		UE_LOG(LogAudioMixerWasapi, Display, TEXT("FAudioMixerWasapiRenderStream: Initializing mono downmix buffer"));
 
 		// The downmixing code assumes float audio buffers
-		check(InFormat.GetEncoding() == EWasapiAudioEncoding::FLOATING_POINT_32);
+		check(InAudioOutputFormat.GetEncoding() == EWasapiAudioEncoding::FLOATING_POINT_32);
 
 		DownmixScratchBuffer.Reset();
-		DownmixScratchBuffer.AddUninitialized(InNumFrames);
+		DownmixScratchBuffer.AddUninitialized(InNumFrames * InAudioOutputFormat.GetNumChannels());
 
 		MixdownGainsMap.Reset();
 
diff --git a/Engine/Source/Runtime/Windows/AudioMixerWasapi/Private/AudioMixerWasapiRenderStream.h b/Engine/Source/Runtime/Windows/AudioMixerWasapi/Private/AudioMixerWasapiRenderStream.h
--- a/Engine/Source/Runtime/Windows/AudioMixerWasapi/Private/AudioMixerWasapiRenderStream.h
+++ b/Engine/Source/Runtime/Windows/AudioMixerWasapi/Private/AudioMixerWasapiRenderStream.h
@@ -89,6 +89,6 @@ namespace Audio
 	private:
 		
 		/** Initializes the necessary data structures needed for downmixing audio buffers. */
-		void InitializeDownmixBuffers(const int32 InNumInputChannels, const int32 InNumFrames, const FWasapiAudioFormat& InFormat);
+		void InitializeDownmixBuffers(const int32 InNumInputChannels, const int32 InNumFrames, const FWasapiAudioFormat& InAudioOutputFormat);
 	};
 }
