Enforce memory ordering in render stream state machine.

#rnx

diff --git a/Engine/Source/Runtime/Windows/AudioMixerWasapi/Private/AudioMixerWasapiRenderStream.cpp b/Engine/Source/Runtime/Windows/AudioMixerWasapi/Private/AudioMixerWasapiRenderStream.cpp
--- a/Engine/Source/Runtime/Windows/AudioMixerWasapi/Private/AudioMixerWasapiRenderStream.cpp
+++ b/Engine/Source/Runtime/Windows/AudioMixerWasapi/Private/AudioMixerWasapiRenderStream.cpp
@@ -148,7 +148,7 @@ namespace Audio
 		AudioFormat = StreamFormat;
 		RenderStreamParams = InParams;
 
-		bIsInitialized = true;
+		bIsInitialized.store(true, std::memory_order_release);
 
 		UE_LOG(LogAudioMixerWasapi, Display, TEXT("FAudioMixerWasapiRenderStream::InitializeHardware succeeded with sample rate: %d, buffer period: %d"), InParams.SampleRate, InParams.NumFrames);
 
@@ -157,7 +157,7 @@ namespace Audio
 
 	bool FAudioMixerWasapiRenderStream::TeardownHardware()
 	{
-		if (!bIsInitialized)
+		if (!bIsInitialized.load(std::memory_order_acquire))
 		{
 			UE_LOG(LogAudioMixerWasapi, Warning, TEXT("FAudioMixerWasapiRenderStream::TeardownHardware failed...not initialized. "));
 			return false;
@@ -166,7 +166,7 @@ namespace Audio
 		RenderClient.Reset();
 		AudioClient.Reset();
 
-		bIsInitialized = false;
+		bIsInitialized.store(false, std::memory_order_release);
 
 		UE_LOG(LogAudioMixerWasapi, Verbose, TEXT("FAudioMixerWasapiRenderStream::TeardownHardware succeeded"));
 
@@ -175,7 +175,7 @@ namespace Audio
 
 	bool FAudioMixerWasapiRenderStream::IsInitialized() const
 	{
-		return bIsInitialized;
+		return bIsInitialized.load(std::memory_order_acquire);
 	}
 
 	void FAudioMixerWasapiRenderStream::InitializeDownmixBuffers(const int32 InNumInputChannels, const int32 InNumFrames, const FWasapiAudioFormat& InAudioOutputFormat)
@@ -206,25 +206,32 @@ namespace Audio
 
 	bool FAudioMixerWasapiRenderStream::OpenAudioStream(const FWasapiRenderStreamParams& InParams, HANDLE InEventHandle)
 	{
+		EAudioOutputStreamState::Type CurrStreamState = StreamState.load(std::memory_order_acquire);
+		if (CurrStreamState != EAudioOutputStreamState::Closed)
+		{
+			UE_LOG(LogAudioMixerWasapi, Warning, TEXT("FAudioMixerWasapiRenderStream::OpenAudioStream stream not in closed sate. CurrStreamState: %d"), CurrStreamState);
+			return false;
+		}
+		
 		if (InParams.HardwareDeviceInfo.DeviceId != RenderStreamParams.HardwareDeviceInfo.DeviceId)
 		{
 			if (!InitializeHardware(InParams))
 			{
-				UE_LOG(LogAudioMixerWasapi, Error, TEXT("OpenAudioStream failed InitAudioClient"));
+				UE_LOG(LogAudioMixerWasapi, Warning, TEXT("OpenAudioStream failed InitAudioClient"));
 				return false;
 			}
 		}
 
 		if (InEventHandle == nullptr)
 		{
-			UE_LOG(LogAudioMixerWasapi, Error, TEXT("OpenAudioStream null EventHandle"));
+			UE_LOG(LogAudioMixerWasapi, Warning, TEXT("OpenAudioStream null EventHandle"));
 			return false;
 		}
 
 		HRESULT Result = AudioClient->SetEventHandle(InEventHandle);
 		if (FAILED(Result))
 		{
-			UE_LOG(LogAudioMixerWasapi, Error, TEXT("OpenAudioStream failed IAudioClient::SetEventHandle %s"), *AudioClientErrorToFString(Result));
+			UE_LOG(LogAudioMixerWasapi, Warning, TEXT("OpenAudioStream failed IAudioClient::SetEventHandle %s"), *AudioClientErrorToFString(Result));
 			return false;
 		}
 
@@ -232,13 +239,21 @@ namespace Audio
 		Result = AudioClient->GetService(__uuidof(IAudioRenderClient), IID_PPV_ARGS_Helper(&TempRenderClient));
 		if (FAILED(Result))
 		{
-			UE_LOG(LogAudioMixerWasapi, Error, TEXT("OpenAudioStream failed IAudioClient::GetService IAudioRenderClient %s"), *AudioClientErrorToFString(Result));
+			UE_LOG(LogAudioMixerWasapi, Warning, TEXT("OpenAudioStream failed IAudioClient::GetService IAudioRenderClient %s"), *AudioClientErrorToFString(Result));
 			return false;
 		}
 
 		RenderClient = MoveTemp(TempRenderClient);
-		bIsInitialized = true;
 
+		// It's the caller's responsibility to maintain thread safety. This class uses compare_exchange() calls to detect any state errors and 
+		// bubble them up the chain.
+		const EAudioOutputStreamState::Type OriginalState = CurrStreamState;
+		if (!ensure(StreamState.compare_exchange_strong(CurrStreamState, EAudioOutputStreamState::Open, std::memory_order_release, std::memory_order_relaxed)))
+		{
+			UE_LOG(LogAudioMixerWasapi, Warning, TEXT("FAudioMixerWasapiRenderStream::OpenAudioStream StreamState modified during open. OriginalState: %d CurrStreamState: %d"), OriginalState, CurrStreamState);
+			return false;
+		}
+		
 		UE_LOG(LogAudioMixerWasapi, Verbose, TEXT("FAudioMixerWasapiRenderStream::OpenAudioStream succeeded with SampeRate: %d, NumFrames: %d"), InParams.SampleRate, InParams.NumFrames);
 
 		return true;
@@ -246,61 +261,81 @@ namespace Audio
 
 	bool FAudioMixerWasapiRenderStream::CloseAudioStream()
 	{
-		if (!bIsInitialized || StreamState == EAudioOutputStreamState::Closed)
+		EAudioOutputStreamState::Type CurrStreamState = StreamState.load(std::memory_order_acquire);
+		if (!bIsInitialized.load(std::memory_order_acquire) || CurrStreamState == EAudioOutputStreamState::Closed)
 		{
-			UE_LOG(LogAudioMixerWasapi, Verbose, TEXT("FAudioMixerWasapiRenderStream::CloseAudioStream stream appears to be already closed"));
-
+			UE_LOG(LogAudioMixerWasapi, Warning, TEXT("FAudioMixerWasapiRenderStream::CloseAudioStream stream in unexpected state. bIsInitialized: %d CurrStreamState: %d"),
+				bIsInitialized.load(std::memory_order_relaxed), CurrStreamState);
 			return false;
 		}
 
-		if (StreamState == EAudioOutputStreamState::Running)
+		if (CurrStreamState == EAudioOutputStreamState::Running)
 		{
-			UE_LOG(LogAudioMixerWasapi, Verbose, TEXT("FAudioMixerWasapiRenderStream::CloseAudioStream stream appears to be running. StopAudioStream() must be called prior to closing."));
-
+			UE_LOG(LogAudioMixerWasapi, Warning, TEXT("FAudioMixerWasapiRenderStream::CloseAudioStream stream appears to be running. StopAudioStream() must be called prior to closing."));
 			return false;
 		}
 
-		StreamState = EAudioOutputStreamState::Closed;
+		const EAudioOutputStreamState::Type OriginalState = CurrStreamState;
+		if (!ensure(StreamState.compare_exchange_strong(CurrStreamState, EAudioOutputStreamState::Closed, std::memory_order_release, std::memory_order_relaxed)))
+		{
+			UE_LOG(LogAudioMixerWasapi, Warning, TEXT("FAudioMixerWasapiRenderStream::CloseAudioStream StreamState modified during close. OriginalState: %d CurrStreamState: %d"), OriginalState, CurrStreamState);
+			return false;
+		}
 
 		return true;
 	}
 
 	bool FAudioMixerWasapiRenderStream::StartAudioStream()
 	{
-		if (bIsInitialized)
+		EAudioOutputStreamState::Type CurrStreamState = StreamState.load(std::memory_order_acquire);
+		if (CurrStreamState != EAudioOutputStreamState::Open)
 		{
-			StreamState = EAudioOutputStreamState::Running;
-
-			if (!AudioClient.IsValid())
-			{
-				UE_LOG(LogAudioMixerWasapi, Error, TEXT("StartAudioStream failed invalid audio client"));
-				return false;
-			}
+			UE_LOG(LogAudioMixerWasapi, Warning, TEXT("FAudioMixerWasapiRenderStream::StartAudioStream stream not in open sate. bIsInitialized: %d, CurrStreamState: %d"),
+				bIsInitialized.load(std::memory_order_relaxed), CurrStreamState);
+			return false;
+		}
 
-			AudioClient->Start();
+		if (!AudioClient.IsValid())
+		{
+			UE_LOG(LogAudioMixerWasapi, Warning, TEXT("StartAudioStream failed invalid audio client"));
+			return false;
 		}
 
+		AudioClient->Start();
 		UE_LOG(LogAudioMixerWasapi, Verbose, TEXT("FAudioMixerWasapiRenderStream::StartAudioStream stream started"));
 
+		const EAudioOutputStreamState::Type OriginalState = CurrStreamState;
+		if (!ensure(StreamState.compare_exchange_strong(CurrStreamState, EAudioOutputStreamState::Running, std::memory_order_release, std::memory_order_relaxed)))
+		{
+			UE_LOG(LogAudioMixerWasapi, Warning, TEXT("FAudioMixerWasapiRenderStream::StartAudioStream StreamState modified during start. OriginalState: %d CurrStreamState: %d"), OriginalState, CurrStreamState);
+			return false;
+		}
+
 		return true;
 	}
 
 	bool FAudioMixerWasapiRenderStream::StopAudioStream()
 	{
-		if (!bIsInitialized)
+		if (!bIsInitialized.load(std::memory_order_acquire))
 		{
-			UE_LOG(LogAudioMixerWasapi, Error, TEXT("FAudioMixerWasapiRenderStream::StopAudioStream() not initialized"));
+			UE_LOG(LogAudioMixerWasapi, Warning, TEXT("FAudioMixerWasapiRenderStream::StopAudioStream() not initialized"));
 			return false;
 		}
 
-		if (StreamState != EAudioOutputStreamState::Stopped && StreamState != EAudioOutputStreamState::Closed)
+		EAudioOutputStreamState::Type CurrStreamState = StreamState.load(std::memory_order_acquire);
+		if (CurrStreamState != EAudioOutputStreamState::Stopped && CurrStreamState != EAudioOutputStreamState::Closed)
 		{
 			if (AudioClient.IsValid())
 			{
 				AudioClient->Stop();
 			}
-
-			StreamState = EAudioOutputStreamState::Stopped;
+			
+			const EAudioOutputStreamState::Type OriginalState = CurrStreamState;
+			if (!ensure(StreamState.compare_exchange_strong(CurrStreamState, EAudioOutputStreamState::Stopped, std::memory_order_release, std::memory_order_relaxed)))
+			{
+				UE_LOG(LogAudioMixerWasapi, Warning, TEXT("FAudioMixerWasapiRenderStream::StopAudioStream StreamState modified during stop. OriginalState: %d CurrStreamState: %d"), OriginalState, CurrStreamState);
+				return false;
+			}
 		}
 
 		if (CallbackBufferErrors > 0)
diff --git a/Engine/Source/Runtime/Windows/AudioMixerWasapi/Private/AudioMixerWasapiRenderStream.h b/Engine/Source/Runtime/Windows/AudioMixerWasapi/Private/AudioMixerWasapiRenderStream.h
--- a/Engine/Source/Runtime/Windows/AudioMixerWasapi/Private/AudioMixerWasapiRenderStream.h
+++ b/Engine/Source/Runtime/Windows/AudioMixerWasapi/Private/AudioMixerWasapiRenderStream.h
@@ -6,6 +6,8 @@
 #include "IAudioMixerWasapiDeviceManager.h"
 #include "WasapiAudioFormat.h"
 
+#include <atomic>
+
 #include "Windows/AllowWindowsPlatformTypes.h"
 #include "Windows/AllowWindowsPlatformAtomics.h"
 
@@ -62,10 +64,10 @@ namespace Audio
 		FWasapiAudioFormat AudioFormat;
 
 		/** Indicates if this object has been successfully initialized. */
-		bool bIsInitialized = false;
+		std::atomic<bool> bIsInitialized = false;
 
 		/** The state of the output audio stream. */
-		EAudioOutputStreamState::Type StreamState = EAudioOutputStreamState::Closed;
+		std::atomic<EAudioOutputStreamState::Type> StreamState = EAudioOutputStreamState::Closed;
 
 		/** Render output device info. */
 		FWasapiRenderStreamParams RenderStreamParams;
diff --git a/Engine/Source/Runtime/Windows/AudioMixerWasapi/Private/WasapiAggregateRenderStream.cpp b/Engine/Source/Runtime/Windows/AudioMixerWasapi/Private/WasapiAggregateRenderStream.cpp
--- a/Engine/Source/Runtime/Windows/AudioMixerWasapi/Private/WasapiAggregateRenderStream.cpp
+++ b/Engine/Source/Runtime/Windows/AudioMixerWasapi/Private/WasapiAggregateRenderStream.cpp
@@ -64,7 +64,8 @@ namespace Audio
 	{
 		SCOPED_NAMED_EVENT(FWasapiAggregateRenderStream_DeviceRenderCallback, FColor::Blue);
 
-		if (bIsInitialized)
+		const EAudioOutputStreamState::Type CurrStreamState = StreamState.load(std::memory_order_acquire);
+		if (CurrStreamState == EAudioOutputStreamState::Running)
 		{
 			uint32 NumFramesPadding = 0;
 			AudioClient->GetCurrentPadding(&NumFramesPadding);
diff --git a/Engine/Source/Runtime/Windows/AudioMixerWasapi/Private/WasapiDefaultRenderStream.cpp b/Engine/Source/Runtime/Windows/AudioMixerWasapi/Private/WasapiDefaultRenderStream.cpp
--- a/Engine/Source/Runtime/Windows/AudioMixerWasapi/Private/WasapiDefaultRenderStream.cpp
+++ b/Engine/Source/Runtime/Windows/AudioMixerWasapi/Private/WasapiDefaultRenderStream.cpp
@@ -40,7 +40,8 @@ namespace Audio
 	{
 		SCOPED_NAMED_EVENT(FWasapiDefaultRenderStream_DeviceRenderCallback, FColor::Blue);
 
-		if (bIsInitialized)
+		const EAudioOutputStreamState::Type CurrStreamState = StreamState.load(std::memory_order_acquire);
+		if (CurrStreamState == EAudioOutputStreamState::Running)
 		{
 			uint32 NumFramesPadding = 0;
 			AudioClient->GetCurrentPadding(&NumFramesPadding);
@@ -89,7 +90,8 @@ namespace Audio
 
 	void FWasapiDefaultRenderStream::SubmitBuffer(const uint8* InBuffer, const SIZE_T InNumFrames)
 	{
-		if (InBuffer && RenderBufferView.Num() >= InNumFrames)
+		const EAudioOutputStreamState::Type CurrStreamState = StreamState.load(std::memory_order_acquire);
+		if (CurrStreamState == EAudioOutputStreamState::Running && InBuffer && RenderBufferView.Num() >= InNumFrames)
 		{
 			checkf(InNumFrames == RenderStreamParams.NumFrames, TEXT("InNumFrames = %d, RenderStreamParams.NumFrames = %d"), InNumFrames, RenderStreamParams.NumFrames);
 			const SIZE_T NumBytes = InNumFrames * AudioFormat.GetFrameSizeInBytes();
