Support for arbitrary audio device buffer sizes.

#rnx

diff --git a/Engine/Source/Runtime/Windows/AudioMixerWasapi/Private/AudioMixerWasapiRenderStream.cpp b/Engine/Source/Runtime/Windows/AudioMixerWasapi/Private/AudioMixerWasapiRenderStream.cpp
--- a/Engine/Source/Runtime/Windows/AudioMixerWasapi/Private/AudioMixerWasapiRenderStream.cpp
+++ b/Engine/Source/Runtime/Windows/AudioMixerWasapi/Private/AudioMixerWasapiRenderStream.cpp
@@ -86,11 +86,34 @@ namespace Audio
 		}
 
 		// Determine buffer size to use. 
-		uint32 BufferFramesToRequest = FMath::Max<uint32>(InParams.NumFrames, DefaultDevicePeriod);
+		bool bRoundUpFramesToRequest = false;
+		uint32 BufferFramesToRequest = InParams.NumFrames;
+
+		if (BufferFramesToRequest < DefaultDevicePeriod)
+		{
+			if (DefaultDevicePeriod % BufferFramesToRequest == 0)
+			{
+				// For example, DefaultDevicePeriod = 480, BufferFramesToRequest = 240
+				// no need to round up, just use the default device period
+				BufferFramesToRequest = DefaultDevicePeriod;
+			}
+			else
+			{
+				// For example, DefaultDevicePeriod = 512, BufferFramesToRequest = 240
+				// non-integral buffer size requested, round up to ensure there are always enough frames buffered
+				bRoundUpFramesToRequest = true;
+			}
+		}
+		else if (BufferFramesToRequest > DefaultDevicePeriod && (BufferFramesToRequest % DefaultDevicePeriod != 0))
+		{
+			// For example, DefaultDevicePeriod = 480, BufferFramesToRequest = 1024
+			// non-integral buffer size requested, round up to ensure there are always enough frames buffered
+			bRoundUpFramesToRequest = true;
+		}
 		
 		// If the engine buffer size is not an integral multiple of the device period then we must
 		// account for buffer phasing by padding the requested buffer size.
-		if (BufferFramesToRequest % DefaultDevicePeriod != 0)
+		if (bRoundUpFramesToRequest)
 		{
 			// Round up to nearest multiple of the device period
 			uint32 Multiple = FMath::CeilToInt32(static_cast<float>(BufferFramesToRequest + DefaultDevicePeriod - 1) / DefaultDevicePeriod);
diff --git a/Engine/Source/Runtime/Windows/AudioMixerWasapi/Private/WasapiAggregateRenderStream.cpp b/Engine/Source/Runtime/Windows/AudioMixerWasapi/Private/WasapiAggregateRenderStream.cpp
--- a/Engine/Source/Runtime/Windows/AudioMixerWasapi/Private/WasapiAggregateRenderStream.cpp
+++ b/Engine/Source/Runtime/Windows/AudioMixerWasapi/Private/WasapiAggregateRenderStream.cpp
@@ -26,8 +26,8 @@ namespace Audio
 			InterleaveBuffers.Reset();
 			InterleaveBuffers.SetNum(InParams.HardwareDeviceInfo.NumChannels);
 
-			const uint32 MinBufferSize = GetDefaultDevicePeriod();
-			WriteNumFrames = FMath::Max(MinBufferSize, InParams.NumFrames);
+			const uint32 MinBufferSize = NumFramesPerDeviceBuffer;
+			const uint32 WriteNumFrames = FMath::Max(MinBufferSize, InParams.NumFrames);
 
 			for (FAlignedFloatBuffer& Buffer : InterleaveBuffers)
 			{
@@ -64,8 +64,10 @@ namespace Audio
 	{
 		SCOPED_NAMED_EVENT(FWasapiAggregateRenderStream_DeviceRenderCallback, FColor::Blue);
 
+		const int32 NumDirectOutChannels = DirectOutBuffers.Num();
 		const EAudioOutputStreamState::Type CurrStreamState = StreamState.load(std::memory_order_acquire);
-		if (CurrStreamState == EAudioOutputStreamState::Running)
+		
+		if (CurrStreamState == EAudioOutputStreamState::Running && NumDirectOutChannels > 0)
 		{
 			uint32 NumFramesPadding = 0;
 			AudioClient->GetCurrentPadding(&NumFramesPadding);
@@ -74,9 +76,11 @@ namespace Audio
 			// be at least the amount requested. For example, if we request a 1024 frame buffer, WASAPI
 			// might allocate a 1056 frame buffer. The padding is subtracted from the allocated amount
 			// to determine how much space is available currently in the buffer.
-			const int32 NumFramesAvailable = NumFramesPerDeviceBuffer - NumFramesPadding;
+			const int32 OutputSpaceAvailable = NumFramesPerDeviceBuffer - NumFramesPadding;
+			// The number of frames rendered by the engine currently in the circular buffer
+			const int32 EngineFramesAvailable = DirectOutBuffers[0].Num();
 			
-			if (NumFramesAvailable >= WriteNumFrames)
+			if (const int32 WriteNumFrames = FMath::Min(EngineFramesAvailable, OutputSpaceAvailable))
 			{
 				uint8* RenderBufferPtr = nullptr;
 
@@ -101,9 +105,14 @@ namespace Audio
 	void FWasapiAggregateRenderStream::InterleaveOutput(float* OutRenderBufferPtr, const uint32 InNumFrames)
 	{
 		const int32 NumChannels = InterleaveBuffers.Num();
+		if (NumChannels <= 0)
+		{
+			return;
+		}
+		
 		const int32 NumFrames = InterleaveBuffers[0].Num();
 
-		if (NumFrames == InNumFrames)
+		if (NumFrames >= static_cast<int32>(InNumFrames))
 		{
 			for (int32 Index = 0; Index < DirectOutBuffers.Num(); ++Index)
 			{
@@ -128,7 +137,7 @@ namespace Audio
 
 			const float** InterleaveBufferPtr = BufferPtrArray.GetData();
 
-			ArrayInterleave(InterleaveBufferPtr, OutRenderBufferPtr, NumFrames, NumChannels);
+			ArrayInterleave(InterleaveBufferPtr, OutRenderBufferPtr, InNumFrames, NumChannels);
 		}
 	}
 }
diff --git a/Engine/Source/Runtime/Windows/AudioMixerWasapi/Private/WasapiAggregateRenderStream.h b/Engine/Source/Runtime/Windows/AudioMixerWasapi/Private/WasapiAggregateRenderStream.h
--- a/Engine/Source/Runtime/Windows/AudioMixerWasapi/Private/WasapiAggregateRenderStream.h
+++ b/Engine/Source/Runtime/Windows/AudioMixerWasapi/Private/WasapiAggregateRenderStream.h
@@ -30,7 +30,6 @@ namespace Audio
 
 	private:
 
-		int32 WriteNumFrames = 0;
 		TArray<FAlignedFloatBuffer> InterleaveBuffers;
 		TArray<TCircularAudioBuffer<float>> DirectOutBuffers;
 
diff --git a/Engine/Source/Runtime/Windows/AudioMixerWasapi/Private/WasapiDefaultRenderStream.cpp b/Engine/Source/Runtime/Windows/AudioMixerWasapi/Private/WasapiDefaultRenderStream.cpp
--- a/Engine/Source/Runtime/Windows/AudioMixerWasapi/Private/WasapiDefaultRenderStream.cpp
+++ b/Engine/Source/Runtime/Windows/AudioMixerWasapi/Private/WasapiDefaultRenderStream.cpp
@@ -11,17 +11,8 @@ namespace Audio
 	{
 		if (FAudioMixerWasapiRenderStream::InitializeHardware(InParams))
 		{
-			NumPartialBuffersToWrite = 1;
-			const uint32 MinBufferSize = GetDefaultDevicePeriod();
-
-			if (MinBufferSize > InParams.NumFrames)
-			{
-				NumPartialBuffersToWrite = FMath::DivideAndRoundUp(MinBufferSize, InParams.NumFrames);
-			}
-
-			WriteNumFrames = NumPartialBuffersToWrite * InParams.NumFrames;
-			AudioBufferTotalBytes = WriteNumFrames * AudioFormat.GetFrameSizeInBytes();
-			PartialBufferNumBytes = InParams.NumFrames * AudioFormat.GetFrameSizeInBytes();
+			EngineBufferNumFrames = InParams.NumFrames;
+			EngineBufferNumBytes = EngineBufferNumFrames * AudioFormat.GetFrameSizeInBytes();
 
 			return true;
 		}
@@ -50,21 +41,28 @@ namespace Audio
 			// be at least the amount requested. For example, if we request a 1024 frame buffer, WASAPI
 			// might allocate a 1056 frame buffer. The padding is subtracted from the allocated amount
 			// to determine how much space is available currently in the buffer.
-			const int32 NumFramesAvailable = NumFramesPerDeviceBuffer - NumFramesPadding;
+			const int32 OutputSpaceAvailable = NumFramesPerDeviceBuffer - NumFramesPadding;
 			
-			if (NumFramesAvailable >= WriteNumFrames)
+			if (OutputSpaceAvailable >= EngineBufferNumFrames)
 			{
 				check(RenderBufferView.IsEmpty());
 
+				// The number engine buffers that can fit in the space available in the output buffer
+				const int32 NumEngineBuffersNeeded = OutputSpaceAvailable / EngineBufferNumFrames;
+				// The total number of frames that will be placed in the output buffer
+				const int32 TotalFrames = NumEngineBuffersNeeded * EngineBufferNumFrames;
+				// The total number of bytes that will be placed in the output buffer
+				const int32 TotalBytes = TotalFrames * AudioFormat.GetFrameSizeInBytes();
+
 				uint8* BufferStartPtr = nullptr;
-				if (SUCCEEDED(RenderClient->GetBuffer(WriteNumFrames, &BufferStartPtr)))
+				if (SUCCEEDED(RenderClient->GetBuffer(TotalFrames, &BufferStartPtr)))
 				{
-					TArrayView<uint8> BufferStartView(BufferStartPtr, AudioBufferTotalBytes);
+					TArrayView<uint8> BufferStartView(BufferStartPtr, TotalBytes);
 					
-					for (int32 i = 0; i < NumPartialBuffersToWrite; ++i)
+					for (int32 i = 0; i < NumEngineBuffersNeeded; ++i)
 					{
-						const SIZE_T ByteOffset = i * PartialBufferNumBytes;
-						RenderBufferView = BufferStartView.Slice(ByteOffset, AudioBufferTotalBytes - ByteOffset);
+						const SIZE_T ByteOffset = i * EngineBufferNumBytes;
+						RenderBufferView = BufferStartView.Slice(ByteOffset, TotalBytes - ByteOffset);
 						
 						if (!ReadNextBufferDelegate.ExecuteIfBound())
 						{
@@ -72,7 +70,7 @@ namespace Audio
 						}
 					}
 
-					HRESULT Result = RenderClient->ReleaseBuffer(WriteNumFrames, 0 /* flags */);
+					HRESULT Result = RenderClient->ReleaseBuffer(TotalFrames, 0 /* flags */);
 					if (FAILED(Result))
 					{
 						++CallbackBufferErrors;
diff --git a/Engine/Source/Runtime/Windows/AudioMixerWasapi/Private/WasapiDefaultRenderStream.h b/Engine/Source/Runtime/Windows/AudioMixerWasapi/Private/WasapiDefaultRenderStream.h
--- a/Engine/Source/Runtime/Windows/AudioMixerWasapi/Private/WasapiDefaultRenderStream.h
+++ b/Engine/Source/Runtime/Windows/AudioMixerWasapi/Private/WasapiDefaultRenderStream.h
@@ -34,17 +34,11 @@ namespace Audio
 	private:
 
 		/** Number of frames to write every Wasapi period (must be at least the minimum required for current driver) */
-		int32 WriteNumFrames = 0;
-
-		/** Total number of bytes per audio buffer copied to Wasapi each buffer period. */
-		SIZE_T AudioBufferTotalBytes = 0;
-
+		int32 EngineBufferNumFrames = 0;
+		
 		/** Number of bytes per partial buffer where a partial buffer is used when the engine buffer size
 		 *  is less than the Wasapi drivers buffer size. */
-		SIZE_T PartialBufferNumBytes = 0;
-
-		/** Number of partial buffers needed to create a full Wasapi buffer. Used for buffer sizes smaller than the minimum for the current driver. */
-		int32 NumPartialBuffersToWrite = 0;
+		SIZE_T EngineBufferNumBytes = 0;
 		
 		/** Delegate called each buffer callback to signal the mixer to process the next buffer. */
 		FAudioMixerReadNextBufferDelegate ReadNextBufferDelegate;
