Fix compare operator type-o in SubmitBuffer

#rnx

diff --git a/Engine/Source/Runtime/Windows/AudioMixerWasapi/Private/WasapiDefaultRenderStream.cpp b/Engine/Source/Runtime/Windows/AudioMixerWasapi/Private/WasapiDefaultRenderStream.cpp
--- a/Engine/Source/Runtime/Windows/AudioMixerWasapi/Private/WasapiDefaultRenderStream.cpp
+++ b/Engine/Source/Runtime/Windows/AudioMixerWasapi/Private/WasapiDefaultRenderStream.cpp
@@ -102,7 +102,7 @@ namespace Audio
 				const int32 InputNumChannels = RenderStreamParams.HardwareDeviceInfo.NumChannels;
 				const int32 OutputNumChannels = AudioFormat.GetNumChannels();
 				
-				if (OutputNumChannels != InputNumChannels)
+				if (OutputNumChannels < InputNumChannels)
 				{
 					DownmixBuffer(InputNumChannels, OutputNumChannels, reinterpret_cast<const float*>(InBuffer), DownmixScratchBuffer.GetData(), InNumFrames, MixdownGainsMap.GetData());
 					FMemory::Memcpy(RenderBufferPtr, DownmixScratchBuffer.GetData(), NumBytes);
