Use RAW stream type to avoid system audio enhancements

#rnx

diff --git a/Engine/Source/Runtime/Windows/AudioMixerWasapi/Private/AudioMixerWasapiRenderStream.cpp b/Engine/Source/Runtime/Windows/AudioMixerWasapi/Private/AudioMixerWasapiRenderStream.cpp
--- a/Engine/Source/Runtime/Windows/AudioMixerWasapi/Private/AudioMixerWasapiRenderStream.cpp
+++ b/Engine/Source/Runtime/Windows/AudioMixerWasapi/Private/AudioMixerWasapiRenderStream.cpp
@@ -42,13 +42,27 @@ namespace Audio
 			return false;
 		}
 
-		TComPtr<IAudioClient> TempAudioClient;
-		HRESULT Result = MMDevice->Activate(__uuidof(IAudioClient), CLSCTX_INPROC_SERVER, nullptr, IID_PPV_ARGS_Helper(&TempAudioClient));
+		TComPtr<IAudioClient2> TempAudioClient;
+		HRESULT Result = MMDevice->Activate(__uuidof(IAudioClient2), CLSCTX_INPROC_SERVER, nullptr, IID_PPV_ARGS_Helper(&TempAudioClient));
 		if (FAILED(Result))
 		{
 			UE_LOG(LogAudioMixerWasapi, Error, TEXT("OpenAudioStream failed IMMDevice::Activate %s"), *AudioClientErrorToFString(Result));
 			return false;
 		}
+
+		AudioClientProperties ClientProperties = {};
+		ClientProperties.cbSize = sizeof(AudioClientProperties);
+		ClientProperties.bIsOffload = false;
+		ClientProperties.eCategory = AudioCategory_GameEffects;
+		ClientProperties.Options = AUDCLNT_STREAMOPTIONS_RAW;
+
+		// Set the stream type to RAW in order to avoid system audio enhancements
+		Result = TempAudioClient->SetClientProperties(&ClientProperties);
+		if (FAILED(Result))
+		{
+			// Not all clients support raw mode. Log warning and continue.
+			UE_LOG(LogAudioMixerWasapi, Warning, TEXT("InitializeHardware failed SetClientProperties %s"), *AudioClientErrorToFString(Result));
+		}
 		
 		WAVEFORMATEX* MixFormat = nullptr;
 		Result = TempAudioClient->GetMixFormat(&MixFormat);
diff --git a/Engine/Source/Runtime/Windows/AudioMixerWasapi/Private/AudioMixerWasapiRenderStream.h b/Engine/Source/Runtime/Windows/AudioMixerWasapi/Private/AudioMixerWasapiRenderStream.h
--- a/Engine/Source/Runtime/Windows/AudioMixerWasapi/Private/AudioMixerWasapiRenderStream.h
+++ b/Engine/Source/Runtime/Windows/AudioMixerWasapi/Private/AudioMixerWasapiRenderStream.h
@@ -55,7 +55,7 @@ namespace Audio
 		uint32 GetDefaultDevicePeriod() const { return DefaultDevicePeriod; }
 
 		/** COM pointer to the WASAPI audio client object. */
-		TComPtr<IAudioClient> AudioClient;
+		TComPtr<IAudioClient2> AudioClient;
 
 		/** COM pointer to the WASAPI render client object. */
 		TComPtr<IAudioRenderClient> RenderClient;
