AudioMixerWasapi swaps the audio driver to refresh the OnFomatChanged is called on the live interface.

diff --git a/Engine/Source/Runtime/Windows/AudioMixerWasapi/Private/AudioMixerWasapi.h b/Engine/Source/Runtime/Windows/AudioMixerWasapi/Private/AudioMixerWasapi.h
--- a/Engine/Source/Runtime/Windows/AudioMixerWasapi/Private/AudioMixerWasapi.h
+++ b/Engine/Source/Runtime/Windows/AudioMixerWasapi/Private/AudioMixerWasapi.h
@@ -93,6 +93,7 @@ namespace Audio
 		virtual void OnDeviceAdded(const FString& DeviceId, bool bIsRenderDevice) override;
 		virtual void OnDeviceRemoved(const FString& DeviceId, bool bIsRenderDevice) override;
 		virtual void OnDeviceStateChanged(const FString& DeviceId, const EAudioDeviceState InState, bool bIsRenderDevice) override;
+		virtual void OnFormatChanged(const FString& InDeviceId, const FFormatChangedData& InFormat) override;
 		virtual void OnSessionDisconnect(Audio::IAudioMixerDeviceChangedListener::EDisconnectReason InReason) override;
 		virtual FString GetDeviceId() const override;
 		//~ End IAudioMixerDeviceChangedListener
diff --git a/Engine/Source/Runtime/Windows/AudioMixerWasapi/Private/AudioMixerWasapiDeviceListener.cpp b/Engine/Source/Runtime/Windows/AudioMixerWasapi/Private/AudioMixerWasapiDeviceListener.cpp
--- a/Engine/Source/Runtime/Windows/AudioMixerWasapi/Private/AudioMixerWasapiDeviceListener.cpp
+++ b/Engine/Source/Runtime/Windows/AudioMixerWasapi/Private/AudioMixerWasapiDeviceListener.cpp
@@ -198,6 +198,45 @@ namespace Audio
 		}
 	}
 	
+	void FAudioMixerWasapi::OnFormatChanged(const FString& InDeviceId, const FFormatChangedData& InFormat)
+	{
+		bool bShouldSwapToThisDevice = false;
+
+		{
+			// Access to device swap context must be protected by DeviceSwapCriticalSection
+			FScopeLock Lock(&DeviceSwapCriticalSection);
+
+			// If we are currently swapping and the device we are swapping to is changing format, queue up a new swap
+			if (AudioStreamInfo.StreamState == EAudioOutputStreamState::SwappingDevice)
+			{
+				if (DeviceSwapContext.IsValid())
+				{
+					if (DeviceSwapContext->NewDevice.IsSet())
+					{
+						if (DeviceSwapContext->NewDevice->DeviceId.Equals(InDeviceId))
+						{
+							bShouldSwapToThisDevice = true;
+						}
+					}
+				}
+			}
+			// if we are trying to change the format of the current live device, force a device swap to refresh
+			else if (AudioStreamInfo.DeviceInfo.DeviceId == InDeviceId)
+			{
+				bShouldSwapToThisDevice = true;
+			}
+		}
+
+		if (bShouldSwapToThisDevice)
+		{
+			constexpr bool bForceDeviceSwap = true;
+			const FString SwapReason = FString::Printf(TEXT("FAudioMixerWasapi - OnFormatChange for live audio device"));
+
+			// Refresh the newly reformatted device
+			RequestDeviceSwap(InDeviceId, bForceDeviceSwap, *SwapReason);
+		}
+	};
+
 	FString FAudioMixerWasapi::GetDeviceId() const
 	{
 		return AudioStreamInfo.DeviceInfo.DeviceId;
