Fix device swapping crash with Wasapi by eliminating delete race condition.

diff --git a/Engine/Source/Runtime/Windows/AudioMixerWasapi/Private/AudioMixerWasapi.cpp b/Engine/Source/Runtime/Windows/AudioMixerWasapi/Private/AudioMixerWasapi.cpp
--- a/Engine/Source/Runtime/Windows/AudioMixerWasapi/Private/AudioMixerWasapi.cpp
+++ b/Engine/Source/Runtime/Windows/AudioMixerWasapi/Private/AudioMixerWasapi.cpp
@@ -503,11 +503,11 @@ namespace Audio
 		
 		if (DeviceSwapContext.IsValid())
 		{
-			UE_LOG(LogAudioMixer, Display, TEXT("FAudioMixerWasapi::InitDeviceSwapContextInternal - DeviceSwapContext in-flight, replacing"));
+			UE_LOG(LogAudioMixer, Display, TEXT("FAudioMixerWasapi::InitDeviceSwapContextInternal - DeviceSwapContext in-flight, ignoring"));
+			return false;
 		}
 
-		// Create the device swap context, replacing any previous request
-		DeviceSwapContext.Reset();
+		// Create the new device swap context
 		DeviceSwapContext = MakeUnique<FWasapiDeviceSwapContext>(InRequestedDeviceID, InReason);
 		if (!DeviceSwapContext.IsValid())
 		{
