Unshelved from pending changelist '46136459':

[Wasapi] Complete the audio device disconnect process so we can act on the next disconnect

diff --git a/Engine/Source/Runtime/Windows/AudioMixerWasapi/Private/AudioMixerWasapi.cpp b/Engine/Source/Runtime/Windows/AudioMixerWasapi/Private/AudioMixerWasapi.cpp
--- a/Engine/Source/Runtime/Windows/AudioMixerWasapi/Private/AudioMixerWasapi.cpp
+++ b/Engine/Source/Runtime/Windows/AudioMixerWasapi/Private/AudioMixerWasapi.cpp
@@ -637,6 +637,10 @@ namespace Audio
 				return {};
 			}
 			
+#if PLATFORM_WINDOWS
+			RegisterForSessionEvents(InDeviceContext->RequestedDeviceId);
+#endif // PLATFORM_WINDOWS
+
 			DeviceSwapResult->SuccessfulDurationMs = FPlatformTime::ToMilliseconds64(FPlatformTime::Cycles64() - StartTimeCycles);
 			DeviceSwapResult->DeviceInfo = InDeviceContext->StreamParams[0].HardwareDeviceInfo;
 			DeviceSwapResult->bIsAggregateDevice = InDeviceContext->bIsAggregateDevice;
diff --git a/Engine/Source/Runtime/Windows/AudioMixerWasapi/Private/AudioMixerWasapi.h b/Engine/Source/Runtime/Windows/AudioMixerWasapi/Private/AudioMixerWasapi.h
--- a/Engine/Source/Runtime/Windows/AudioMixerWasapi/Private/AudioMixerWasapi.h
+++ b/Engine/Source/Runtime/Windows/AudioMixerWasapi/Private/AudioMixerWasapi.h
@@ -146,6 +146,14 @@ namespace Audio
 		
 		/** Performs a device swap with the given context. Static method enforces no other state sharing occurs. */
 		static TUniquePtr<FDeviceSwapResult> PerformDeviceSwap(TUniquePtr<FWasapiDeviceSwapContext>&& InDeviceContext);
+
+#if PLATFORM_WINDOWS
+		/** Register with the Windows MM Notification Client for updates */
+		static void RegisterForSessionEvents(const FString& InDeviceId);
+		static void UnregisterForSessionEvents();
+#endif //PLATFORM_WINDOWS
+
+
 	};
 
  }
diff --git a/Engine/Source/Runtime/Windows/AudioMixerWasapi/Private/AudioMixerWasapiDeviceListener.cpp b/Engine/Source/Runtime/Windows/AudioMixerWasapi/Private/AudioMixerWasapiDeviceListener.cpp
--- a/Engine/Source/Runtime/Windows/AudioMixerWasapi/Private/AudioMixerWasapiDeviceListener.cpp
+++ b/Engine/Source/Runtime/Windows/AudioMixerWasapi/Private/AudioMixerWasapiDeviceListener.cpp
@@ -26,6 +26,23 @@ namespace Audio
 {	
 	TSharedPtr<FWindowsMMNotificationClient> WasapiWinNotificationClient;
 
+#if PLATFORM_WINDOWS
+	void FAudioMixerWasapi::RegisterForSessionEvents(const FString& InDeviceId)
+	{
+		if (WasapiWinNotificationClient)
+		{
+			WasapiWinNotificationClient->RegisterForSessionNotifications(InDeviceId);
+		}
+	}
+	void FAudioMixerWasapi::UnregisterForSessionEvents()
+	{
+		if (WasapiWinNotificationClient)
+		{
+			WasapiWinNotificationClient->UnregisterForSessionNotifications();
+		}
+	}
+#endif //PLATFORM_WINDOWS
+
 	void FAudioMixerWasapi::RegisterDeviceChangedListener()
 	{
 		if (!WasapiWinNotificationClient.IsValid())
