[Renderer] Stereo rendering issues DX11/DX12

Since 5.6 (UE-282114), nDisplay copies the right-eye texture only when the backbuffer texture dimension is ETextureDimension::Texture2DArray.
Quadbuffer textures on DX11/DX12 are still created as ETextureDimension::Texture2D instead of Texture2DArray, so nDisplay cannot treat them as valid backbuffers for stereo rendering.

This regression is present in both 5.6 and 5.7 and results in no stereo in nDisplay: the right eye remains black, which appears as blinking when stereo is enabled.

Note: Here is a link to the original CL: https://p4-swarm.epicgames.net/reviews/45890199

diff --git a/Engine/Source/Runtime/D3D12RHI/Private/Windows/WindowsD3D12Viewport.cpp b/Engine/Source/Runtime/D3D12RHI/Private/Windows/WindowsD3D12Viewport.cpp
--- a/Engine/Source/Runtime/D3D12RHI/Private/Windows/WindowsD3D12Viewport.cpp
+++ b/Engine/Source/Runtime/D3D12RHI/Private/Windows/WindowsD3D12Viewport.cpp
@@ -241,8 +241,12 @@ void FD3D12Viewport::ConditionalResetSwapChain(bool bIgnoreFocus)
 */
 FD3D12Texture* FD3D12Viewport::CreateDummyBackBufferTextures(FD3D12Adapter* InAdapter, EPixelFormat InPixelFormat, uint32 InSizeX, uint32 InSizeY)
 {
+	// active stereoscopy uses RHI resources with 2DArray type
+	bool const bQuadBufferStereo = FD3D12DynamicRHI::GetD3DRHI()->IsQuadBufferStereoEnabled();
 	FRHITextureCreateDesc CreateDesc =
-		FRHITextureCreateDesc::Create2D(TEXT("BackBufferReference"))
+		FRHITextureCreateDesc::Create(TEXT("BackBufferReference"),
+			bQuadBufferStereo ? ETextureDimension::Texture2DArray : ETextureDimension::Texture2D)
+		.SetArraySize(bQuadBufferStereo ? 2 : 1)
 		.SetExtent(FIntPoint(InSizeX, InSizeY))
 		.SetFormat(InPixelFormat)
 		.SetFlags(ETextureCreateFlags::RenderTargetable | ETextureCreateFlags::Presentable | ETextureCreateFlags::ResolveTargetable)
diff --git a/Engine/Source/Runtime/Windows/D3D11RHI/Private/D3D11Viewport.cpp b/Engine/Source/Runtime/Windows/D3D11RHI/Private/D3D11Viewport.cpp
--- a/Engine/Source/Runtime/Windows/D3D11RHI/Private/D3D11Viewport.cpp
+++ b/Engine/Source/Runtime/Windows/D3D11RHI/Private/D3D11Viewport.cpp
@@ -143,7 +143,8 @@ FD3D11Texture* FD3D11Viewport::GetSwapChainSurface(FD3D11DynamicRHI* D3DRHI, EPi
 	TRefCountPtr<ID3D11RenderTargetView> BackBufferRenderTargetViewRight; // right eye RTV
 	
 	// dx11.1 active stereoscopy initialization
-	if (D3DRHI->IsQuadBufferStereoEnabled())
+	bool const bQuadBufferStereo = D3DRHI->IsQuadBufferStereoEnabled();
+	if (bQuadBufferStereo)
 	{
 		// left
 		CD3D11_RENDER_TARGET_VIEW_DESC RTVDescCD3D11_left(D3D11_RTV_DIMENSION_TEXTURE2DARRAY, DXGI_FORMAT_R10G10B10A2_UNORM, 0, 0, 1);
@@ -169,7 +170,7 @@ FD3D11Texture* FD3D11Viewport::GetSwapChainSurface(FD3D11DynamicRHI* D3DRHI, EPi
 	RenderTargetViews.Add(BackBufferRenderTargetView);
 
 	// add right eye render target view
-	if (D3DRHI->IsQuadBufferStereoEnabled())
+	if (bQuadBufferStereo)
 	{
 		RenderTargetViews.Add(BackBufferRenderTargetViewRight);
 	}
@@ -184,7 +185,11 @@ FD3D11Texture* FD3D11Viewport::GetSwapChainSurface(FD3D11DynamicRHI* D3DRHI, EPi
 	VERIFYD3D11RESULT_EX(D3DRHI->GetDevice()->CreateShaderResourceView(BackBufferResource,&SRVDesc,BackBufferShaderResourceView.GetInitReference()), D3DRHI->GetDevice());
 
 	const FRHITextureCreateDesc CreateDesc =
-		FRHITextureCreateDesc::Create2D(TEXT("FD3D11Viewport::GetSwapChainSurface"), TextureDesc.Width, TextureDesc.Height, PixelFormat)
+		FRHITextureCreateDesc::Create(TEXT("FD3D11Viewport::GetSwapChainSurface"),
+			bQuadBufferStereo ? ETextureDimension::Texture2DArray : ETextureDimension::Texture2D)
+		.SetArraySize(bQuadBufferStereo ? 2 : 1)
+		.SetExtent(FIntPoint(TextureDesc.Width, TextureDesc.Height))
+		.SetFormat(PixelFormat)
 		.SetFlags(ETextureCreateFlags::RenderTargetable)
 		.DetermineInititialState();
 
@@ -192,7 +197,7 @@ FD3D11Texture* FD3D11Viewport::GetSwapChainSurface(FD3D11DynamicRHI* D3DRHI, EPi
 		CreateDesc,
 		BackBufferResource,
 		BackBufferShaderResourceView,
-		1,
+		bQuadBufferStereo ? 2 : 1,
 		false,
 		RenderTargetViews,
 		{}
