[Wasapi] Use only the Console role for default devices

diff --git a/Engine/Source/Runtime/Windows/AudioMixerWasapi/Private/AudioMixerWasapiDeviceListener.cpp b/Engine/Source/Runtime/Windows/AudioMixerWasapi/Private/AudioMixerWasapiDeviceListener.cpp
--- a/Engine/Source/Runtime/Windows/AudioMixerWasapi/Private/AudioMixerWasapiDeviceListener.cpp
+++ b/Engine/Source/Runtime/Windows/AudioMixerWasapi/Private/AudioMixerWasapiDeviceListener.cpp
@@ -96,7 +96,10 @@ namespace Audio
 		// However, the Windows UX only allows console+multimedia to be toggle as a pair. This means you get two notifications
 		// for default device changing typically. To prevent a trouble trigger we only listen to "Console" here. For more information on 
 		// device roles: https://docs.microsoft.com/en-us/windows/win32/coreaudio/device-roles
-		// Note: On XBox, headphones plugged into an XBox controller will use the Communications role.
+
+		// Note: On XBox, headphones plugged into an XBox controller give a device notification on the Communications role. However, you 
+		// should not swap to a device for a role other than the default Console, so we re-swap to the Console default here and let
+		// the Xbox OS take care of the routing details. 
 		
 		if (InAudioDeviceRole == EAudioDeviceRole::Console || InAudioDeviceRole == EAudioDeviceRole::Communications)
 		{
@@ -105,9 +108,20 @@ namespace Audio
 
 			// Ignore if not listening for events. This means the user has specified a specific device
 			// to use and is not interested in following the current system default device.
+
+			FString DeviceIdForSwap;
+#if PLATFORM_WINDOWS
+			DeviceIdForSwap = DeviceId;
+#else
+			if (IAudioPlatformDeviceInfoCache* Cache = GetDeviceInfoCache())
+			{
+				DeviceIdForSwap = Cache->GetDefaultOutputDevice(GetDefaultDeviceRole()).ToString();
+			}
+#endif // PLATFORM_WINDOWS
+			
 			if (GetIsListeningForDeviceEvents())
 			{
-				RequestDeviceSwap(DeviceId, /* force */true, TEXT("FAudioMixerWasapi::OnDefaultRenderDeviceChanged"));
+				RequestDeviceSwap(DeviceIdForSwap, /* force */true, TEXT("FAudioMixerWasapi::OnDefaultRenderDeviceChanged"));
 			}
 		}
 
