Using the audio downmix buffer correctly when our mixer has more channels than the hardware.

diff --git a/Engine/Source/Runtime/Windows/AudioMixerWasapi/Private/AudioMixerWasapiRenderStream.cpp b/Engine/Source/Runtime/Windows/AudioMixerWasapi/Private/AudioMixerWasapiRenderStream.cpp
--- a/Engine/Source/Runtime/Windows/AudioMixerWasapi/Private/AudioMixerWasapiRenderStream.cpp
+++ b/Engine/Source/Runtime/Windows/AudioMixerWasapi/Private/AudioMixerWasapiRenderStream.cpp
@@ -224,17 +224,31 @@ namespace Audio
 		// The downmixing code assumes float audio buffers
 		check(InAudioOutputFormat.GetEncoding() == EWasapiAudioEncoding::FLOATING_POINT_32);
 
+		int32 StreamFormatNumChannels = InAudioOutputFormat.GetNumChannels();
+
 		DownmixScratchBuffer.Reset();
-		DownmixScratchBuffer.AddUninitialized(InNumFrames * InAudioOutputFormat.GetNumChannels());
+		DownmixScratchBuffer.AddUninitialized(InNumFrames * StreamFormatNumChannels);
 
 		MixdownGainsMap.Reset();
 
-		const float DownmixGain = 1.0f / InNumInputChannels;
-
-		// Gain values used when downmixing from input channel count to mono
-		for (int32 Index = 0; Index < InNumInputChannels; ++Index)
+		// There is a special case where the audio mixer is initialized to 2 channels but the OS is configured to use 
+		// a headphone spatializer.In this case, the callback requests 8 channels of audio despite originally telling 
+		// us that we're 2 channels. So we need to upmix our audio mixer's render to 8 channels.
+		if (InNumInputChannels == 2 && StreamFormatNumChannels == 8)
+		{
+			MixdownGainsMap = { 1.0f, 0.0f, 0.0f, 0.0f, 0.0f, 0.0f, 0.0f, 0.0f, 
+								0.0f, 1.0f, 0.0f, 0.0f, 0.0f, 0.0f, 0.0f, 0.0f };
+		}
+		else
 		{
-			MixdownGainsMap.Add(DownmixGain);
+			// otherwise, just make this a true mono signal
+			const float DownmixGain = 1.0f / StreamFormatNumChannels;
+
+			// Gain values used when downmixing from input channel count to mono
+			for (int32 Index = 0; Index < InNumInputChannels * StreamFormatNumChannels; ++Index)
+			{
+				MixdownGainsMap.Add(DownmixGain);
+			}
 		}
 	}
 
diff --git a/Engine/Source/Runtime/Windows/AudioMixerWasapi/Private/WasapiDefaultRenderStream.cpp b/Engine/Source/Runtime/Windows/AudioMixerWasapi/Private/WasapiDefaultRenderStream.cpp
--- a/Engine/Source/Runtime/Windows/AudioMixerWasapi/Private/WasapiDefaultRenderStream.cpp
+++ b/Engine/Source/Runtime/Windows/AudioMixerWasapi/Private/WasapiDefaultRenderStream.cpp
@@ -99,10 +99,11 @@ namespace Audio
 				// HardwareDeviceInfo is somewhat of a misnomer, HardwareDeviceInfo.NumChannels is the number of channels
 				// that the AudioMixer is requesting. AudioFormat is fetched from the hardware and represents the 
 				// actual number of output channels on the hardware device.
+				// Another misnomer, "DownmixBuffer" will also upmix if requested.
 				const int32 InputNumChannels = RenderStreamParams.HardwareDeviceInfo.NumChannels;
 				const int32 OutputNumChannels = AudioFormat.GetNumChannels();
-				
-				if (OutputNumChannels < InputNumChannels)
+
+				if (OutputNumChannels != InputNumChannels)
 				{
 					DownmixBuffer(InputNumChannels, OutputNumChannels, reinterpret_cast<const float*>(InBuffer), DownmixScratchBuffer.GetData(), InNumFrames, MixdownGainsMap.GetData());
 					FMemory::Memcpy(RenderBufferPtr, DownmixScratchBuffer.GetData(), NumBytes);
