Refactor D3D11 RHI GPU profiler event handling to use final timestamp query as frame completion marker
- Suspected cause of crashes in the GPU profiler where events have been stomped with timestamp values. The theory is that D3D11 "event" queries can signal completion out-of-order in relation to "timestamp" queries.
- This means we could be attempting to process a frame's GPU profiler events before the timestamp queries that write into them have been fully read back via PollQueryResults().
- Instead, using the last timestamp query inserted in the frame (which we always have for the FEndWork marker) means we're guaranteed to have read back all prior results, since timestamp queries are held in an ordered linked list.

diff --git a/Engine/Source/Runtime/Windows/D3D11RHI/Private/D3D11Device.cpp b/Engine/Source/Runtime/Windows/D3D11RHI/Private/D3D11Device.cpp
--- a/Engine/Source/Runtime/Windows/D3D11RHI/Private/D3D11Device.cpp
+++ b/Engine/Source/Runtime/Windows/D3D11RHI/Private/D3D11Device.cpp
@@ -294,7 +294,7 @@ void FD3D11DynamicRHI::Shutdown()
 		FlushProfilerStats();
 
 		auto& Event = EmplaceProfilerEvent<UE::RHI::GPUProfiler::FEvent::FBeginBreadcrumb>(Breadcrumb);
-		InsertProfilerTimestamp(&Event.GPUTimestampTOP);
+		InsertProfilerTimestamp(&Event.GPUTimestampTOP, false);
 	#else
 		if (GPUProfilingData.IsProfilingGPU())
 		{
@@ -313,7 +313,7 @@ void FD3D11DynamicRHI::Shutdown()
 		FlushProfilerStats();
 
 		auto& Event = EmplaceProfilerEvent<UE::RHI::GPUProfiler::FEvent::FEndBreadcrumb>(Breadcrumb);
-		InsertProfilerTimestamp(&Event.GPUTimestampBOP);
+		InsertProfilerTimestamp(&Event.GPUTimestampBOP, false);
 	#else
 		if (GPUProfilingData.IsProfilingGPU())
 		{
@@ -605,7 +605,7 @@ void FD3D11DynamicRHI::SetupAfterDeviceCreation()
 	//
 	{
 		auto& Event = EmplaceProfilerEvent<UE::RHI::GPUProfiler::FEvent::FBeginWork>(FPlatformTime::Cycles64());
-		InsertProfilerTimestamp(&Event.GPUTimestampTOP);
+		InsertProfilerTimestamp(&Event.GPUTimestampTOP, false);
 	}
 #endif
 }
diff --git a/Engine/Source/Runtime/Windows/D3D11RHI/Private/D3D11Query.cpp b/Engine/Source/Runtime/Windows/D3D11RHI/Private/D3D11Query.cpp
--- a/Engine/Source/Runtime/Windows/D3D11RHI/Private/D3D11Query.cpp
+++ b/Engine/Source/Runtime/Windows/D3D11RHI/Private/D3D11Query.cpp
@@ -32,7 +32,11 @@ FD3D11RenderQuery::FD3D11RenderQuery(EType Type)
 	{
 	case EType::Occlusion: Desc.Query = D3D11_QUERY_OCCLUSION; break;
 	case EType::Timestamp: Desc.Query = D3D11_QUERY_TIMESTAMP; break;
-	case EType::Profiler : Desc.Query = D3D11_QUERY_TIMESTAMP; break;
+
+	case EType::Profiler:
+	case EType::ProfilerEndFrame:
+		Desc.Query = D3D11_QUERY_TIMESTAMP;
+		break;
 
 	default:
 		checkNoEntry();
@@ -192,6 +196,7 @@ bool FD3D11RenderQuery::CacheResult(FD3D11DynamicRHI& RHI, bool bWait)
 	{
 #if RHI_NEW_GPU_PROFILER
 	case EType::Profiler:
+	case EType::ProfilerEndFrame:
 		{
 			// Convert from GPU timestamp to CPU timestamp (relative to FPlatformTime::Cycles64())
 			uint64 GPUDelta = Temp - RHI.TimestampCalibration->GPUTimestamp;
@@ -221,6 +226,9 @@ bool FD3D11RenderQuery::CacheResult(FD3D11DynamicRHI& RHI, bool bWait)
 	LastCachedBOPCounter.store(BOPCounter, std::memory_order_release);
 
 #if RHI_NEW_GPU_PROFILER
+	// Don't return ProfilerEndFrame queries to the pool yet.
+	// References to these are held by historic GPU profiler frames, and are used to determine when it is safe to process the timestamp data.
+	// They are returned to the pool manually in PollQueryResultsForEndFrame(...).
 	if (Type == EType::Profiler)
 	{
 		// Return the query to the pool
@@ -303,6 +311,28 @@ void FD3D11DynamicRHI::PollQueryResults()
 	}
 }
 
+bool FD3D11DynamicRHI::PollQueryResultsForEndFrame(FD3D11RenderQuery* Query)
+{
+	check(Query->Type == FD3D11RenderQuery::EType::ProfilerEndFrame);
+
+	// Poll as many query results as possible.
+	PollQueryResults();
+
+	if (Query->IsLinked())
+	{
+		// The end-frame query is still linked. The results have not arrived yet.
+		return false;
+	}
+	else
+	{
+		// An unlinked end-frame query indicates the results have arrived.
+		// Return the query to the pool for reuse.
+		Profiler.TimestampPoolEndFrame.Push(Query);
+
+		return true;
+	}
+}
+
 bool FD3D11DynamicRHI::GetQueryData(ID3D11Query* Query, void* Data, SIZE_T DataSize, bool bTimestamp, bool bWait, bool bStallRHIThread)
 {
 	// Request the data from the query.
diff --git a/Engine/Source/Runtime/Windows/D3D11RHI/Private/D3D11RHI.cpp b/Engine/Source/Runtime/Windows/D3D11RHI/Private/D3D11RHI.cpp
--- a/Engine/Source/Runtime/Windows/D3D11RHI/Private/D3D11RHI.cpp
+++ b/Engine/Source/Runtime/Windows/D3D11RHI/Private/D3D11RHI.cpp
@@ -28,9 +28,10 @@ void FD3D11DynamicRHI::RHIEndFrame(const FRHIEndFrameArgs& Args)
 	// End Frame
 #if RHI_NEW_GPU_PROFILER
 	{
-		// End GPU work
+		// End GPU work.
+		// The EndWork query is always the last timestamp query of the frame. Keep a reference to it, so we can poll for frame completion below.
 		auto& EndWork = EmplaceProfilerEvent<UE::RHI::GPUProfiler::FEvent::FEndWork>();
-		InsertProfilerTimestamp(&EndWork.GPUTimestampBOP);
+		Profiler.Current.EndFrameQuery.Ptr = InsertProfilerTimestamp(&EndWork.GPUTimestampBOP, true);
 
 		uint64 Timestamp = FPlatformTime::Cycles64();
 
@@ -44,49 +45,27 @@ void FD3D11DynamicRHI::RHIEndFrame(const FRHIEndFrameArgs& Args)
 	#endif
 		);
 
-		// Issue a completion query so we know when to readback these profiler results.
-		{
-			if (Profiler.EventPool.IsEmpty())
-			{
-				D3D11_QUERY_DESC QueryDesc {};
-				QueryDesc.Query = D3D11_QUERY_EVENT;
-				VERIFYD3D11RESULT(Direct3DDevice->CreateQuery(&QueryDesc, Profiler.Current.CompletionQuery.GetInitReference()));
-			}
-			else
-			{
-				Profiler.Current.CompletionQuery = Profiler.EventPool.Pop();
-			}
-
-			Direct3DDeviceIMContext->End(Profiler.Current.CompletionQuery);
-			Profiler.Pending.Enqueue(MakeUnique<FProfiler::FFrame>(MoveTemp(Profiler.Current)));
-		}
+		Profiler.Pending.Enqueue(MakeUnique<FProfiler::FFrame>(MoveTemp(Profiler.Current)));
 
 		// Attempt to process historic results
 		while (TUniquePtr<FProfiler::FFrame>* PreviousFramePtr = Profiler.Pending.Peek())
 		{
 			TUniquePtr<FProfiler::FFrame>& PreviousFrame = *PreviousFramePtr;
-
-			BOOL EventComplete = false;
-			VERIFYD3D11RESULT(Direct3DDeviceIMContext->GetData(PreviousFrame->CompletionQuery, &EventComplete, sizeof(EventComplete), 0));
-			if (!EventComplete)
+			if (!PollQueryResultsForEndFrame(PreviousFrame->EndFrameQuery.Ptr))
 			{
 				// Frame not yet finished on the GPU
 				break;
 			}
 
-			// Ensure we have the latest timestamp data
-			PollQueryResults();
-
 			// Previous frame has completed and the data is available. Publish the profiler events.
 			UE::RHI::GPUProfiler::ProcessEvents(MakeArrayView(&PreviousFrame->EventStream, 1));
 
-			Profiler.EventPool.Push(MoveTemp(PreviousFrame->CompletionQuery));
 			Profiler.Pending.Pop();
 		}
 
 		// Start the next frame's GPU work
 		auto& BeginWork = EmplaceProfilerEvent<UE::RHI::GPUProfiler::FEvent::FBeginWork>(Timestamp);
-		InsertProfilerTimestamp(&BeginWork.GPUTimestampTOP);
+		InsertProfilerTimestamp(&BeginWork.GPUTimestampTOP, false);
 	}
 
 #else
@@ -715,18 +694,26 @@ void FD3D11DynamicRHI::RHIVerifyResult(ID3D11Device* Device, HRESULT Result, con
 }
 
 #if RHI_NEW_GPU_PROFILER
-void FD3D11DynamicRHI::InsertProfilerTimestamp(uint64* Target)
+FD3D11RenderQuery* FD3D11DynamicRHI::InsertProfilerTimestamp(uint64* Target, bool bEndFrame)
 {
+	TArray<FD3D11RenderQuery*>& Pool = bEndFrame
+		? Profiler.TimestampPoolEndFrame
+		: Profiler.TimestampPool;
+
 	FD3D11RenderQuery* Query;
-	if (Profiler.TimestampPool.IsEmpty())
+	if (Pool.IsEmpty())
 	{
-		Query = new FD3D11RenderQuery(FD3D11RenderQuery::EType::Profiler);
+		Query = new FD3D11RenderQuery(bEndFrame
+			? FD3D11RenderQuery::EType::ProfilerEndFrame
+			: FD3D11RenderQuery::EType::Profiler
+		);
 	}
 	else
 	{
-		Query = Profiler.TimestampPool.Pop();
+		Query = Pool.Pop();
 	}
 
 	Query->End(Direct3DDeviceIMContext, Target);
+	return Query;
 }
 #endif // RHI_NEW_GPU_PROFILER
diff --git a/Engine/Source/Runtime/Windows/D3D11RHI/Private/D3D11RHIPrivate.h b/Engine/Source/Runtime/Windows/D3D11RHI/Private/D3D11RHIPrivate.h
--- a/Engine/Source/Runtime/Windows/D3D11RHI/Private/D3D11RHIPrivate.h
+++ b/Engine/Source/Runtime/Windows/D3D11RHI/Private/D3D11RHIPrivate.h
@@ -585,6 +585,7 @@ public:
 #endif
 
 	D3D11RHI_API void PollQueryResults();
+	D3D11RHI_API bool PollQueryResultsForEndFrame(FD3D11RenderQuery* Query);
 
 	// *_RenderThread functions. Command lists call these functions on RT. You can implement your own behavior inside these functions.
 	// For example, deferring the actual creation to RHI thread by sending an RHI command.
@@ -799,7 +800,18 @@ protected:
 	{
 		struct FFrame
 		{
-			TRefCountPtr<ID3D11Query> CompletionQuery;
+			struct FEndFrameQuery
+			{
+				FD3D11RenderQuery* Ptr = nullptr;
+
+				FEndFrameQuery() = default;
+				FEndFrameQuery(FEndFrameQuery&& Other)
+					: Ptr(Other.Ptr)
+				{
+					Other.Ptr = nullptr;
+				}
+
+			} EndFrameQuery;
 
 		#if RHI_NEW_GPU_PROFILER
 			UE::RHI::GPUProfiler::FEventStream EventStream;
@@ -820,7 +832,7 @@ protected:
 		TQueue<TUniquePtr<FFrame>> Pending;
 
 		TArray<FD3D11RenderQuery*> TimestampPool;
-		TArray<TRefCountPtr<ID3D11Query>> EventPool;
+		TArray<FD3D11RenderQuery*> TimestampPoolEndFrame;
 
 	} Profiler;
 
@@ -831,7 +843,7 @@ protected:
 		return Profiler.Current.EventStream.Emplace<TEventType>(Forward<TArgs>(Args)...);
 	}
 
-	D3D11RHI_API void InsertProfilerTimestamp(uint64* Target);
+	D3D11RHI_API FD3D11RenderQuery* InsertProfilerTimestamp(uint64* Target, bool bEndFrame);
 
 	void FlushProfilerStats()
 	{
diff --git a/Engine/Source/Runtime/Windows/D3D11RHI/Public/D3D11Resources.h b/Engine/Source/Runtime/Windows/D3D11RHI/Public/D3D11Resources.h
--- a/Engine/Source/Runtime/Windows/D3D11RHI/Public/D3D11Resources.h
+++ b/Engine/Source/Runtime/Windows/D3D11RHI/Public/D3D11Resources.h
@@ -340,7 +340,8 @@ public:
 	{
 		Timestamp,
 		Occlusion,
-		Profiler
+		Profiler,
+		ProfilerEndFrame
 	} const Type;
 
 	FD3D11RenderQuery(EType Type);
