Remove legacy GPU profiler code paths from D3D11 RHI

diff --git a/Engine/Source/Runtime/Windows/D3D11RHI/Private/D3D11Commands.cpp b/Engine/Source/Runtime/Windows/D3D11RHI/Private/D3D11Commands.cpp
--- a/Engine/Source/Runtime/Windows/D3D11RHI/Private/D3D11Commands.cpp
+++ b/Engine/Source/Runtime/Windows/D3D11RHI/Private/D3D11Commands.cpp
@@ -168,10 +168,6 @@ void FD3D11DynamicRHI::RHIDispatchComputeShader(uint32 ThreadGroupCountX, uint32
 
 	StateCache.SetComputeShader(ComputeShader->Resource);
 
-#if (RHI_NEW_GPU_PROFILER == 0)
-	RegisterGPUDispatch(FIntVector(ThreadGroupCountX, ThreadGroupCountY, ThreadGroupCountZ));
-#endif
-
 	if (ComputeShader->bShaderNeedsGlobalConstantBuffer)
 	{
 		CommitComputeShaderConstants();
@@ -191,10 +187,6 @@ void FD3D11DynamicRHI::RHIDispatchIndirectComputeShader(FRHIBuffer* ArgumentBuff
 	FD3D11ComputeShader* ComputeShader = ResourceCast(ComputeShaderRHI);
 	FD3D11Buffer* ArgumentBuffer = ResourceCast(ArgumentBufferRHI);
 
-#if (RHI_NEW_GPU_PROFILER == 0)
-	RegisterGPUDispatch(FIntVector(1, 1, 1));
-#endif
-
 	StateCache.SetComputeShader(ComputeShader->Resource);
 	
 	if (ComputeShader->bShaderNeedsGlobalConstantBuffer)
@@ -1110,10 +1102,6 @@ void FD3D11DynamicRHI::RHIDrawPrimitive(uint32 BaseVertexIndex,uint32 NumPrimiti
 
 	RHI_DRAW_CALL_STATS(PrimitiveType, VertexCount, NumPrimitives, NumInstances);
 
-#if (RHI_NEW_GPU_PROFILER == 0)
-	RegisterGPUWork(NumPrimitives * NumInstances, VertexCount * NumInstances);
-#endif
-
 	StateCache.SetPrimitiveTopology(GetD3D11PrimitiveType(PrimitiveType));
 	if(NumInstances > 1)
 	{
@@ -1133,10 +1121,6 @@ void FD3D11DynamicRHI::RHIDrawPrimitiveIndirect(FRHIBuffer* ArgumentBufferRHI, u
 
 	RHI_DRAW_CALL_INC();
 
-#if (RHI_NEW_GPU_PROFILER == 0)
-	RegisterGPUWork(0);
-#endif
-
 	CommitGraphicsResourceTables();
 	CommitNonComputeShaderConstants();
 
@@ -1155,10 +1139,6 @@ void FD3D11DynamicRHI::RHIDrawIndexedPrimitive(FRHIBuffer* IndexBufferRHI, int32
 	// called should make sure the input is valid, this avoid hidden bugs
 	ensure(NumPrimitives > 0);
 
-#if (RHI_NEW_GPU_PROFILER == 0)
-	RegisterGPUWork(NumPrimitives * NumInstances, NumVertices * NumInstances);
-#endif
-
 	CommitGraphicsResourceTables();
 	CommitNonComputeShaderConstants();
 
@@ -1196,10 +1176,6 @@ void FD3D11DynamicRHI::RHIDrawIndexedPrimitiveIndirect(FRHIBuffer* IndexBufferRH
 	FD3D11Buffer* ArgumentBuffer = ResourceCast(ArgumentBufferRHI);
 
 	RHI_DRAW_CALL_INC();
-
-#if (RHI_NEW_GPU_PROFILER == 0)
-	RegisterGPUWork(0);
-#endif
 	
 	CommitGraphicsResourceTables();
 	CommitNonComputeShaderConstants();
@@ -1262,10 +1238,6 @@ void FD3D11DynamicRHI::RHIClearMRTImpl(const bool* bClearColorArray, int32 NumCl
 		}
 		Direct3DDeviceIMContext->ClearDepthStencilView(DepthStencilView,ClearFlags,Depth,Stencil);
 	}
-
-#if (RHI_NEW_GPU_PROFILER == 0)
-	RegisterGPUWork(0);
-#endif
 }
 
 // Blocks the CPU until the GPU catches up and goes idle.
@@ -1396,9 +1368,7 @@ struct FD3D11PlatformCommandList : public IRHIPlatformCommandList
 
 void FD3D11DynamicRHI::RHIFinalizeContext(FRHIFinalizeContextArgs&& Args, TRHIPipelineArray<IRHIPlatformCommandList*>& Output)
 {
-#if RHI_NEW_GPU_PROFILER
 	FlushProfilerStats();
-#endif
 
 	// "Context" will always be the default context, since we don't implement parallel execution.
 	for (IRHIComputeContext* Context : Args.Contexts)
@@ -1406,7 +1376,7 @@ void FD3D11DynamicRHI::RHIFinalizeContext(FRHIFinalizeContextArgs&& Args, TRHIPi
 		// "Context" will always be the default context, since we don't implement parallel execution.
 		check(Context == this);
 
-#if RHI_NEW_GPU_PROFILER && WITH_RHI_BREADCRUMBS
+#if WITH_RHI_BREADCRUMBS
 		// We need platform command lists to contain the breadcrumb allocators
 		Output[Context->GetPipeline()] = new FD3D11PlatformCommandList;
 #endif
@@ -1430,7 +1400,7 @@ void FD3D11DynamicRHI::RHISubmitCommandLists(FRHISubmitCommandListsArgs&& Args)
 	PollQueryResults();
 	FD3D11GPUFence::PollFences();
 
-#if RHI_NEW_GPU_PROFILER && WITH_RHI_BREADCRUMBS
+#if WITH_RHI_BREADCRUMBS
 	for (IRHIPlatformCommandList* CmdList : Args.CommandLists)
 	{
 		FD3D11PlatformCommandList* D3DCmdList = static_cast<FD3D11PlatformCommandList*>(CmdList);
diff --git a/Engine/Source/Runtime/Windows/D3D11RHI/Private/D3D11Device.cpp b/Engine/Source/Runtime/Windows/D3D11RHI/Private/D3D11Device.cpp
--- a/Engine/Source/Runtime/Windows/D3D11RHI/Private/D3D11Device.cpp
+++ b/Engine/Source/Runtime/Windows/D3D11RHI/Private/D3D11Device.cpp
@@ -56,9 +56,6 @@ FD3D11DynamicRHI::FD3D11DynamicRHI(IDXGIFactory1* InDXGIFactory1, D3D_FEATURE_LE
 	, PresentCounter(0)
 	, CurrentDSVAccessType(FExclusiveDepthStencil::DepthWrite_StencilWrite)
 	, bDiscardSharedConstants(false)
-#if (RHI_NEW_GPU_PROFILER == 0)
-	, GPUProfilingData(this)
-#endif
 	, Adapter(InAdapter)
 {
 	// This should be called once at the start 
@@ -256,11 +253,6 @@ void FD3D11DynamicRHI::Shutdown()
 	// Cleanup the D3D device.
 	CleanupD3DDevice();
 
-#if (RHI_NEW_GPU_PROFILER == 0)
-	// Release buffered timestamp queries
-	GPUProfilingData.FrameTiming.ReleaseResource();
-#endif
-
 	// Release the buffer of zeroes.
 	FMemory::Free(ZeroBuffer);
 	ZeroBuffer = NULL;
@@ -293,36 +285,18 @@ void FD3D11DynamicRHI::Shutdown()
 		#endif
 		}
 
-	#if RHI_NEW_GPU_PROFILER
 		FlushProfilerStats();
 
 		auto& Event = EmplaceProfilerEvent<UE::RHI::GPUProfiler::FEvent::FBeginBreadcrumb>(Breadcrumb);
 		InsertProfilerTimestamp(&Event.GPUTimestampTOP, false);
-	#else
-		if (GPUProfilingData.IsProfilingGPU())
-		{
-			// @todo dev-pr avoid TCHAR -> ANSI conversion
-			FRHIBreadcrumb::FBuffer Buffer;
-			TCHAR const* Name = Breadcrumb->GetTCHAR(Buffer);
-
-			GPUProfilingData.PushEvent(Name, FColor::White);
-		}
-	#endif
 	}
 
 	void FD3D11DynamicRHI::RHIEndBreadcrumbGPU(FRHIBreadcrumbNode* Breadcrumb)
 	{
-	#if RHI_NEW_GPU_PROFILER
 		FlushProfilerStats();
 
 		auto& Event = EmplaceProfilerEvent<UE::RHI::GPUProfiler::FEvent::FEndBreadcrumb>(Breadcrumb);
 		InsertProfilerTimestamp(&Event.GPUTimestampBOP, false);
-	#else
-		if (GPUProfilingData.IsProfilingGPU())
-		{
-			GPUProfilingData.PopEvent();
-		}
-	#endif
 
 		// Only emit formatted strings to platform APIs when requested.
 		if (ShouldEmitBreadcrumbs())
@@ -596,7 +570,6 @@ void FD3D11DynamicRHI::SetupAfterDeviceCreation()
 
 	TimestampCalibration = CalibrateTimers();
 
-#if RHI_NEW_GPU_PROFILER
 	// Register the single graphics GPU queue we have access to in D3D11.
 	UE::RHI::GPUProfiler::FQueue Queue(UE::RHI::GPUProfiler::FQueue::EType::Graphics, 0, 0);
 	UE::RHI::GPUProfiler::InitializeQueues(MakeConstArrayView(&Queue, 1));
@@ -610,7 +583,6 @@ void FD3D11DynamicRHI::SetupAfterDeviceCreation()
 		auto& Event = EmplaceProfilerEvent<UE::RHI::GPUProfiler::FEvent::FBeginWork>(FPlatformTime::Cycles64());
 		InsertProfilerTimestamp(&Event.GPUTimestampTOP, false);
 	}
-#endif
 }
 
 void FD3D11DynamicRHI::UpdateMSAASettings()
diff --git a/Engine/Source/Runtime/Windows/D3D11RHI/Private/D3D11Query.cpp b/Engine/Source/Runtime/Windows/D3D11RHI/Private/D3D11Query.cpp
--- a/Engine/Source/Runtime/Windows/D3D11RHI/Private/D3D11Query.cpp
+++ b/Engine/Source/Runtime/Windows/D3D11RHI/Private/D3D11Query.cpp
@@ -194,7 +194,6 @@ bool FD3D11RenderQuery::CacheResult(FD3D11DynamicRHI& RHI, bool bWait)
 	// Adjust timer queries to engine-clock ticks.
 	switch (Type)
 	{
-#if RHI_NEW_GPU_PROFILER
 	case EType::Profiler:
 	case EType::ProfilerEndFrame:
 		{
@@ -205,7 +204,6 @@ bool FD3D11RenderQuery::CacheResult(FD3D11DynamicRHI& RHI, bool bWait)
 			Temp = CPUDelta + RHI.TimestampCalibration->CPUTimestamp;
 		}
 		break;
-#endif
 
 	case EType::Timestamp:
 		{
@@ -225,7 +223,6 @@ bool FD3D11RenderQuery::CacheResult(FD3D11DynamicRHI& RHI, bool bWait)
 
 	LastCachedBOPCounter.store(BOPCounter, std::memory_order_release);
 
-#if RHI_NEW_GPU_PROFILER
 	// Don't return ProfilerEndFrame queries to the pool yet.
 	// References to these are held by historic GPU profiler frames, and are used to determine when it is safe to process the timestamp data.
 	// They are returned to the pool manually in PollQueryResultsForEndFrame(...).
@@ -234,7 +231,6 @@ bool FD3D11RenderQuery::CacheResult(FD3D11DynamicRHI& RHI, bool bWait)
 		// Return the query to the pool
 		RHI.Profiler.TimestampPool.Push(this);
 	}
-#endif
 
 	return true;
 }
@@ -547,335 +543,3 @@ TOptional<FD3D11DynamicRHI::FTimestampCalibration> FD3D11DynamicRHI::CalibrateTi
 
 	return {};
 }
-
-
-
-/*=============================================================================
- * class FD3D11BufferedGPUTiming
- *=============================================================================*/
-
-#if (RHI_NEW_GPU_PROFILER == 0)
-
-/**
- * Constructor.
- *
- * @param InD3DRHI			RHI interface
- * @param InBufferSize		Number of buffered measurements
- */
-FD3D11BufferedGPUTiming::FD3D11BufferedGPUTiming( FD3D11DynamicRHI* InD3DRHI, int32 InBufferSize )
-:	D3DRHI( InD3DRHI )
-,	BufferSize( InBufferSize )
-,	CurrentTimestamp( -1 )
-,	NumIssuedTimestamps( 0 )
-,	StartTimestamps( NULL )
-,	EndTimestamps( NULL )
-,	bIsTiming( false )
-{
-}
-
-/**
- * Initializes the static variables, if necessary.
- */
-void FD3D11BufferedGPUTiming::PlatformStaticInitialize(void* UserData)
-{
-	// Are the static variables initialized?
-	check( !GAreGlobalsInitialized );
-
-	// Get the GPU timestamp frequency.
-	SetTimingFrequency(0);
-	TRefCountPtr<ID3D11Query> FreqQuery;
-	FD3D11DynamicRHI* D3DRHI = (FD3D11DynamicRHI*)UserData;
-	ID3D11DeviceContext *D3D11DeviceContext = D3DRHI->GetDeviceContext();
-	HRESULT D3DResult;
-
-	D3D11_QUERY_DESC QueryDesc;
-	QueryDesc.Query = D3D11_QUERY_TIMESTAMP_DISJOINT;
-	QueryDesc.MiscFlags = 0;
-
-	{
-		// to track down some rare event where GTimingFrequency is 0 or <1000*1000
-		uint32 DebugState = 0;
-		uint32 DebugCounter = 0;
-
-		D3DResult = D3DRHI->GetDevice()->CreateQuery(&QueryDesc, FreqQuery.GetInitReference());
-		if (D3DResult == S_OK)
-		{
-			DebugState = 1;
-			D3D11DeviceContext->Begin(FreqQuery);
-			D3D11DeviceContext->End(FreqQuery);
-
-			D3D11_QUERY_DATA_TIMESTAMP_DISJOINT FreqQueryData;
-
-			{
-				FScopedD3D11RHIThreadStaller StallRHIThread;
-				D3DResult = D3D11DeviceContext->GetData(FreqQuery, &FreqQueryData, sizeof(D3D11_QUERY_DATA_TIMESTAMP_DISJOINT), 0);
-			}
-
-			double StartTime = FPlatformTime::Seconds();
-			while (D3DResult == S_FALSE && (FPlatformTime::Seconds() - StartTime) < 0.5f)
-			{
-				++DebugCounter;
-				FPlatformProcess::Sleep(0.005f);
-
-				FScopedD3D11RHIThreadStaller StallRHIThread;
-				D3DResult = D3D11DeviceContext->GetData(FreqQuery, &FreqQueryData, sizeof(D3D11_QUERY_DATA_TIMESTAMP_DISJOINT), 0);
-			}
-
-			if (D3DResult == S_OK)
-			{
-				DebugState = 2;
-				SetTimingFrequency(FreqQueryData.Frequency);
-				checkSlow(!FreqQueryData.Disjoint);
-
-				if (FreqQueryData.Disjoint)
-				{
-					DebugState = 3;
-				}
-			}
-		}
-
-		UE_LOG(LogD3D11RHI, Log, TEXT("GPU Timing Frequency: %f (Debug: %d %d)"), GetTimingFrequency() / (double)(1000 * 1000), DebugState, DebugCounter);
-	}
-
-	FreqQuery = NULL;
-
-	CalibrateTimers(D3DRHI);
-}
-
-void FD3D11BufferedGPUTiming::CalibrateTimers(FD3D11DynamicRHI* InD3DRHI)
-{
-	TOptional<FD3D11DynamicRHI::FTimestampCalibration> Data = InD3DRHI->CalibrateTimers();
-	if (Data.IsSet())
-	{
-		FGPUTimingCalibrationTimestamp CalibrationTimestamp;
-		CalibrationTimestamp.CPUMicroseconds = uint64(FPlatformTime::ToSeconds64(Data->CPUTimestamp) * 1e6);
-		CalibrationTimestamp.GPUMicroseconds = uint64(Data->GPUTimestamp * (1e6 / Data->GPUFrequency));
-		SetCalibrationTimestamp(CalibrationTimestamp);
-	}
-}
-
-/**
- * Initializes all D3D resources and if necessary, the static variables.
- */
-void FD3D11BufferedGPUTiming::InitRHI(FRHICommandListBase& RHICmdList)
-{
-	StaticInitialize(D3DRHI, PlatformStaticInitialize);
-
-	CurrentTimestamp = 0;
-	NumIssuedTimestamps = 0;
-	bIsTiming = false;
-
-	// Now initialize the queries for this timing object.
-	if ( GIsSupported )
-	{
-		StartTimestamps = new TRefCountPtr<ID3D11Query>[ BufferSize ];
-		EndTimestamps = new TRefCountPtr<ID3D11Query>[ BufferSize ];
-		for ( int32 TimestampIndex = 0; TimestampIndex < BufferSize; ++TimestampIndex )
-		{
-			HRESULT D3DResult;
-
-			D3D11_QUERY_DESC QueryDesc;
-			QueryDesc.Query = D3D11_QUERY_TIMESTAMP;
-			QueryDesc.MiscFlags = 0;
-
-			CA_SUPPRESS(6385);	// Doesn't like COM
-			D3DResult = D3DRHI->GetDevice()->CreateQuery(&QueryDesc,StartTimestamps[TimestampIndex].GetInitReference());
-			GIsSupported = GIsSupported && (D3DResult == S_OK);
-			CA_SUPPRESS(6385);	// Doesn't like COM
-			D3DResult = D3DRHI->GetDevice()->CreateQuery(&QueryDesc,EndTimestamps[TimestampIndex].GetInitReference());
-			GIsSupported = GIsSupported && (D3DResult == S_OK);
-		}
-	}
-}
-
-/**
- * Releases all D3D resources.
- */
-void FD3D11BufferedGPUTiming::ReleaseRHI()
-{
-	if ( StartTimestamps && EndTimestamps )
-	{
-		for ( int32 TimestampIndex = 0; TimestampIndex < BufferSize; ++TimestampIndex )
-		{
-			StartTimestamps[TimestampIndex] = NULL;
-			EndTimestamps[TimestampIndex] = NULL;
-		}
-		delete [] StartTimestamps;
-		delete [] EndTimestamps;
-		StartTimestamps = NULL;
-		EndTimestamps = NULL;
-	}
-}
-
-/**
- * Start a GPU timing measurement.
- */
-void FD3D11BufferedGPUTiming::StartTiming()
-{
-	// Issue a timestamp query for the 'start' time.
-	if ( GIsSupported && !bIsTiming )
-	{
-		int32 NewTimestampIndex = (CurrentTimestamp + 1) % BufferSize;
-		D3DRHI->GetDeviceContext()->End(StartTimestamps[NewTimestampIndex]);
-		CurrentTimestamp = NewTimestampIndex;
-		bIsTiming = true;
-	}
-}
-
-/**
- * End a GPU timing measurement.
- * The timing for this particular measurement will be resolved at a later time by the GPU.
- */
-void FD3D11BufferedGPUTiming::EndTiming()
-{
-	// Issue a timestamp query for the 'end' time.
-	if ( GIsSupported && bIsTiming )
-	{
-		checkSlow( CurrentTimestamp >= 0 && CurrentTimestamp < BufferSize );
-		D3DRHI->GetDeviceContext()->End(EndTimestamps[CurrentTimestamp]);
-		NumIssuedTimestamps = FMath::Min<int32>(NumIssuedTimestamps + 1, BufferSize);
-		bIsTiming = false;
-	}
-}
-
-/**
- * Retrieves the most recently resolved timing measurement.
- * The unit is the same as for FPlatformTime::Cycles(). Returns 0 if there are no resolved measurements.
- *
- * @return	Value of the most recently resolved timing, or 0 if no measurements have been resolved by the GPU yet.
- */
-uint64 FD3D11BufferedGPUTiming::GetTiming(bool bGetCurrentResultsAndBlock)
-{
-	if ( GIsSupported )
-	{
-		checkSlow( CurrentTimestamp >= 0 && CurrentTimestamp < BufferSize );
-		uint64 StartTime, EndTime;
-		HRESULT D3DResult;
-
-		int32 TimestampIndex = CurrentTimestamp;
-		if (!bGetCurrentResultsAndBlock)
-		{
-			// Quickly check the most recent measurements to see if any of them has been resolved.  Do not flush these queries.
-			for ( int32 IssueIndex = 1; IssueIndex < NumIssuedTimestamps; ++IssueIndex )
-			{
-				D3DResult = D3DRHI->GetDeviceContext()->GetData(EndTimestamps[TimestampIndex],&EndTime,sizeof(EndTime),D3D11_ASYNC_GETDATA_DONOTFLUSH);
-				if ( D3DResult == S_OK )
-				{
-					D3DResult = D3DRHI->GetDeviceContext()->GetData(StartTimestamps[TimestampIndex],&StartTime,sizeof(StartTime),D3D11_ASYNC_GETDATA_DONOTFLUSH);
-					if ( D3DResult == S_OK && EndTime > StartTime)
-					{
-						return EndTime - StartTime;
-					}
-				}
-
-				TimestampIndex = (TimestampIndex + BufferSize - 1) % BufferSize;
-			}
-		}
-		
-		if ( NumIssuedTimestamps > 0 || bGetCurrentResultsAndBlock )
-		{
-			// None of the (NumIssuedTimestamps - 1) measurements were ready yet,
-			// so check the oldest measurement more thoroughly.
-			// This really only happens if occlusion and frame sync event queries are disabled, otherwise those will block until the GPU catches up to 1 frame behind
-			const bool bBlocking = ( NumIssuedTimestamps == BufferSize ) || bGetCurrentResultsAndBlock;
-			const uint32 AsyncFlags = bBlocking ? 0 : D3D11_ASYNC_GETDATA_DONOTFLUSH;
-			{
-				FRenderThreadIdleScope IdleScope(ERenderThreadIdleTypes::WaitingForGPUQuery);
-				double StartTimeoutTime = FPlatformTime::Seconds();
-
-				SCOPE_CYCLE_COUNTER( STAT_RenderQueryResultTime );
-				// If we are blocking, retry until the GPU processes the time stamp command
-				do 
-				{
-					D3DResult = D3DRHI->GetDeviceContext()->GetData( EndTimestamps[TimestampIndex], &EndTime, sizeof(EndTime), AsyncFlags );
-
-					if ((FPlatformTime::Seconds() - StartTimeoutTime) > 0.5)
-					{
-						UE_LOG(LogD3D11RHI, Log, TEXT("Timed out while waiting for GPU to catch up. (500 ms)"));
-						return 0;
-					}
-				} while ( D3DResult == S_FALSE && bBlocking );
-			}
-
-			if ( D3DResult == S_OK )
-			{
-				{
-					FRenderThreadIdleScope IdleScope(ERenderThreadIdleTypes::WaitingForGPUQuery);
-					double StartTimeoutTime = FPlatformTime::Seconds();
-					do 
-					{
-						D3DResult = D3DRHI->GetDeviceContext()->GetData( StartTimestamps[TimestampIndex], &StartTime, sizeof(StartTime), AsyncFlags );
-
-						if ((FPlatformTime::Seconds() - StartTimeoutTime) > 0.5)
-						{
-							UE_LOG(LogD3D11RHI, Log, TEXT("Timed out while waiting for GPU to catch up. (500 ms)"));
-							return 0;
-						}
-					} while ( D3DResult == S_FALSE && bBlocking );
-				}
-
-				if ( D3DResult == S_OK && EndTime > StartTime )
-				{
-					return EndTime - StartTime;
-				}
-			}
-		}
-	}
-	return 0;
-}
-
-FD3D11DisjointTimeStampQuery::FD3D11DisjointTimeStampQuery(class FD3D11DynamicRHI* InD3DRHI) :
-	D3DRHI(InD3DRHI)
-{
-
-}
-
-void FD3D11DisjointTimeStampQuery::StartTracking()
-{
-	ID3D11DeviceContext* D3D11DeviceContext = D3DRHI->GetDeviceContext();
-	D3D11DeviceContext->Begin(DisjointQuery);
-}
-
-void FD3D11DisjointTimeStampQuery::EndTracking()
-{
-	ID3D11DeviceContext* D3D11DeviceContext = D3DRHI->GetDeviceContext();
-	D3D11DeviceContext->End(DisjointQuery);
-}
-
-bool FD3D11DisjointTimeStampQuery::IsResultValid()
-{
-	return GetResult().Disjoint == 0;
-}
-
-D3D11_QUERY_DATA_TIMESTAMP_DISJOINT FD3D11DisjointTimeStampQuery::GetResult()
-{
-	D3D11_QUERY_DATA_TIMESTAMP_DISJOINT DisjointQueryData;
-
-	ID3D11DeviceContext* D3D11DeviceContext = D3DRHI->GetDeviceContext();
-	HRESULT D3DResult = D3D11DeviceContext->GetData(DisjointQuery, &DisjointQueryData, sizeof(D3D11_QUERY_DATA_TIMESTAMP_DISJOINT), 0);
-
-	const double StartTime = FPlatformTime::Seconds();
-	while (D3DResult == S_FALSE && (FPlatformTime::Seconds() - StartTime) < 0.5)
-	{
-		FPlatformProcess::Sleep(0.005f);
-		D3DResult = D3D11DeviceContext->GetData(DisjointQuery, &DisjointQueryData, sizeof(D3D11_QUERY_DATA_TIMESTAMP_DISJOINT), 0);
-	}
-
-	return DisjointQueryData;
-}
-
-void FD3D11DisjointTimeStampQuery::InitRHI(FRHICommandListBase& RHICmdList)
-{
-	D3D11_QUERY_DESC QueryDesc;
-	QueryDesc.Query = D3D11_QUERY_TIMESTAMP_DISJOINT;
-	QueryDesc.MiscFlags = 0;
-
-	VERIFYD3D11RESULT_EX(D3DRHI->GetDevice()->CreateQuery(&QueryDesc, DisjointQuery.GetInitReference()), D3DRHI->GetDevice());
-}
-
-void FD3D11DisjointTimeStampQuery::ReleaseRHI()
-{
-
-}
-
-#endif // (RHI_NEW_GPU_PROFILER == 0)
diff --git a/Engine/Source/Runtime/Windows/D3D11RHI/Private/D3D11RHI.cpp b/Engine/Source/Runtime/Windows/D3D11RHI/Private/D3D11RHI.cpp
--- a/Engine/Source/Runtime/Windows/D3D11RHI/Private/D3D11RHI.cpp
+++ b/Engine/Source/Runtime/Windows/D3D11RHI/Private/D3D11RHI.cpp
@@ -26,7 +26,6 @@ extern void UniformBufferBeginFrame();
 void FD3D11DynamicRHI::RHIEndFrame(const FRHIEndFrameArgs& Args)
 {
 	// End Frame
-#if RHI_NEW_GPU_PROFILER
 	{
 		// End GPU work.
 		// The EndWork query is always the last timestamp query of the frame. Keep a reference to it, so we can poll for frame completion below.
@@ -68,18 +67,11 @@ void FD3D11DynamicRHI::RHIEndFrame(const FRHIEndFrameArgs& Args)
 		InsertProfilerTimestamp(&BeginWork.GPUTimestampTOP, false);
 	}
 
-#else
-	GPUProfilingData.EndFrame();
-#endif
-
 	UpdateMemoryStats();
 	CurrentComputeShader = nullptr;
 
 	// Begin Frame
 	UniformBufferBeginFrame();
-#if (RHI_NEW_GPU_PROFILER == 0)
-	GPUProfilingData.BeginFrame(this);
-#endif
 }
 
 template <int32 Frequency>
@@ -323,235 +315,6 @@ void FD3D11DynamicRHI::ClearAllShaderResources()
 	ClearAllShaderResourcesForFrequency<SF_Compute>();
 }
 
-#if (RHI_NEW_GPU_PROFILER == 0)
-
-void FD3DGPUProfiler::BeginFrame(FD3D11DynamicRHI* InRHI)
-{
-	CurrentEventNode = NULL;
-	check(!bTrackingEvents);
-	check(!CurrentEventNodeFrame); // this should have already been cleaned up and the end of the previous frame
-	
-	// latch the bools from the game thread into our private copy
-	bLatchedGProfilingGPU = GTriggerGPUProfile;
-	bLatchedGProfilingGPUHitches = GTriggerGPUHitchProfile;
-	if (bLatchedGProfilingGPUHitches)
-	{
-		bLatchedGProfilingGPU = false; // we do NOT permit an ordinary GPU profile during hitch profiles
-	}
-
-	// if we are starting a hitch profile or this frame is a gpu profile, then save off the state of the draw events
-	if (bLatchedGProfilingGPU || (!bPreviousLatchedGProfilingGPUHitches && bLatchedGProfilingGPUHitches))
-	{
-		bOriginalGEmitDrawEvents = GetEmitDrawEvents();
-	}
-
-	if (bLatchedGProfilingGPU || bLatchedGProfilingGPUHitches)
-	{
-		if (bLatchedGProfilingGPUHitches && GPUHitchDebounce)
-		{
-			// if we are doing hitches and we had a recent hitch, wait to recover
-			// the reasoning is that collecting the hitch report may itself hitch the GPU
-			GPUHitchDebounce--; 
-		}
-		else
-		{
-			SetEmitDrawEvents(true);  // thwart an attempt to turn this off on the game side
-			bTrackingEvents = true;
-			CurrentEventNodeFrame = new FD3D11EventNodeFrame(InRHI);
-			CurrentEventNodeFrame->StartFrame();
-		}
-	}
-	else if (bPreviousLatchedGProfilingGPUHitches)
-	{
-		// hitch profiler is turning off, clear history and restore draw events
-		GPUHitchEventNodeFrames.Empty();
-		SetEmitDrawEvents(bOriginalGEmitDrawEvents);
-	}
-	bPreviousLatchedGProfilingGPUHitches = bLatchedGProfilingGPUHitches;
-
-	FrameTiming.StartTiming();
-}
-
-void FD3DGPUProfiler::EndFrame()
-{
-	FrameTiming.EndTiming();
-
-	if (FrameTiming.IsSupported())
-	{
-		uint64 GPUTiming = FrameTiming.GetTiming();
-		uint64 GPUFreq = FrameTiming.GetTimingFrequency();
-		GRHIGPUFrameTimeHistory.PushFrameCycles(GPUFreq, GPUTiming);
-	}
-	else
-	{
-		GRHIGPUFrameTimeHistory.PushFrameCycles(1, 0);
-	}
-
-	// if we have a frame open, close it now.
-	if (CurrentEventNodeFrame)
-	{
-		CurrentEventNodeFrame->EndFrame();
-	}
-
-	check(!bTrackingEvents || bLatchedGProfilingGPU || bLatchedGProfilingGPUHitches);
-	check(!bTrackingEvents || CurrentEventNodeFrame);
-	if (bLatchedGProfilingGPU)
-	{
-		if (bTrackingEvents)
-		{
-			SetEmitDrawEvents(bOriginalGEmitDrawEvents);
-			UE_LOG(LogD3D11RHI, Warning, TEXT(""));
-			UE_LOG(LogD3D11RHI, Warning, TEXT(""));
-			CurrentEventNodeFrame->DumpEventTree();
-			GTriggerGPUProfile = false;
-			bLatchedGProfilingGPU = false;
-
-			if (RHIConfig::ShouldSaveScreenshotAfterProfilingGPU()
-				&& GEngine->GameViewport)
-			{
-				GEngine->GameViewport->Exec( NULL, TEXT("SCREENSHOT"), *GLog);
-			}
-		}
-	}
-	else if (bLatchedGProfilingGPUHitches)
-	{
-		//@todo this really detects any hitch, even one on the game thread.
-		// it would be nice to restrict the test to stalls on D3D, but for now...
-		// this needs to be out here because bTrackingEvents is false during the hitch debounce
-		static double LastTime = -1.0;
-		double Now = FPlatformTime::Seconds();
-		if (bTrackingEvents)
-		{
-			/** How long, in seconds a frame much be to be considered a hitch **/
-			const float HitchThreshold = RHIConfig::GetGPUHitchThreshold();
-			float ThisTime = Now - LastTime;
-			bool bHitched = (ThisTime > HitchThreshold) && LastTime > 0.0 && CurrentEventNodeFrame;
-			if (bHitched)
-			{
-				UE_LOG(LogD3D11RHI, Warning, TEXT("*******************************************************************************"));
-				UE_LOG(LogD3D11RHI, Warning, TEXT("********** Hitch detected on CPU, frametime = %6.1fms"),ThisTime * 1000.0f);
-				UE_LOG(LogD3D11RHI, Warning, TEXT("*******************************************************************************"));
-
-				for (int32 Frame = 0; Frame < GPUHitchEventNodeFrames.Num(); Frame++)
-				{
-					UE_LOG(LogD3D11RHI, Warning, TEXT(""));
-					UE_LOG(LogD3D11RHI, Warning, TEXT(""));
-					UE_LOG(LogD3D11RHI, Warning, TEXT("********** GPU Frame: Current - %d"),GPUHitchEventNodeFrames.Num() - Frame);
-					GPUHitchEventNodeFrames[Frame].DumpEventTree();
-				}
-				UE_LOG(LogD3D11RHI, Warning, TEXT(""));
-				UE_LOG(LogD3D11RHI, Warning, TEXT(""));
-				UE_LOG(LogD3D11RHI, Warning, TEXT("********** GPU Frame: Current"));
-				CurrentEventNodeFrame->DumpEventTree();
-
-				UE_LOG(LogD3D11RHI, Warning, TEXT("*******************************************************************************"));
-				UE_LOG(LogD3D11RHI, Warning, TEXT("********** End Hitch GPU Profile"));
-				UE_LOG(LogD3D11RHI, Warning, TEXT("*******************************************************************************"));
-				if (GEngine->GameViewport)
-				{
-					GEngine->GameViewport->Exec( NULL, TEXT("SCREENSHOT"), *GLog);
-				}
-
-				GPUHitchDebounce = 5; // don't trigger this again for a while
-				GPUHitchEventNodeFrames.Empty(); // clear history
-			}
-			else if (CurrentEventNodeFrame) // this will be null for discarded frames while recovering from a recent hitch
-			{
-				/** How many old frames to buffer for hitch reports **/
-				static const int32 HitchHistorySize = 4;
-
-				if (GPUHitchEventNodeFrames.Num() >= HitchHistorySize)
-				{
-					GPUHitchEventNodeFrames.RemoveAt(0);
-				}
-				GPUHitchEventNodeFrames.Add((FD3D11EventNodeFrame*)CurrentEventNodeFrame);
-				CurrentEventNodeFrame = NULL;  // prevent deletion of this below; ke kept it in the history
-			}
-		}
-		LastTime = Now;
-	}
-	bTrackingEvents = false;
-	bTrackingGPUCrashData = false;
-	delete CurrentEventNodeFrame;
-	CurrentEventNodeFrame = NULL;
-}
-
-float FD3D11EventNode::GetTiming()
-{
-	float Result = 0;
-
-	if (Timing.IsSupported())
-	{
-		// Get the timing result and block the CPU until it is ready
-		const uint64 GPUTiming = Timing.GetTiming(true);
-		const uint64 GPUFreq = Timing.GetTimingFrequency();
-
-		Result = double(GPUTiming) / double(GPUFreq);
-	}
-
-	return Result;
-}
-
-FD3DGPUProfiler::FD3DGPUProfiler(class FD3D11DynamicRHI* InD3DRHI)
-	: FGPUProfiler()
-	, FrameTiming(InD3DRHI, 4)
-	, D3D11RHI(InD3DRHI)
-{
-	// Initialize Buffered timestamp queries 
-	FrameTiming.InitResource(FRHICommandListImmediate::Get());
-
-	BeginFrame(InD3DRHI);
-}
-
-void FD3DGPUProfiler::PushEvent(const TCHAR* Name, FColor Color)
-{
-	FGPUProfiler::PushEvent(Name, Color);
-}
-
-void FD3DGPUProfiler::PopEvent()
-{
-	FGPUProfiler::PopEvent();
-}
-
-/** Start this frame of per tracking */
-void FD3D11EventNodeFrame::StartFrame()
-{
-	EventTree.Reset();
-	DisjointQuery.StartTracking();
-	RootEventTiming.StartTiming();
-}
-
-/** End this frame of per tracking, but do not block yet */
-void FD3D11EventNodeFrame::EndFrame()
-{
-	RootEventTiming.EndTiming();
-	DisjointQuery.EndTracking();
-}
-
-float FD3D11EventNodeFrame::GetRootTimingResults()
-{
-	double RootResult = 0.0f;
-	if (RootEventTiming.IsSupported())
-	{
-		const uint64 GPUTiming = RootEventTiming.GetTiming(true);
-		const uint64 GPUFreq = RootEventTiming.GetTimingFrequency();
-
-		RootResult = double(GPUTiming) / double(GPUFreq);
-	}
-
-	return (float)RootResult;
-}
-
-void FD3D11EventNodeFrame::LogDisjointQuery()
-{
-	if (!DisjointQuery.IsResultValid())
-	{
-		UE_LOG(LogRHI, Warning, TEXT("%s"), TEXT("Profiled range was disjoint!  GPU switched to doing something else while profiling."));
-	}
-}
-
-#endif // (RHI_NEW_GPU_PROFILER == 0)
-
 static void D3D11UpdateBufferStatsCommon(ID3D11Buffer* Buffer, int64 BufferSize, bool bAllocating)
 {
 	// this is a work-around on Windows. Due to the fact that there is no way
@@ -681,19 +444,11 @@ ID3D11ShaderResourceView* FD3D11DynamicRHI::RHIGetShaderResourceView(FRHITexture
 	return D3D11Texture->GetShaderResourceView();
 }
 
-void FD3D11DynamicRHI::RHIRegisterWork(uint32 NumPrimitives)
-{
-#if (RHI_NEW_GPU_PROFILER == 0)
-	RegisterGPUWork(NumPrimitives);
-#endif
-}
-
 void FD3D11DynamicRHI::RHIVerifyResult(ID3D11Device* Device, HRESULT Result, const ANSICHAR* Code, const ANSICHAR* Filename, uint32 Line) const
 {
 	VerifyD3D11Result(Result, Code, Filename, Line, Device);
 }
 
-#if RHI_NEW_GPU_PROFILER
 FD3D11RenderQuery* FD3D11DynamicRHI::InsertProfilerTimestamp(uint64* Target, bool bEndFrame)
 {
 	TArray<FD3D11RenderQuery*>& Pool = bEndFrame
@@ -716,4 +471,3 @@ FD3D11RenderQuery* FD3D11DynamicRHI::InsertProfilerTimestamp(uint64* Target, boo
 	Query->End(Direct3DDeviceIMContext, Target);
 	return Query;
 }
-#endif // RHI_NEW_GPU_PROFILER
diff --git a/Engine/Source/Runtime/Windows/D3D11RHI/Private/D3D11RHIPrivate.h b/Engine/Source/Runtime/Windows/D3D11RHI/Private/D3D11RHIPrivate.h
--- a/Engine/Source/Runtime/Windows/D3D11RHI/Private/D3D11RHIPrivate.h
+++ b/Engine/Source/Runtime/Windows/D3D11RHI/Private/D3D11RHIPrivate.h
@@ -102,209 +102,6 @@ struct FD3D11GlobalStats
 	static int64 GTotalGraphicsMemory;
 };
 
-#if (RHI_NEW_GPU_PROFILER == 0)
-
-// This class has multiple inheritance but really FGPUTiming is a static class
-class FD3D11BufferedGPUTiming : public FRenderResource, public FGPUTiming
-{
-public:
-	/**
-	 * Constructor.
-	 *
-	 * @param InD3DRHI			RHI interface
-	 * @param InBufferSize		Number of buffered measurements
-	 */
-	FD3D11BufferedGPUTiming(class FD3D11DynamicRHI* InD3DRHI, int32 BufferSize);
-
-	/**
-	 * Start a GPU timing measurement.
-	 */
-	void	StartTiming();
-
-	/**
-	 * End a GPU timing measurement.
-	 * The timing for this particular measurement will be resolved at a later time by the GPU.
-	 */
-	void	EndTiming();
-
-	/**
-	 * Retrieves the most recently resolved timing measurement.
-	 * The unit is the same as for FPlatformTime::Cycles(). Returns 0 if there are no resolved measurements.
-	 *
-	 * @return	Value of the most recently resolved timing, or 0 if no measurements have been resolved by the GPU yet.
-	 */
-	uint64	GetTiming(bool bGetCurrentResultsAndBlock = false);
-
-	/**
-	 * Initializes all D3D resources.
-	 */
-	virtual void InitRHI(FRHICommandListBase& RHICmdList) override;
-
-	/**
-	 * Releases all D3D resources.
-	 */
-	virtual void ReleaseRHI() override;
-
-	static void CalibrateTimers(FD3D11DynamicRHI* InD3DRHI);
-
-private:
-	/**
-	 * Initializes the static variables, if necessary.
-	 */
-	static void PlatformStaticInitialize(void* UserData);
-
-	/** RHI interface */
-	FD3D11DynamicRHI*			D3DRHI;
-	/** Number of timestamps created in 'StartTimestamps' and 'EndTimestamps'. */
-	int32						BufferSize;
-	/** Current timing being measured on the CPU. */
-	int32						CurrentTimestamp;
-	/** Number of measurements in the buffers (0 - BufferSize). */
-	int32						NumIssuedTimestamps;
-	/** Timestamps for all StartTimings. */
-	TRefCountPtr<ID3D11Query>*	StartTimestamps;
-	/** Timestamps for all EndTimings. */
-	TRefCountPtr<ID3D11Query>*	EndTimestamps;
-	/** Whether we are currently timing the GPU: between StartTiming() and EndTiming(). */
-	bool						bIsTiming;
-};
-
-/** Used to track whether a period was disjoint on the GPU, which means GPU timings are invalid. */
-class FD3D11DisjointTimeStampQuery : public FRenderResource
-{
-public:
-	FD3D11DisjointTimeStampQuery(class FD3D11DynamicRHI* InD3DRHI);
-
-	void StartTracking();
-	void EndTracking();
-	bool IsResultValid();
-	D3D11_QUERY_DATA_TIMESTAMP_DISJOINT GetResult();
-
-	/**
-	 * Initializes all D3D resources.
-	 */
-	virtual void InitRHI(FRHICommandListBase& RHICmdList) override;
-
-	/**
-	 * Releases all D3D resources.
-	 */
-	virtual void ReleaseRHI() override;
-
-
-private:
-
-	TRefCountPtr<ID3D11Query> DisjointQuery;
-
-	FD3D11DynamicRHI* D3DRHI;
-};
-
-/** A single perf event node, which tracks information about a appBeginDrawEvent/appEndDrawEvent range. */
-class FD3D11EventNode : public FGPUProfilerEventNode
-{
-public:
-	FD3D11EventNode(const TCHAR* InName, FGPUProfilerEventNode* InParent, class FD3D11DynamicRHI* InRHI) :
-		FGPUProfilerEventNode(InName, InParent),
-		Timing(InRHI, 1)
-	{
-		// Initialize Buffered timestamp queries 
-		Timing.InitRHI(FRHICommandListExecutor::GetImmediateCommandList()); // can't do this from the RHI thread
-	}
-
-	virtual ~FD3D11EventNode()
-	{
-		Timing.ReleaseRHI();  // can't do this from the RHI thread
-	}
-
-	/** 
-	 * Returns the time in ms that the GPU spent in this draw event.  
-	 * This blocks the CPU if necessary, so can cause hitching.
-	 */
-	virtual float GetTiming() override;
-
-
-	virtual void StartTiming() override
-	{
-		Timing.StartTiming();
-	}
-
-	virtual void StopTiming() override
-	{
-		Timing.EndTiming();
-	}
-
-	FD3D11BufferedGPUTiming Timing;
-};
-
-/** An entire frame of perf event nodes, including ancillary timers. */
-class FD3D11EventNodeFrame : public FGPUProfilerEventNodeFrame
-{
-public:
-
-	FD3D11EventNodeFrame(class FD3D11DynamicRHI* InRHI) :
-		FGPUProfilerEventNodeFrame(),
-		RootEventTiming(InRHI, 1),
-		DisjointQuery(InRHI)
-	{
-		FRHICommandListBase& RHICmdList = FRHICommandListExecutor::GetImmediateCommandList();
-		RootEventTiming.InitRHI(RHICmdList);
-		DisjointQuery.InitRHI(RHICmdList);
-	}
-
-	~FD3D11EventNodeFrame()
-	{
-		RootEventTiming.ReleaseRHI();
-		DisjointQuery.ReleaseRHI();
-	}
-
-	/** Start this frame of per tracking */
-	virtual void StartFrame() override;
-
-	/** End this frame of per tracking, but do not block yet */
-	virtual void EndFrame() override;
-
-	/** Calculates root timing base frequency (if needed by this RHI) */
-	virtual float GetRootTimingResults() override;
-
-	virtual void LogDisjointQuery() override;
-
-	/** Timer tracking inclusive time spent in the root nodes. */
-	FD3D11BufferedGPUTiming RootEventTiming;
-
-	/** Disjoint query tracking whether the times reported by DumpEventTree are reliable. */
-	FD3D11DisjointTimeStampQuery DisjointQuery;
-};
-
-/** 
- * Encapsulates GPU profiling logic and data. 
- * There's only one global instance of this struct so it should only contain global data, nothing specific to a frame.
- */
-struct FD3DGPUProfiler : public FGPUProfiler
-{
-	/** Used to measure GPU time per frame. */
-	FD3D11BufferedGPUTiming FrameTiming;
-
-	class FD3D11DynamicRHI* D3D11RHI;
-
-	/** GPU hitch profile histories */
-	TIndirectArray<FD3D11EventNodeFrame> GPUHitchEventNodeFrames;
-
-	FD3DGPUProfiler(class FD3D11DynamicRHI* InD3DRHI);
-
-	virtual FGPUProfilerEventNode* CreateEventNode(const TCHAR* InName, FGPUProfilerEventNode* InParent) override
-	{
-		FD3D11EventNode* EventNode = new FD3D11EventNode(InName, InParent, D3D11RHI);
-		return EventNode;
-	}
-
-	virtual void PushEvent(const TCHAR* Name, FColor Color) override;
-	virtual void PopEvent() override;
-
-	void BeginFrame(class FD3D11DynamicRHI* InRHI);
-	void EndFrame();
-};
-
-#endif // (RHI_NEW_GPU_PROFILER == 0)
-
 struct FD3D11TransitionData
 {
 	bool bUAVBarrier;
@@ -610,7 +407,6 @@ public:
 	D3D11RHI_API virtual int64                 RHIGetResourceMemorySize(FRHITexture* InTexture) const final override;
 	D3D11RHI_API virtual ID3D11RenderTargetView* RHIGetRenderTargetView(FRHITexture* InTexture, int32 InMipIndex = 0, int32 InArraySliceIndex = -1) const final override;
 	D3D11RHI_API virtual ID3D11ShaderResourceView* RHIGetShaderResourceView(FRHITexture* InTexture) const final override;
-	D3D11RHI_API virtual void                  RHIRegisterWork(uint32 NumPrimitives) final override;
 	D3D11RHI_API virtual void                  RHIVerifyResult(ID3D11Device* Device, HRESULT Result, const ANSICHAR* Code, const ANSICHAR* Filename, uint32 Line) const final override;
 
 	D3D11RHI_API virtual void				  RHIGetDisplaysInformation(FDisplayInformationArray& OutDisplayInformation) final override;
@@ -811,18 +607,14 @@ protected:
 
 			} EndFrameQuery;
 
-		#if RHI_NEW_GPU_PROFILER
 			UE::RHI::GPUProfiler::FEventStream EventStream;
-		#endif
 
 		#if WITH_RHI_BREADCRUMBS
 			FRHIBreadcrumbAllocatorArray BreadcrumbAllocators{};
 		#endif
 
 			FFrame()
-		#if RHI_NEW_GPU_PROFILER
 				: EventStream(UE::RHI::GPUProfiler::FQueue(UE::RHI::GPUProfiler::FQueue::EType::Graphics, 0, 0))
-		#endif
 			{}
 		};
 
@@ -834,7 +626,6 @@ protected:
 
 	} Profiler;
 
-#if RHI_NEW_GPU_PROFILER
 	template <typename TEventType, typename... TArgs>
 	TEventType& EmplaceProfilerEvent(TArgs&&... Args)
 	{
@@ -852,7 +643,6 @@ protected:
 			StatEvent = {};
 		}
 	}
-#endif // RHI_NEW_GPU_PROFILER
 
 public:
 	struct FTimestampCalibration
@@ -903,18 +693,9 @@ protected:
 	bool bRenderDoc = false;
 
 public:
-#if (RHI_NEW_GPU_PROFILER == 0)
-	void RegisterGPUWork(uint32 NumPrimitives = 0, uint32 NumVertices = 0)	{ GPUProfilingData.RegisterGPUWork(NumPrimitives, NumVertices); }
-	void RegisterGPUDispatch(FIntVector GroupCount)	                        { GPUProfilingData.RegisterGPUDispatch(GroupCount); }
-#endif
-
 	inline const FD3D11Adapter& GetAdapter() const { return Adapter; }
 
 protected:
-#if (RHI_NEW_GPU_PROFILER == 0)
-	FD3DGPUProfiler GPUProfilingData;
-#endif
-
 	FD3D11Adapter Adapter;
 
 	D3D11RHI_API FD3D11Texture* CreateTextureFromResource(bool bTextureArray, bool bCubeTexture, EPixelFormat Format, ETextureCreateFlags TexCreateFlags, const FClearValueBinding& ClearValueBinding, ID3D11Texture2D* TextureResource);
@@ -1006,9 +787,6 @@ protected:
 	D3D11RHI_API bool SetupDisplayHDRMetaData();
 
 	D3D11RHI_API void UpdateMemoryStats();
-
-	friend struct FD3DGPUProfiler;
-
 };
 
 PRAGMA_ENABLE_DEPRECATION_WARNINGS
diff --git a/Engine/Source/Runtime/Windows/D3D11RHI/Private/D3D11RenderTarget.cpp b/Engine/Source/Runtime/Windows/D3D11RHI/Private/D3D11RenderTarget.cpp
--- a/Engine/Source/Runtime/Windows/D3D11RHI/Private/D3D11RenderTarget.cpp
+++ b/Engine/Source/Runtime/Windows/D3D11RHI/Private/D3D11RenderTarget.cpp
@@ -102,10 +102,6 @@ void FD3D11DynamicRHI::ResolveTextureUsingShader(
 				// Clear the destination texture.
 				if (bClearDestTexture)
 				{
-				#if (RHI_NEW_GPU_PROFILER == 0)
-					Context.RegisterGPUWork(0);
-				#endif
-
 					Context.Direct3DDeviceIMContext->ClearDepthStencilView(DestTextureDSV,D3D11_CLEAR_DEPTH | D3D11_CLEAR_STENCIL,0,0);
 				}
 
@@ -126,10 +122,6 @@ void FD3D11DynamicRHI::ResolveTextureUsingShader(
 				// Clear the destination texture.
 				if (bClearDestTexture)
 				{
-				#if (RHI_NEW_GPU_PROFILER == 0)
-					Context.RegisterGPUWork(0);
-				#endif
-
 					FLinearColor ClearColor(0,0,0,0);
 					Context.Direct3DDeviceIMContext->ClearRenderTargetView(DestTextureRTV,(float*)&ClearColor);
 				}
@@ -878,10 +870,6 @@ void FD3D11DynamicRHI::RHIEndRenderPass()
 
 void FD3D11DynamicRHI::ResolveTexture(UE::RHICore::FResolveTextureInfo Info)
 {
-#if (RHI_NEW_GPU_PROFILER == 0)
-	RegisterGPUWork();
-#endif
-
 	FD3D11Texture* SourceTexture      = ResourceCast(Info.SourceTexture);
 	const FRHITextureDesc& SourceDesc = SourceTexture->GetDesc();
 
diff --git a/Engine/Source/Runtime/Windows/D3D11RHI/Private/D3D11Texture.cpp b/Engine/Source/Runtime/Windows/D3D11RHI/Private/D3D11Texture.cpp
--- a/Engine/Source/Runtime/Windows/D3D11RHI/Private/D3D11Texture.cpp
+++ b/Engine/Source/Runtime/Windows/D3D11RHI/Private/D3D11Texture.cpp
@@ -1683,10 +1683,6 @@ void FD3D11DynamicRHI::RHICopyTexture(FRHITexture* SourceTextureRHI, FRHITexture
 
 	check(SourceTexture && DestTexture);
 
-#if (RHI_NEW_GPU_PROFILER == 0)
-	RegisterGPUWork();
-#endif
-
 	const FRHITextureDesc& SourceDesc = SourceTextureRHI->GetDesc();
 	const FRHITextureDesc& DestDesc = DestTextureRHI->GetDesc();
 
@@ -1757,10 +1753,6 @@ void FD3D11DynamicRHI::RHICopyBufferRegion(FRHIBuffer* DstBuffer, uint64 DstOffs
 	check(DstBufferD3D11 && SrcBufferD3D11);
 	check(DstOffset + NumBytes <= DstBuffer->GetSize() && SrcOffset + NumBytes <= SrcBuffer->GetSize());
 
-#if (RHI_NEW_GPU_PROFILER == 0)
-	RegisterGPUWork();
-#endif
-
 	D3D11_BOX SrcBox;
 	SrcBox.left = SrcOffset;
 	SrcBox.right = SrcOffset + NumBytes;
diff --git a/Engine/Source/Runtime/Windows/D3D11RHI/Private/D3D11UAV.cpp b/Engine/Source/Runtime/Windows/D3D11RHI/Private/D3D11UAV.cpp
--- a/Engine/Source/Runtime/Windows/D3D11RHI/Private/D3D11UAV.cpp
+++ b/Engine/Source/Runtime/Windows/D3D11RHI/Private/D3D11UAV.cpp
@@ -199,9 +199,6 @@ void FD3D11DynamicRHI::ClearUAV(TRHICommandList_RecursiveHazardous<FD3D11Dynamic
 			RHICmdList.RunOnContext([UnorderedAccessView, ClearValues](auto& Context)
 			{
 				Context.Direct3DDeviceIMContext->ClearUnorderedAccessViewUint(UnorderedAccessView->View, *reinterpret_cast<const UINT(*)[4]>(ClearValues));
-			#if (RHI_NEW_GPU_PROFILER == 0)
-				Context.RegisterGPUWork(1);
-			#endif
 			});
 		}
 		else
diff --git a/Engine/Source/Runtime/Windows/D3D11RHI/Public/ID3D11DynamicRHI.h b/Engine/Source/Runtime/Windows/D3D11RHI/Public/ID3D11DynamicRHI.h
--- a/Engine/Source/Runtime/Windows/D3D11RHI/Public/ID3D11DynamicRHI.h
+++ b/Engine/Source/Runtime/Windows/D3D11RHI/Public/ID3D11DynamicRHI.h
@@ -33,9 +33,12 @@ struct ID3D11DynamicRHI : public FDynamicRHIPSOFallback
 	virtual ID3D11RenderTargetView*   RHIGetRenderTargetView(FRHITexture* InTexture, int32 InMipIndex = 0, int32 InArraySliceIndex = -1) const = 0;
 	virtual ID3D11ShaderResourceView* RHIGetShaderResourceView(FRHITexture* InTexture) const = 0;
 
-	virtual void                  RHIRegisterWork(uint32 NumPrimitives) = 0;
-
 	virtual void                  RHIVerifyResult(ID3D11Device* Device, HRESULT Result, const ANSICHAR* Code, const ANSICHAR* Filename, uint32 Line) const = 0;
+
+	UE_DEPRECATED(5.8, "RHIRegisterWork is deprecated, and there is no replacement. Remove calls to this function.")
+	void RHIRegisterWork(uint32 NumPrimitives)
+	{
+	}
 };
 
 inline bool IsRHID3D11()
