Remove FD3D11EventQuery
- Only used to implement the FD3D11Viewport IssueFrameEvent / WaitForFrameEventCompletion functions
- Replacement uses the same mechanism as FD3D11GPUFence, which has been broken out into its own type: FD3D11Sync.
- Also allows us to remove FScopedD3D11RHIThreadStaller. There are no more RHI thread stalls left in D3D11RHI.

Also cleaned up FD3D11Viewport constructors

#Synced-CL 48676711

diff --git a/Engine/Source/Runtime/Windows/D3D11RHI/Private/D3D11Commands.cpp b/Engine/Source/Runtime/Windows/D3D11RHI/Private/D3D11Commands.cpp
--- a/Engine/Source/Runtime/Windows/D3D11RHI/Private/D3D11Commands.cpp
+++ b/Engine/Source/Runtime/Windows/D3D11RHI/Private/D3D11Commands.cpp
@@ -1398,7 +1398,7 @@ void FD3D11DynamicRHI::RHISubmitCommandLists(FRHISubmitCommandListsArgs&& Args)
 {
 	// Attempt to readback completed queries and fences
 	PollQueryResults();
-	FD3D11GPUFence::PollFences();
+	FD3D11Sync::PollSyncs();
 
 #if WITH_RHI_BREADCRUMBS
 	for (IRHIPlatformCommandList* CmdList : Args.CommandLists)
@@ -1643,48 +1643,25 @@ void FD3D11DynamicRHI::RHIUnlockStagingBuffer(FRHIStagingBuffer* StagingBufferRH
 }
 
 
-TQueue<FD3D11GPUFence::FD3D11Sync, EQueueMode::SingleThreaded> FD3D11GPUFence::ActiveSyncs;
+TSpscQueue<FD3D11Sync> FD3D11Sync::ActiveSyncs;
 
-FD3D11GPUFence::FD3D11GPUFence(FName InName)
-	: FRHIGPUFence(InName)
-{
-}
-
-void FD3D11GPUFence::Clear()
+void FD3D11Sync::EnqueueSync(FGraphEventRef Event)
 {
-	Event = nullptr;
-}
+	FD3D11DynamicRHI& RHI = FD3D11DynamicRHI::Get();
 
-bool FD3D11GPUFence::Poll() const
-{
-	return Event && Event->IsComplete();
-}
+	// Insert an event query on the device context
+	D3D11_QUERY_DESC Desc{};
+	Desc.Query = D3D11_QUERY_EVENT;
 
-void FD3D11GPUFence::Wait(FRHICommandListImmediate& RHICmdList, FRHIGPUMask GPUMask) const
-{
-	if (Event && !Event->IsComplete())
-	{
-		//
-		// The fence might get signalled by an earlier RHI command polling them, but we can't be sure that will happen.
-		// The GPU might finish work after the RHI thread has gone idle, and then we'll never see the fence complete.
-		//
-		// Enqueue a command here that will block and wait for the fence if it still hasn't signalled by the time
-		// the RHI thread is done with all prior commands.
-		//
-		RHICmdList.EnqueueLambda([Event = Event](FRHICommandListImmediate&)
-		{
-			if (!Event->IsComplete())
-			{
-				PollFencesUntil(Event);
-			}
-		});
-		RHICmdList.ImmediateFlush(EImmediateFlushType::DispatchToRHIThread);
+	TRefCountPtr<ID3D11Query> Query;
+	VERIFYD3D11RESULT(RHI.GetDevice()->CreateQuery(&Desc, Query.GetInitReference()));
+	RHI.GetDeviceContext()->End(Query);
 
-		Event->Wait();
-	}
+	// Store the query in the list of queries to poll
+	ActiveSyncs.Enqueue(FD3D11Sync(MoveTemp(Event), MoveTemp(Query)));
 }
 
-void FD3D11GPUFence::PollFencesUntil(FGraphEvent* Target)
+void FD3D11Sync::PollSyncsUntil(FGraphEvent* Target)
 {
 	FD3D11DynamicRHI& RHI = FD3D11DynamicRHI::Get();
 
@@ -1718,7 +1695,7 @@ void FD3D11GPUFence::PollFencesUntil(FGraphEvent* Target)
 
 		// The fence has completed. Signal the graph event and remove the node.
 		Sync->Event->DispatchSubsequents();
-		ActiveSyncs.Pop();
+		ActiveSyncs.Dequeue();
 
 		if (bFoundTarget)
 		{
@@ -1727,27 +1704,55 @@ void FD3D11GPUFence::PollFencesUntil(FGraphEvent* Target)
 		}
 	}
 
-	UE_CLOG(Target, LogD3D11RHI, Fatal, TEXT("FD3D11GPUFence::PollFencesUntil - Attempt to poll for a specific fence (0x%p), but it was not found in the queue."), Target);
+	UE_CLOG(Target, LogD3D11RHI, Fatal, TEXT("FD3D11Sync::PollSyncsUntil - Attempt to poll for a specific fence (0x%p), but it was not found in the queue."), Target);
 }
 
-void FD3D11GPUFence::WriteGPUFence_TopOfPipe(FRHICommandListBase& RHICmdList)
+FD3D11GPUFence::FD3D11GPUFence(FName InName)
+	: FRHIGPUFence(InName)
 {
-	Event = FGraphEvent::CreateGraphEvent();
+}
 
-	RHICmdList.EnqueueLambda([Event = Event](FRHICommandListBase&) mutable
+void FD3D11GPUFence::Clear()
+{
+	Event = nullptr;
+}
+
+bool FD3D11GPUFence::Poll() const
+{
+	return Event && Event->IsComplete();
+}
+
+void FD3D11GPUFence::Wait(FRHICommandListImmediate& RHICmdList, FRHIGPUMask GPUMask) const
+{
+	if (Event && !Event->IsComplete())
 	{
-		FD3D11DynamicRHI& RHI = FD3D11DynamicRHI::Get();
+		//
+		// The fence might get signalled by an earlier RHI command polling them, but we can't be sure that will happen.
+		// The GPU might finish work after the RHI thread has gone idle, and then we'll never see the fence complete.
+		//
+		// Enqueue a command here that will block and wait for the fence if it still hasn't signalled by the time
+		// the RHI thread is done with all prior commands.
+		//
+		RHICmdList.EnqueueLambda([Event = Event](FRHICommandListImmediate&)
+		{
+			if (!Event->IsComplete())
+			{
+				FD3D11Sync::PollSyncsUntil(Event);
+			}
+		});
+		RHICmdList.ImmediateFlush(EImmediateFlushType::DispatchToRHIThread);
 
-		// Insert an event query on the device context
-		D3D11_QUERY_DESC Desc {};
-		Desc.Query = D3D11_QUERY_EVENT;
+		Event->Wait();
+	}
+}
 
-		TRefCountPtr<ID3D11Query> Query;
-		VERIFYD3D11RESULT(RHI.GetDevice()->CreateQuery(&Desc, Query.GetInitReference()));
-		RHI.GetDeviceContext()->End(Query);
+void FD3D11GPUFence::WriteGPUFence_TopOfPipe(FRHICommandListBase& RHICmdList)
+{
+	Event = FGraphEvent::CreateGraphEvent();
 
-		// Store the query in the list of queries to poll
-		ActiveSyncs.Enqueue(FD3D11Sync(MoveTemp(Event), MoveTemp(Query)));
+	RHICmdList.EnqueueLambda([Event = Event](FRHICommandListBase&) mutable
+	{
+		FD3D11Sync::EnqueueSync(Event);
 	});
 }
 
diff --git a/Engine/Source/Runtime/Windows/D3D11RHI/Private/D3D11Query.cpp b/Engine/Source/Runtime/Windows/D3D11RHI/Private/D3D11Query.cpp
--- a/Engine/Source/Runtime/Windows/D3D11RHI/Private/D3D11Query.cpp
+++ b/Engine/Source/Runtime/Windows/D3D11RHI/Private/D3D11Query.cpp
@@ -185,7 +185,7 @@ bool FD3D11RenderQuery::CacheResult(FD3D11DynamicRHI& RHI, bool bWait)
 
 	// Attempt to read the result from the GPU.
 	uint64 Temp;
-	if (!RHI.GetQueryData(Resource, &Temp, sizeof(Temp), Type == EType::Timestamp, /*bWait = */ bWait, /*bStallRHIThread = */ false))
+	if (!RHI.GetQueryData(Resource, &Temp, sizeof(Temp), Type == EType::Timestamp, /*bWait = */ bWait))
 	{
 		return false;
 	}
@@ -328,17 +328,10 @@ bool FD3D11DynamicRHI::PollQueryResultsForEndFrame(FD3D11RenderQuery* Query)
 	}
 }
 
-bool FD3D11DynamicRHI::GetQueryData(ID3D11Query* Query, void* Data, SIZE_T DataSize, bool bTimestamp, bool bWait, bool bStallRHIThread)
+bool FD3D11DynamicRHI::GetQueryData(ID3D11Query* Query, void* Data, SIZE_T DataSize, bool bTimestamp, bool bWait)
 {
 	// Request the data from the query.
-	HRESULT Result;
-
-	auto SafeGetQueryData = [&]()
-	{
-		FScopedD3D11RHIThreadStaller StallRHIThread(bStallRHIThread);
-		Result = Direct3DDeviceIMContext->GetData(Query, Data, DataSize, 0);
-	};
-	SafeGetQueryData();
+	HRESULT Result = Direct3DDeviceIMContext->GetData(Query, Data, DataSize, 0);
 
 	// Isn't the query finished yet, and can we wait for it?
 	if (Result == S_FALSE && bWait)
@@ -353,7 +346,7 @@ bool FD3D11DynamicRHI::GetQueryData(ID3D11Query* Query, void* Data, SIZE_T DataS
 
 		do
 		{
-			SafeGetQueryData();
+			Result = Direct3DDeviceIMContext->GetData(Query, Data, DataSize, 0);
 
 			if (Result == S_OK)
 			{
@@ -393,42 +386,6 @@ bool FD3D11DynamicRHI::GetQueryData(ID3D11Query* Query, void* Data, SIZE_T DataS
 	}
 }
 
-void FD3D11EventQuery::IssueEvent()
-{
-	if (ShouldNotEnqueueRHICommand())
-	{
-		D3DRHI->GetDeviceContext()->End(Query);
-	}
-	else
-	{
-		RunOnRHIThread(
-			[InQuery = Query]()
-		{
-			D3D11RHI_IMMEDIATE_CONTEXT->End(InQuery);
-		});
-	}
-}
-
-void FD3D11EventQuery::WaitForCompletion()
-{
-	BOOL bRenderingIsFinished = false;
-	while(
-		D3DRHI->GetQueryData(Query, &bRenderingIsFinished, sizeof(bRenderingIsFinished), false, true, true) &&
-		!bRenderingIsFinished
-		)
-	{};
-}
-
-FD3D11EventQuery::FD3D11EventQuery(class FD3D11DynamicRHI* InD3DRHI):
-	D3DRHI(InD3DRHI)
-{
-	D3D11_QUERY_DESC QueryDesc;
-	QueryDesc.Query = D3D11_QUERY_EVENT;
-	QueryDesc.MiscFlags = 0;
-	VERIFYD3D11RESULT_EX(D3DRHI->GetDevice()->CreateQuery(&QueryDesc,Query.GetInitReference()), D3DRHI->GetDevice());
-}
-
-
 TOptional<FD3D11DynamicRHI::FTimestampCalibration> FD3D11DynamicRHI::CalibrateTimers()
 {
 	// Attempt to generate a timestamp on GPU and CPU as closely to each other as possible.
diff --git a/Engine/Source/Runtime/Windows/D3D11RHI/Private/D3D11RHIPrivate.h b/Engine/Source/Runtime/Windows/D3D11RHI/Private/D3D11RHIPrivate.h
--- a/Engine/Source/Runtime/Windows/D3D11RHI/Private/D3D11RHIPrivate.h
+++ b/Engine/Source/Runtime/Windows/D3D11RHI/Private/D3D11RHIPrivate.h
@@ -221,10 +221,9 @@ public:
 	 * @param DataSize - The size of the buffer.
 	 * @param QueryType e.g. RQT_Occlusion or RQT_AbsoluteTime
 	 * @param bWait - If true, it will wait for the query to finish.
-	 * @param bStallRHIThread - if true, stall RHIT before accessing immediate context
 	 * @return true if the query finished.
 	 */
-	D3D11RHI_API bool GetQueryData(ID3D11Query* Query, void* Data, SIZE_T DataSize, bool bTimestamp, bool bWait, bool bStallRHIThread);
+	D3D11RHI_API bool GetQueryData(ID3D11Query* Query, void* Data, SIZE_T DataSize, bool bTimestamp, bool bWait);
 
 	D3D11RHI_API virtual FSamplerStateRHIRef RHICreateSamplerState(const FSamplerStateInitializerRHI& Initializer) final override;
 	D3D11RHI_API virtual FRasterizerStateRHIRef RHICreateRasterizerState(const FRasterizerStateInitializerRHI& Initializer) final override;
diff --git a/Engine/Source/Runtime/Windows/D3D11RHI/Private/D3D11Viewport.cpp b/Engine/Source/Runtime/Windows/D3D11RHI/Private/D3D11Viewport.cpp
--- a/Engine/Source/Runtime/Windows/D3D11RHI/Private/D3D11Viewport.cpp
+++ b/Engine/Source/Runtime/Windows/D3D11RHI/Private/D3D11Viewport.cpp
@@ -700,6 +700,21 @@ void* FD3D11Viewport::GetNativeBackBufferRT() const
 	return GetBackBuffer()->GetRenderTargetView(0, 0);
 }
 
+void FD3D11Viewport::WaitForFrameEventCompletion()
+{
+	if (FrameSyncEvent.IsValid() && !FrameSyncEvent->IsComplete())
+	{
+		TRACE_CPUPROFILER_EVENT_SCOPE(WaitForFrameEventCompletion);
+		FD3D11Sync::PollSyncsUntil(FrameSyncEvent);
+	}
+}
+
+void FD3D11Viewport::IssueFrameEvent()
+{
+	FrameSyncEvent = FGraphEvent::CreateGraphEvent();
+	FD3D11Sync::EnqueueSync(FrameSyncEvent);
+}
+
 /*=============================================================================
  *	The following RHI functions must be called from the main thread.
  *=============================================================================*/
diff --git a/Engine/Source/Runtime/Windows/D3D11RHI/Private/Windows/WindowsD3D11Viewport.cpp b/Engine/Source/Runtime/Windows/D3D11RHI/Private/Windows/WindowsD3D11Viewport.cpp
--- a/Engine/Source/Runtime/Windows/D3D11RHI/Private/Windows/WindowsD3D11Viewport.cpp
+++ b/Engine/Source/Runtime/Windows/D3D11RHI/Private/Windows/WindowsD3D11Viewport.cpp
@@ -30,25 +30,13 @@ static FAutoConsoleVariableRef CVarD3DUseAllowTearing(
 
 uint32 FD3D11Viewport::GSwapChainFlags = DXGI_SWAP_CHAIN_FLAG_ALLOW_MODE_SWITCH;
 
-FD3D11Viewport::FD3D11Viewport(FD3D11DynamicRHI* InD3DRHI,HWND InWindowHandle,uint32 InSizeX,uint32 InSizeY,bool bInIsFullscreen, EPixelFormat InPreferredPixelFormat):
-	D3DRHI(InD3DRHI),
-	LastFlipTime(0),
-	LastFrameComplete(0),
-	LastCompleteTime(0),
-	SyncCounter(0),
-	bSyncedLastFrame(false),
-	WindowHandle(InWindowHandle),
-	MaximumFrameLatency(3),
-	SizeX(InSizeX),
-	SizeY(InSizeY),
-	PresentFailCount(0),
-	ValidState(0),
-	PixelFormat(InPreferredPixelFormat),
-	DisplayColorGamut(EDisplayColorGamut::sRGB_D65),
-	DisplayOutputFormat(EDisplayOutputFormat::SDR_sRGB),
-	bIsFullscreen(bInIsFullscreen),
-	bAllowTearing(false),
-	FrameSyncEvent(InD3DRHI)
+FD3D11Viewport::FD3D11Viewport(FD3D11DynamicRHI* InD3DRHI, HWND InWindowHandle, uint32 InSizeX, uint32 InSizeY, bool bInIsFullscreen, EPixelFormat InPreferredPixelFormat)
+	: D3DRHI       (InD3DRHI)
+	, WindowHandle (InWindowHandle)
+	, PixelFormat  (InD3DRHI->GetDisplayFormat(InPreferredPixelFormat))
+	, SizeX        (InSizeX)
+	, SizeY        (InSizeY)
+	, bIsFullscreen(bInIsFullscreen)
 {
 	check(IsInGameThread());
 	D3DRHI->Viewports.Add(this);
@@ -56,8 +44,6 @@ FD3D11Viewport::FD3D11Viewport(FD3D11DynamicRHI* InD3DRHI,HWND InWindowHandle,ui
 	// Ensure that the D3D device has been created.
 	D3DRHI->InitD3DDevice();
 
-	PixelFormat = InD3DRHI->GetDisplayFormat(InPreferredPixelFormat);
-
 	TArray<IDXGISwapchainProvider*> DXGISwapchainProviderModules = IModularFeatures::Get().GetModularFeatureImplementations<IDXGISwapchainProvider>(IDXGISwapchainProvider::GetModularFeatureName());
 	IDXGISwapchainProvider* DXGISwapchainProvider = nullptr;
 	for (IDXGISwapchainProvider* ProviderModule : DXGISwapchainProviderModules)
@@ -291,13 +277,6 @@ FD3D11Viewport::FD3D11Viewport(FD3D11DynamicRHI* InD3DRHI,HWND InWindowHandle,ui
 	// Tell the window to redraw when they can.
 	// @todo: For Slate viewports, it doesn't make sense to post WM_PAINT messages (we swallow those.)
 	::PostMessageW( WindowHandle, WM_PAINT, 0, 0 );
-
-	ENQUEUE_RENDER_COMMAND(FD3D11Viewport)(
-		[this](FRHICommandListImmediate& RHICmdList)
-	{
-		// Initialize the query by issuing an initial event.
-		FrameSyncEvent.IssueEvent();
-	});	
 }
 
 static const FString GetDXGIColorSpaceString(DXGI_COLOR_SPACE_TYPE ColorSpace)
diff --git a/Engine/Source/Runtime/Windows/D3D11RHI/Public/D3D11Resources.h b/Engine/Source/Runtime/Windows/D3D11RHI/Public/D3D11Resources.h
--- a/Engine/Source/Runtime/Windows/D3D11RHI/Public/D3D11Resources.h
+++ b/Engine/Source/Runtime/Windows/D3D11RHI/Public/D3D11Resources.h
@@ -485,27 +485,37 @@ private:
 	uint32 ShadowBufferSize;
 };
 
-class FD3D11GPUFence final : public FRHIGPUFence
+class FD3D11Sync
 {
-private:
-	struct FD3D11Sync
+	// The graph event to trigger when the sync is resolved.
+	FGraphEventRef Event;
+
+	// The D3D11 API query object to poll / wait on.
+	TRefCountPtr<ID3D11Query> Query;
+
+	FD3D11Sync() = default;
+	FD3D11Sync(FGraphEventRef&& Event, TRefCountPtr<ID3D11Query>&& Query)
+		: Event(MoveTemp(Event))
+		, Query(MoveTemp(Query))
 	{
-		// The graph event to trigger when the sync is resolved.
-		FGraphEventRef Event;
+	}
 
-		// The D3D11 API query object to poll / wait on.
-		TRefCountPtr<ID3D11Query> Query;
+	static TSpscQueue<FD3D11Sync> ActiveSyncs;
 
-		FD3D11Sync() = default;
-		FD3D11Sync(FGraphEventRef Event, TRefCountPtr<ID3D11Query> Query)
-			: Event(MoveTemp(Event))
-			, Query(MoveTemp(Query))
-		{}
-	};
+public:
+	static void EnqueueSync(FGraphEventRef Event);
 
-	static TQueue<FD3D11Sync, EQueueMode::SingleThreaded> ActiveSyncs;
-	static void PollFencesUntil(FGraphEvent* Target);
+	static void PollSyncsUntil(FGraphEvent* Target);
 
+	static void PollSyncs()
+	{
+		PollSyncsUntil(nullptr);
+	}
+};
+
+class FD3D11GPUFence final : public FRHIGPUFence
+{
+private:
 	FGraphEventRef Event;
 
 public:
@@ -516,11 +526,6 @@ public:
 	void Wait(FRHICommandListImmediate& RHICmdList, FRHIGPUMask GPUMask) const override;
 
 	void WriteGPUFence_TopOfPipe(FRHICommandListBase& RHICmdList);
-
-	static void PollFences()
-	{
-		PollFencesUntil(nullptr);
-	}
 };
 
 namespace D3D11BufferStats
diff --git a/Engine/Source/Runtime/Windows/D3D11RHI/Public/D3D11Util.h b/Engine/Source/Runtime/Windows/D3D11RHI/Public/D3D11Util.h
--- a/Engine/Source/Runtime/Windows/D3D11RHI/Public/D3D11Util.h
+++ b/Engine/Source/Runtime/Windows/D3D11RHI/Public/D3D11Util.h
@@ -283,11 +283,3 @@ inline bool ShouldNotEnqueueRHICommand()
 	FRHICommandListImmediate& RHICmdList = FRHICommandListExecutor::GetImmediateCommandList();
 	return RHICmdList.Bypass() || (IsRunningRHIInSeparateThread() && IsInRHIThread()) || (!IsRunningRHIInSeparateThread() && IsInRenderingThread());
 }
-
-struct FScopedD3D11RHIThreadStaller : public FScopedRHIThreadStaller
-{
-	FScopedD3D11RHIThreadStaller(bool bDoStall = true)
-		: FScopedRHIThreadStaller(FRHICommandListExecutor::GetImmediateCommandList(), bDoStall && IsInRenderingThread() && GRHICommandList.AreRHITasksActive())
-	{
-	}
-};
diff --git a/Engine/Source/Runtime/Windows/D3D11RHI/Public/D3D11Viewport.h b/Engine/Source/Runtime/Windows/D3D11RHI/Public/D3D11Viewport.h
--- a/Engine/Source/Runtime/Windows/D3D11RHI/Public/D3D11Viewport.h
+++ b/Engine/Source/Runtime/Windows/D3D11RHI/Public/D3D11Viewport.h
@@ -13,25 +13,6 @@
 #include "Windows/D3D11ThirdParty.h"
 #include "DXGIUtilities.h"
 
-/** A D3D event query resource. */
-class FD3D11EventQuery
-{
-public:
-
-	/** Initialization constructor. */
-	D3D11RHI_API FD3D11EventQuery(class FD3D11DynamicRHI* InD3DRHI);
-
-	/** Issues an event for the query to poll. */
-	D3D11RHI_API void IssueEvent();
-
-	/** Waits for the event query to finish. */
-	D3D11RHI_API void WaitForCompletion();
-
-private:
-	FD3D11DynamicRHI* D3DRHI;
-	TRefCountPtr<ID3D11Query> Query;
-};
-
 class FD3D11Viewport : public FRHIViewport
 {
 public:
@@ -42,7 +23,6 @@ public:
 
 	};
 
-	FD3D11Viewport(class FD3D11DynamicRHI* InD3DRHI) : D3DRHI(InD3DRHI), PresentFailCount(0), ValidState (0), FrameSyncEvent(InD3DRHI) {}
 	D3D11RHI_API FD3D11Viewport(class FD3D11DynamicRHI* InD3DRHI, HWND InWindowHandle, uint32 InSizeX, uint32 InSizeY, bool bInIsFullscreen, EPixelFormat InPreferredPixelFormat);
 	D3D11RHI_API ~FD3D11Viewport();
 
@@ -70,15 +50,8 @@ public:
 	FIntPoint GetSizeXY() const { return FIntPoint(SizeX, SizeY); }
 	FD3D11Texture* GetBackBuffer() const { return BackBuffer; }
 
-	virtual void WaitForFrameEventCompletion() override
-	{
-		FrameSyncEvent.WaitForCompletion();
-	}
-
-	virtual void IssueFrameEvent() override
-	{
-		FrameSyncEvent.IssueEvent();
-	}
+	virtual void WaitForFrameEventCompletion() override;
+	virtual void IssueFrameEvent() override;
 
 	IDXGISwapChain* GetSwapChain() const { return SwapChain; }
 
@@ -123,24 +96,26 @@ protected:
 	void ShutdownHDR();
 
 
-	FD3D11DynamicRHI* D3DRHI;
-	uint64 LastFlipTime;
-	uint64 LastFrameComplete;
-	uint64 LastCompleteTime;
-	int32 SyncCounter;
-	bool bSyncedLastFrame;
-	HWND WindowHandle;
-	uint32 MaximumFrameLatency;
+	FD3D11DynamicRHI* const D3DRHI;
+	HWND const WindowHandle;
+
+	EPixelFormat PixelFormat;
 	uint32 SizeX;
 	uint32 SizeY;
-	uint32 BackBufferCount;
-	uint32 PresentFailCount;
-	TAtomic<uint32> ValidState;
-	EPixelFormat PixelFormat;
-	EDisplayColorGamut DisplayColorGamut;
-	EDisplayOutputFormat DisplayOutputFormat;
 	bool bIsFullscreen;
-	bool bAllowTearing;
+
+	uint64 LastFlipTime = 0;
+	uint64 LastFrameComplete = 0;
+	uint64 LastCompleteTime = 0;
+	int32 SyncCounter = 0;
+	bool bSyncedLastFrame = false;
+	uint32 MaximumFrameLatency = 3;
+	uint32 BackBufferCount = 0;
+	uint32 PresentFailCount = 0;
+	TAtomic<uint32> ValidState = 0;
+	EDisplayColorGamut DisplayColorGamut = EDisplayColorGamut::sRGB_D65;
+	EDisplayOutputFormat DisplayOutputFormat = EDisplayOutputFormat::SDR_sRGB;
+	bool bAllowTearing = false;
 
 	static D3D11RHI_API uint32 GSwapChainFlags;
 
@@ -157,7 +132,7 @@ protected:
 	bool						bNeedSwapChain;
 
 	/** An event used to track the GPU's progress. */
-	FD3D11EventQuery FrameSyncEvent;
+	FGraphEventRef FrameSyncEvent;
 
 	FCustomPresentRHIRef CustomPresent;
 
