Remove wrongly defined `IRefCountedObject` virtuals that weren't a) used and b) actually overriding the interface.

#Synced-CL 48778233

diff --git a/Engine/Source/Runtime/Windows/D3D11RHI/Public/D3D11Resources.h b/Engine/Source/Runtime/Windows/D3D11RHI/Public/D3D11Resources.h
--- a/Engine/Source/Runtime/Windows/D3D11RHI/Public/D3D11Resources.h
+++ b/Engine/Source/Runtime/Windows/D3D11RHI/Public/D3D11Resources.h
@@ -445,20 +445,6 @@ public:
 	void TakeOwnership(FD3D11Buffer& Other);
 	void ReleaseOwnership();
 
-	// IRefCountedObject interface.
-	virtual uint32 AddRef() const
-	{
-		return FRHIResource::AddRef();
-	}
-	virtual uint32 Release() const
-	{
-		return FRHIResource::Release();
-	}
-	virtual uint32 GetRefCount() const
-	{
-		return FRHIResource::GetRefCount();
-	}
-
 	ID3D11Resource* GetResource() const
 	{
 		return Resource.GetReference();
