Instantiate FD3D11DynamicRHI::ClearAllShaderResourcesForFrequency<SF_Compute> in D3D11RHI.cpp for D3D11Commands.cpp

#Synced-CL 49434565

diff --git a/Engine/Source/Runtime/Windows/D3D11RHI/Private/D3D11RHI.cpp b/Engine/Source/Runtime/Windows/D3D11RHI/Private/D3D11RHI.cpp
--- a/Engine/Source/Runtime/Windows/D3D11RHI/Private/D3D11RHI.cpp
+++ b/Engine/Source/Runtime/Windows/D3D11RHI/Private/D3D11RHI.cpp
@@ -307,6 +307,9 @@ void FD3D11DynamicRHI::ClearAllShaderResourcesForFrequency()
 	StateCache.ClearConstantBuffers<ShaderFrequency>();
 }
 
+// For D3D11Commands.cpp
+template void FD3D11DynamicRHI::ClearAllShaderResourcesForFrequency<SF_Compute>();
+
 void FD3D11DynamicRHI::ClearAllShaderResources()
 {
 	ClearAllShaderResourcesForFrequency<SF_Vertex>();
