Add RHIGetAvailableResolutionsForDisplay to get resolutions for a specific monitor on Windows

Also remove runtime checks around CreateDXGIFactory2, which were only necessary before Windows 8.1.

#Synced-CL 49434565

diff --git a/Engine/Source/Runtime/D3D12RHI/Private/D3D12Adapter.cpp b/Engine/Source/Runtime/D3D12RHI/Private/D3D12Adapter.cpp
--- a/Engine/Source/Runtime/D3D12RHI/Private/D3D12Adapter.cpp
+++ b/Engine/Source/Runtime/D3D12RHI/Private/D3D12Adapter.cpp
@@ -1713,37 +1713,20 @@ FD3D12Adapter::~FD3D12Adapter()
 #if D3D12_SUPPORTS_DXGI_DEBUG
 	FPlatformProcess::FreeDllHandle(DxgiDebugDllHandle);
 #endif
-
-#if PLATFORM_WINDOWS
-	FPlatformProcess::FreeDllHandle(DxgiDllHandle);
-#endif
 }
 
-void FD3D12Adapter::CreateDXGIFactory(TRefCountPtr<IDXGIFactory2>& DxgiFactory2, bool bWithDebug, HMODULE DxgiDllHandle)
+void FD3D12Adapter::CreateDXGIFactory(TRefCountPtr<IDXGIFactory2>& DxgiFactory2, bool bWithDebug)
 {
 #if PLATFORM_WINDOWS
-	typedef HRESULT(WINAPI FCreateDXGIFactory2)(UINT, REFIID, void**);
-
-	FCreateDXGIFactory2* CreateDXGIFactory2FnPtr = (FCreateDXGIFactory2*)(void*)::GetProcAddress(DxgiDllHandle, "CreateDXGIFactory2");
-
-	check(CreateDXGIFactory2FnPtr);
-
 	const uint32 Flags = bWithDebug ? DXGI_CREATE_FACTORY_DEBUG : 0;
-	VERIFYD3D12RESULT(CreateDXGIFactory2FnPtr(Flags, IID_PPV_ARGS(DxgiFactory2.GetInitReference())));
+	VERIFYD3D12RESULT(CreateDXGIFactory2(Flags, IID_PPV_ARGS(DxgiFactory2.GetInitReference())));
 #endif // PLATFORM_WINDOWS
 }
 
 void FD3D12Adapter::CreateDXGIFactory(bool bWithDebug)
 {
 #if PLATFORM_WINDOWS
-	HMODULE UsedDxgiDllHandle = (HMODULE)0;
-
-	// Dynamically load this otherwise Win7 fails to boot as it's missing on that DLL
-	DxgiDllHandle = (HMODULE)FPlatformProcess::GetDllHandle(TEXT("dxgi.dll"));
-	check(DxgiDllHandle);
-	UsedDxgiDllHandle = DxgiDllHandle;
-
-	CreateDXGIFactory(DxgiFactory2, bWithDebug, UsedDxgiDllHandle);
+	CreateDXGIFactory(DxgiFactory2, bWithDebug);
 
 	InitDXGIFactoryVariants(DxgiFactory2);
 #endif // PLATFORM_WINDOWS
diff --git a/Engine/Source/Runtime/D3D12RHI/Private/D3D12Adapter.h b/Engine/Source/Runtime/D3D12RHI/Private/D3D12Adapter.h
--- a/Engine/Source/Runtime/D3D12RHI/Private/D3D12Adapter.h
+++ b/Engine/Source/Runtime/D3D12RHI/Private/D3D12Adapter.h
@@ -280,7 +280,7 @@ public:
 	FORCEINLINE uint32 GetVRSTileSize() const { return VRSTileSize; }
 
 	void CreateDXGIFactory(bool bWithDebug);
-	static void CreateDXGIFactory(TRefCountPtr<IDXGIFactory2>& DxgiFactory2, bool bWithDebug, HMODULE DxgiDllHandle);
+	static void CreateDXGIFactory(TRefCountPtr<IDXGIFactory2>& DxgiFactory2, bool bWithDebug);
 	void InitDXGIFactoryVariants(IDXGIFactory2* InDxgiFactory2);
 	HRESULT EnumAdapters(IDXGIAdapter** TempAdapter) const;
 
@@ -497,10 +497,6 @@ public:
 	void SetResidencyPriority(ID3D12Pageable* Pageable, D3D12_RESIDENCY_PRIORITY HeapPriority,
 							  uint32 GPUIndex, TOptional<uint64> ResourceSize = {});
 
-#if PLATFORM_WINDOWS
-	HMODULE GetDxgiDllHandle() const { return DxgiDllHandle; };
-#endif
-
 #if PLATFORM_SUPPORTS_BINDLESS_RENDERING
 	FD3D12BindlessDescriptorAllocator& GetBindlessDescriptorAllocator() { return BindlessDescriptorAllocator; }
 #endif
@@ -600,10 +596,6 @@ protected:
 	HANDLE ExceptionHandlerHandle = INVALID_HANDLE_VALUE;
 #endif
 
-#if PLATFORM_WINDOWS
-	HMODULE DxgiDllHandle{};
-#endif
-
 	D3D_ROOT_SIGNATURE_VERSION RootSignatureVersion;
 	bool bDepthBoundsTestSupported = false;
 	bool bCopyQueueTimestampQueriesSupported = false;
diff --git a/Engine/Source/Runtime/D3D12RHI/Private/D3D12RHIPrivate.h b/Engine/Source/Runtime/D3D12RHI/Private/D3D12RHIPrivate.h
--- a/Engine/Source/Runtime/D3D12RHI/Private/D3D12RHIPrivate.h
+++ b/Engine/Source/Runtime/D3D12RHI/Private/D3D12RHIPrivate.h
@@ -417,6 +417,7 @@ public:
 	virtual void RHITick(float DeltaTime) final override;
 	virtual void RHIBlockUntilGPUIdle() final override;
 	virtual bool RHIGetAvailableResolutions(FScreenResolutionArray& Resolutions, bool bIgnoreRefreshRate) final override;
+	virtual bool RHIGetAvailableResolutionsForDisplay(FScreenResolutionArray& Resolutions, bool bIgnoreRefreshRate, void* NativeMonitorHandle) final override;
 	virtual void RHIGetSupportedResolution(uint32& Width, uint32& Height) final override;
 	virtual void* RHIGetNativeDevice() final override;
 	virtual void* RHIGetNativeGraphicsQueue() final override;
@@ -581,6 +582,7 @@ protected:
 
 #if PLATFORM_WINDOWS
 	TRefCountPtr<IDXGIFactory2> DXGIFactoryForDisplayList;
+	TRefCountPtr<IDXGIFactory2> DXGIFactoryForResolutions;
 #endif
 
 public:
diff --git a/Engine/Source/Runtime/D3D12RHI/Private/Windows/WindowsD3D12Device.cpp b/Engine/Source/Runtime/D3D12RHI/Private/Windows/WindowsD3D12Device.cpp
--- a/Engine/Source/Runtime/D3D12RHI/Private/Windows/WindowsD3D12Device.cpp
+++ b/Engine/Source/Runtime/D3D12RHI/Private/Windows/WindowsD3D12Device.cpp
@@ -590,7 +590,7 @@ bool FD3D12DynamicRHI::SetupDisplayHDRMetaData()
 	{
 		if (!DXGIFactoryForDisplayList.IsValid() || !DXGIFactoryForDisplayList->IsCurrent())
 		{
-			FD3D12Adapter::CreateDXGIFactory(DXGIFactoryForDisplayList, false, GetAdapter(0).GetDxgiDllHandle());
+			FD3D12Adapter::CreateDXGIFactory(DXGIFactoryForDisplayList, false);
 		}
 
 		if (DXGIFactoryForDisplayList.IsValid() && DXGIFactoryForDisplayList->IsCurrent())
@@ -2003,6 +2003,102 @@ bool FD3D12Device::GetRayTracingPipelineInfo(ID3D12StateObject* Pipeline, FD3D12
 
 #endif // D3D12_RHI_RAYTRACING
 
+static bool GetAvailableResolutionsForOutput(IDXGIOutput* Output, FScreenResolutionArray& Resolutions, bool bIgnoreRefreshRate)
+{
+	// TODO: GetDisplayModeList is a terribly SLOW call.  It can take up to a second per invocation.
+	//  We might want to work around some DXGI badness here.
+	DXGI_FORMAT Format = DXGI_FORMAT_R8G8B8A8_UNORM;
+	uint32 NumModes = 0;
+	HRESULT HResult = Output->GetDisplayModeList(Format, 0, &NumModes, nullptr);
+	if (HResult == DXGI_ERROR_NOT_FOUND)
+	{
+		return false;
+	}
+
+	if (HResult == DXGI_ERROR_NOT_CURRENTLY_AVAILABLE)
+	{
+		UE_LOG(LogD3D12RHI, Warning,
+			TEXT("RHIGetAvailableResolutions() can not be used over a remote desktop configuration")
+		);
+		return false;
+	}
+
+	if (FAILED(HResult))
+	{
+		UE_LOG(LogD3D12RHI, Warning, TEXT("RHIGetAvailableResolutions failed with unknown error (0x%x)."), HResult);
+		return false;
+	}
+
+	if (NumModes > 0)
+	{
+		TArray<DXGI_MODE_DESC> ModeList;
+		do
+		{
+			VERIFYD3D12RESULT(Output->GetDisplayModeList(Format, 0, &NumModes, nullptr));
+			ModeList.SetNum(NumModes);
+			HResult = Output->GetDisplayModeList(Format, 0, &NumModes, ModeList.GetData());
+		} while (HResult == DXGI_ERROR_MORE_DATA);
+		VERIFYD3D12RESULT(HResult);
+
+		constexpr int32 MinAllowableResolutionX = 0;
+		constexpr int32 MinAllowableResolutionY = 0;
+		constexpr int32 MaxAllowableResolutionX = 10480;
+		constexpr int32 MaxAllowableResolutionY = 10480;
+		constexpr int32 MinAllowableRefreshRate = 0;
+		constexpr int32 MaxAllowableRefreshRate = 10480;
+		for (uint32 m = 0; m < NumModes; m++)
+		{
+			CA_SUPPRESS(6385);
+			if (((int32)ModeList[m].Width >= MinAllowableResolutionX) &&
+				((int32)ModeList[m].Width <= MaxAllowableResolutionX) &&
+				((int32)ModeList[m].Height >= MinAllowableResolutionY) &&
+				((int32)ModeList[m].Height <= MaxAllowableResolutionY)
+				)
+			{
+				bool bAddIt = true;
+				if (bIgnoreRefreshRate == false)
+				{
+					if (((int32)ModeList[m].RefreshRate.Numerator < MinAllowableRefreshRate * ModeList[m].RefreshRate.Denominator) ||
+						((int32)ModeList[m].RefreshRate.Numerator > MaxAllowableRefreshRate * ModeList[m].RefreshRate.Denominator)
+						)
+					{
+						continue;
+					}
+				}
+				else
+				{
+					// See if it is in the list already
+					for (int32 CheckIndex = 0; CheckIndex < Resolutions.Num(); CheckIndex++)
+					{
+						FScreenResolutionRHI& CheckResolution = Resolutions[CheckIndex];
+						if ((CheckResolution.Width == ModeList[m].Width) &&
+							(CheckResolution.Height == ModeList[m].Height))
+						{
+							// Already in the list...
+							bAddIt = false;
+							break;
+						}
+					}
+				}
+
+				if (bAddIt)
+				{
+					FScreenResolutionRHI& ScreenResolution = Resolutions.Emplace_GetRef();
+					ScreenResolution.Width = ModeList[m].Width;
+					ScreenResolution.Height = ModeList[m].Height;
+					ScreenResolution.RefreshRate = ModeList[m].RefreshRate.Numerator / ModeList[m].RefreshRate.Denominator;
+				}
+			}
+		}
+	}
+	else
+	{
+		UE_LOG(LogD3D12RHI, Warning, TEXT("No display modes found for the standard format DXGI_FORMAT_R8G8B8A8_UNORM!"));
+	}
+
+	return !Resolutions.IsEmpty();
+}
+
 /**
  *	Retrieve available screen resolutions.
  *
@@ -2013,26 +2109,7 @@ bool FD3D12Device::GetRayTracingPipelineInfo(ID3D12StateObject* Pipeline, FD3D12
  */
 bool FD3D12DynamicRHI::RHIGetAvailableResolutions(FScreenResolutionArray& Resolutions, bool bIgnoreRefreshRate)
 {
-	int32 MinAllowableResolutionX = 0;
-	int32 MinAllowableResolutionY = 0;
-	int32 MaxAllowableResolutionX = 10480;
-	int32 MaxAllowableResolutionY = 10480;
-	int32 MinAllowableRefreshRate = 0;
-	int32 MaxAllowableRefreshRate = 10480;
-
-	if (MaxAllowableResolutionX == 0) //-V547
-	{
-		MaxAllowableResolutionX = 10480;
-	}
-	if (MaxAllowableResolutionY == 0) //-V547
-	{
-		MaxAllowableResolutionY = 10480;
-	}
-	if (MaxAllowableRefreshRate == 0) //-V547
-	{
-		MaxAllowableRefreshRate = 10480;
-	}
-
+	// TODO: The resolution information can be stale if the adapter's factory isn't current. Use RHIGetAvailableResolutionsForDisplay instead.
 	FD3D12Adapter& ChosenAdapter = GetAdapter();
 
 	HRESULT HResult = S_OK;
@@ -2056,104 +2133,52 @@ bool FD3D12DynamicRHI::RHIGetAvailableResolutions(FScreenResolutionArray& Resolu
 		return false;
 	}
 
-	int32 CurrentOutput = 0;
-	do
+	TRefCountPtr<IDXGIOutput> Output;
+	HResult = Adapter->EnumOutputs(0, Output.GetInitReference());
+	return SUCCEEDED(HResult) ? GetAvailableResolutionsForOutput(Output, Resolutions, bIgnoreRefreshRate) : false;
+}
+
+bool FD3D12DynamicRHI::RHIGetAvailableResolutionsForDisplay(FScreenResolutionArray& Resolutions, bool bIgnoreRefreshRate, void* NativeMonitorHandle)
+{
+	HMONITOR MonitorHandle = static_cast<HMONITOR>(NativeMonitorHandle);
+	if (MonitorHandle == NULL)
 	{
-		TRefCountPtr<IDXGIOutput> Output;
-		HResult = Adapter->EnumOutputs(CurrentOutput, Output.GetInitReference());
-		if (DXGI_ERROR_NOT_FOUND == HResult)
-		{
-			break;
-		}
-		if (FAILED(HResult))
-		{
-			return false;
-		}
+		return false;
+	}
 
-		// TODO: GetDisplayModeList is a terribly SLOW call.  It can take up to a second per invocation.
-		//  We might want to work around some DXGI badness here.
-		DXGI_FORMAT Format = DXGI_FORMAT_R8G8B8A8_UNORM;
-		uint32 NumModes = 0;
-		HResult = Output->GetDisplayModeList(Format, 0, &NumModes, nullptr);
-		if (HResult == DXGI_ERROR_NOT_FOUND)
-		{
-			++CurrentOutput;
-			continue;
-		}
-		else if (HResult == DXGI_ERROR_NOT_CURRENTLY_AVAILABLE)
+	if (!DXGIFactoryForResolutions.IsValid() || !DXGIFactoryForResolutions->IsCurrent())
+	{
+		FD3D12Adapter::CreateDXGIFactory(DXGIFactoryForResolutions, false);
+	}
+
+	HRESULT HResult;
+	TRefCountPtr<IDXGIAdapter> Adapter;
+	for (int32 CurrentAdapter = 0; (HResult = DXGIFactoryForResolutions->EnumAdapters(CurrentAdapter, Adapter.GetInitReference())) != DXGI_ERROR_NOT_FOUND; ++CurrentAdapter)
+	{
+		if (FAILED(HResult))
 		{
-			UE_LOG(LogD3D12RHI, Warning,
-				TEXT("RHIGetAvailableResolutions() can not be used over a remote desktop configuration")
-				);
 			return false;
 		}
 
-		if (NumModes > 0)
+		TRefCountPtr<IDXGIOutput> Output;
+		for (int32 CurrentOutput = 0; (HResult = Adapter->EnumOutputs(CurrentOutput, Output.GetInitReference())) != DXGI_ERROR_NOT_FOUND; ++CurrentOutput)
 		{
-			DXGI_MODE_DESC* ModeList = new DXGI_MODE_DESC[NumModes];
-			VERIFYD3D12RESULT(Output->GetDisplayModeList(Format, 0, &NumModes, ModeList));
-
-			for (uint32 m = 0; m < NumModes; m++)
+			if (FAILED(HResult))
 			{
-				CA_SUPPRESS(6385);
-				if (((int32)ModeList[m].Width >= MinAllowableResolutionX) &&
-					((int32)ModeList[m].Width <= MaxAllowableResolutionX) &&
-					((int32)ModeList[m].Height >= MinAllowableResolutionY) &&
-					((int32)ModeList[m].Height <= MaxAllowableResolutionY)
-					)
-				{
-					bool bAddIt = true;
-					if (bIgnoreRefreshRate == false)
-					{
-						if (((int32)ModeList[m].RefreshRate.Numerator < MinAllowableRefreshRate * ModeList[m].RefreshRate.Denominator) ||
-							((int32)ModeList[m].RefreshRate.Numerator > MaxAllowableRefreshRate * ModeList[m].RefreshRate.Denominator)
-							)
-						{
-							continue;
-						}
-					}
-					else
-					{
-						// See if it is in the list already
-						for (int32 CheckIndex = 0; CheckIndex < Resolutions.Num(); CheckIndex++)
-						{
-							FScreenResolutionRHI& CheckResolution = Resolutions[CheckIndex];
-							if ((CheckResolution.Width == ModeList[m].Width) &&
-								(CheckResolution.Height == ModeList[m].Height))
-							{
-								// Already in the list...
-								bAddIt = false;
-								break;
-							}
-						}
-					}
-
-					if (bAddIt)
-					{
-						// Add the mode to the list
-						int32 Temp2Index = Resolutions.AddZeroed();
-						FScreenResolutionRHI& ScreenResolution = Resolutions[Temp2Index];
-
-						ScreenResolution.Width = ModeList[m].Width;
-						ScreenResolution.Height = ModeList[m].Height;
-						ScreenResolution.RefreshRate = ModeList[m].RefreshRate.Numerator / ModeList[m].RefreshRate.Denominator;
-					}
-				}
+				return false;
 			}
 
-			delete[] ModeList;
-		}
-		else
-		{
-			UE_LOG(LogD3D12RHI, Warning, TEXT("No display modes found for the standard format DXGI_FORMAT_R8G8B8A8_UNORM!"));
+			DXGI_OUTPUT_DESC OutputDesc{};
+			HRESULT OutputDescRes = Output->GetDesc(&OutputDesc);
+			check(SUCCEEDED(OutputDescRes));
+			if (MonitorHandle == OutputDesc.Monitor)
+			{
+				return GetAvailableResolutionsForOutput(Output, Resolutions, bIgnoreRefreshRate);
+			}
 		}
+	}
 
-		++CurrentOutput;
-
-	// TODO: Cap at 1 for default output
-	} while (CurrentOutput < 1);
-
-	return Resolutions.Num() > 0;
+	return false;
 }
 
 void FWindowsD3D12Adapter::CreateCommandSignatures()
diff --git a/Engine/Source/Runtime/RHI/Public/DynamicRHI.h b/Engine/Source/Runtime/RHI/Public/DynamicRHI.h
--- a/Engine/Source/Runtime/RHI/Public/DynamicRHI.h
+++ b/Engine/Source/Runtime/RHI/Public/DynamicRHI.h
@@ -758,6 +758,21 @@ public:
 	// FlushType: Thread safe
 	virtual bool RHIGetAvailableResolutions(FScreenResolutionArray& Resolutions, bool bIgnoreRefreshRate) = 0;
 
+	/**
+	*	Retrieve available screen resolutions for the display.
+	*
+	*	@param	Resolutions			TArray<FScreenResolutionRHI> parameter that will be filled in.
+	*	@param	bIgnoreRefreshRate	If true, ignore refresh rates.
+	*	@param	NativeMonitorHandle	Platform-specific handle for the monitor.
+	*
+	*	@return	bool				true if successfully filled the array
+	*/
+	// FlushType: Thread safe
+	virtual bool RHIGetAvailableResolutionsForDisplay(FScreenResolutionArray& Resolutions, bool bIgnoreRefreshRate, void* NativeMonitorHandle)
+	{
+		return RHIGetAvailableResolutions(Resolutions, bIgnoreRefreshRate);
+	}
+
 	/**
 	* Returns a supported screen resolution that most closely matches input.
 	* @param Width - Input: Desired resolution width in pixels. Output: A width that the platform supports.
@@ -1319,6 +1334,11 @@ inline bool RHIGetAvailableResolutions(FScreenResolutionArray& Resolutions, bool
 	return GDynamicRHI->RHIGetAvailableResolutions(Resolutions, bIgnoreRefreshRate);
 }
 
+inline bool RHIGetAvailableResolutionsForDisplay(FScreenResolutionArray& Resolutions, bool bIgnoreRefreshRate, void* NativeMonitorHandle)
+{
+	return GDynamicRHI->RHIGetAvailableResolutionsForDisplay(Resolutions, bIgnoreRefreshRate, NativeMonitorHandle);
+}
+
 inline void RHIGetSupportedResolution(uint32& Width, uint32& Height)
 {
 	GDynamicRHI->RHIGetSupportedResolution(Width, Height);
diff --git a/Engine/Source/Runtime/Windows/D3D11RHI/Private/D3D11RHIPrivate.h b/Engine/Source/Runtime/Windows/D3D11RHI/Private/D3D11RHIPrivate.h
--- a/Engine/Source/Runtime/Windows/D3D11RHI/Private/D3D11RHIPrivate.h
+++ b/Engine/Source/Runtime/Windows/D3D11RHI/Private/D3D11RHIPrivate.h
@@ -299,6 +299,7 @@ public:
 	D3D11RHI_API virtual void RHITick(float DeltaTime) final override;
 	D3D11RHI_API virtual void RHIBlockUntilGPUIdle() final override;
 	D3D11RHI_API virtual bool RHIGetAvailableResolutions(FScreenResolutionArray& Resolutions, bool bIgnoreRefreshRate) final override;
+	D3D11RHI_API virtual bool RHIGetAvailableResolutionsForDisplay(FScreenResolutionArray& Resolutions, bool bIgnoreRefreshRate, void* NativeMonitorHandle) final override;
 	D3D11RHI_API virtual void RHIGetSupportedResolution(uint32& Width, uint32& Height) final override;
 	D3D11RHI_API virtual void* RHIGetNativeDevice() final override;
 	D3D11RHI_API virtual void* RHIGetNativeCommandBuffer() final override;
@@ -493,6 +494,7 @@ protected:
 	TRefCountPtr<IDXGIFactory1> DXGIFactory1;
 #if PLATFORM_WINDOWS
 	TRefCountPtr<IDXGIFactory1> DXGIFactoryForDisplayList;
+	TRefCountPtr<IDXGIFactory1> DXGIFactoryForResolutions;
 #endif
 
 	// Whether HDR is available from the particular DXGI factories available
diff --git a/Engine/Source/Runtime/Windows/D3D11RHI/Private/Windows/WindowsD3D11Device.cpp b/Engine/Source/Runtime/Windows/D3D11RHI/Private/Windows/WindowsD3D11Device.cpp
--- a/Engine/Source/Runtime/Windows/D3D11RHI/Private/Windows/WindowsD3D11Device.cpp
+++ b/Engine/Source/Runtime/Windows/D3D11RHI/Private/Windows/WindowsD3D11Device.cpp
@@ -1704,175 +1704,190 @@ void FD3D11DynamicRHI::InitD3DDevice()
 	}
 }
 
-/**
- *	Retrieve available screen resolutions.
- *
- *	@param	Resolutions			TArray<FScreenResolutionRHI> parameter that will be filled in.
- *	@param	bIgnoreRefreshRate	If true, ignore refresh rates.
- *
- *	@return	bool				true if successfully filled the array
- */
-bool FD3D11DynamicRHI::RHIGetAvailableResolutions(FScreenResolutionArray& Resolutions, bool bIgnoreRefreshRate)
+static bool GetAvailableResolutionsForOutput(IDXGIOutput* Output, ID3D11Device* Device, int32 CurrentOutput, FScreenResolutionArray& Resolutions, bool bIgnoreRefreshRate)
 {
-	int32 MinAllowableResolutionX = 0;
-	int32 MinAllowableResolutionY = 0;
-	int32 MaxAllowableResolutionX = 10480;
-	int32 MaxAllowableResolutionY = 10480;
-	int32 MinAllowableRefreshRate = 0;
-	int32 MaxAllowableRefreshRate = 10480;
-
-	if (MaxAllowableResolutionX == 0) //-V547
-	{
-		MaxAllowableResolutionX = 10480;
-	}
-	if (MaxAllowableResolutionY == 0) //-V547
+	// TODO: GetDisplayModeList is a terribly SLOW call.  It can take up to a second per invocation.
+	//  We might want to work around some DXGI badness here.
+	const DXGI_FORMAT DisplayFormats[] = { DXGI_FORMAT_R8G8B8A8_UNORM, DXGI_FORMAT_B8G8R8A8_UNORM };
+	DXGI_FORMAT Format = DisplayFormats[0];
+	bool bFoundValidResult = false;
+	uint32 NumModes = 0;
+
+	for (DXGI_FORMAT CurrentFormat : DisplayFormats)
 	{
-		MaxAllowableResolutionY = 10480;
-	}
-	if (MaxAllowableRefreshRate == 0) //-V547
-	{
-		MaxAllowableRefreshRate = 10480;
-	}
+		const HRESULT HResult = Output->GetDisplayModeList(CurrentFormat, 0, &NumModes, NULL);
 
-	HRESULT HResult = S_OK;
-
-	int32 CurrentOutput = 0;
-	do 
-	{
-		TRefCountPtr<IDXGIOutput> Output;
-		HResult = Adapter.DXGIAdapter->EnumOutputs(CurrentOutput,Output.GetInitReference());
-		if(DXGI_ERROR_NOT_FOUND == HResult)
-			break;
-		if(FAILED(HResult))
-			return false;
-
-		// TODO: GetDisplayModeList is a terribly SLOW call.  It can take up to a second per invocation.
-		//  We might want to work around some DXGI badness here.
-		const DXGI_FORMAT DisplayFormats[] = {DXGI_FORMAT_R8G8B8A8_UNORM, DXGI_FORMAT_B8G8R8A8_UNORM};
-		DXGI_FORMAT Format = DisplayFormats[0];
-		bool bFoundValidResult = false;
-		uint32 NumModes = 0;	
-
-		for (DXGI_FORMAT CurrentFormat : DisplayFormats)
+		if (FAILED(HResult))
 		{
-			HResult = Output->GetDisplayModeList(CurrentFormat, 0, &NumModes, NULL);
-
-			if(FAILED(HResult))
+			if (HResult == DXGI_ERROR_NOT_FOUND)
 			{
-				if (HResult == DXGI_ERROR_NOT_FOUND)
-				{
-					UE_LOG(LogD3D11RHI, Warning, TEXT("RHIGetAvailableResolutions failed with generic error."));
-					continue;
-				}
-				else if (HResult == DXGI_ERROR_MORE_DATA)
-				{
-					UE_LOG(LogD3D11RHI, Warning, TEXT("RHIGetAvailableResolutions failed trying to return too much data."));
-					continue;
-				}
-				else if (HResult == DXGI_ERROR_NOT_CURRENTLY_AVAILABLE)
-				{
-					UE_LOG(LogD3D11RHI, Warning, TEXT("RHIGetAvailableResolutions does not return results when running under remote desktop."));
-					return false;
-				}
-				else
-				{
-					UE_LOG(LogD3D11RHI, Warning, TEXT("RHIGetAvailableResolutions failed with unknown error (0x%x)."), HResult);
-					return false;
-				}
+				UE_LOG(LogD3D11RHI, Warning, TEXT("RHIGetAvailableResolutions failed with generic error."));
+			}
+			else if (HResult == DXGI_ERROR_NOT_CURRENTLY_AVAILABLE)
+			{
+				UE_LOG(LogD3D11RHI, Warning, TEXT("RHIGetAvailableResolutions does not return results when running under remote desktop."));
+				return false;
 			}
 			else
 			{
-				bFoundValidResult = true;
-				if (NumModes)
-				{
-					Format = CurrentFormat;
-					break;
-				}
+				UE_LOG(LogD3D11RHI, Warning, TEXT("RHIGetAvailableResolutions failed with unknown error (0x%x)."), HResult);
+				return false;
 			}
 		}
-
-		// If we couldn't iterate any formats for this output, it might be dead in the driver and there's nothing we can do here except move on.
-		if (!bFoundValidResult)
+		else
 		{
-			// Grab diagnostic information from the output, if at all possible.
-			DXGI_OUTPUT_DESC OutputDesc;
-			FMemory::Memzero(OutputDesc);
-			HRESULT OutputDescRes = Output->GetDesc(&OutputDesc);
+			bFoundValidResult = true;
+			if (NumModes)
+			{
+				Format = CurrentFormat;
+				break;
+			}
+		}
+	}
 
-			UE_LOG(LogD3D11RHI, Warning,
-				TEXT("RHIGetAvailableResolutions could not get any display modes from output %i (D:%i) (Res:'%s'(0x%08X)"),
-				CurrentOutput,
-				OutputDesc.AttachedToDesktop,
-				*GetD3D11ErrorString(OutputDescRes, GetDevice()), OutputDescRes);
+	// If we couldn't iterate any formats for this output, it might be dead in the driver and there's nothing we can do here except move on.
+	if (!bFoundValidResult)
+	{
+		// Grab diagnostic information from the output, if at all possible.
+		DXGI_OUTPUT_DESC OutputDesc;
+		FMemory::Memzero(OutputDesc);
+		HRESULT OutputDescRes = Output->GetDesc(&OutputDesc);
 
-			++CurrentOutput;
-			continue;
-		}
+		UE_LOG(LogD3D11RHI, Warning,
+			TEXT("RHIGetAvailableResolutions could not get any display modes from output %i (D:%i) (Res:'%s'(0x%08X)"),
+			CurrentOutput,
+			OutputDesc.AttachedToDesktop,
+			*GetD3D11ErrorString(OutputDescRes, Device), OutputDescRes);
 
-		if (NumModes > 0)
-		{
-			DXGI_MODE_DESC* ModeList = new DXGI_MODE_DESC[NumModes];
-			VERIFYD3D11RESULT(Output->GetDisplayModeList(Format, 0, &NumModes, ModeList));
+		return false;
+	}
 
-			for (uint32 m = 0; m < NumModes; m++)
+	if (NumModes > 0)
+	{
+		TArray<DXGI_MODE_DESC> ModeList;
+		HRESULT HResult;
+		do
+		{
+			VERIFYD3D11RESULT(Output->GetDisplayModeList(Format, 0, &NumModes, nullptr));
+			ModeList.SetNum(NumModes);
+			HResult = Output->GetDisplayModeList(Format, 0, &NumModes, ModeList.GetData());
+		} while (HResult == DXGI_ERROR_MORE_DATA);
+		VERIFYD3D11RESULT(HResult);
+
+		constexpr int32 MinAllowableResolutionX = 0;
+		constexpr int32 MinAllowableResolutionY = 0;
+		constexpr int32 MaxAllowableResolutionX = 10480;
+		constexpr int32 MaxAllowableResolutionY = 10480;
+		constexpr int32 MinAllowableRefreshRate = 0;
+		constexpr int32 MaxAllowableRefreshRate = 10480;
+		for (uint32 m = 0; m < NumModes; m++)
+		{
+			CA_SUPPRESS(6385);
+			if (((int32)ModeList[m].Width >= MinAllowableResolutionX) &&
+				((int32)ModeList[m].Width <= MaxAllowableResolutionX) &&
+				((int32)ModeList[m].Height >= MinAllowableResolutionY) &&
+				((int32)ModeList[m].Height <= MaxAllowableResolutionY)
+				)
 			{
-				CA_SUPPRESS(6385);
-				if (((int32)ModeList[m].Width >= MinAllowableResolutionX) &&
-					((int32)ModeList[m].Width <= MaxAllowableResolutionX) &&
-					((int32)ModeList[m].Height >= MinAllowableResolutionY) &&
-					((int32)ModeList[m].Height <= MaxAllowableResolutionY)
-					)
+				bool bAddIt = true;
+				if (bIgnoreRefreshRate == false)
 				{
-					bool bAddIt = true;
-					if (bIgnoreRefreshRate == false)
+					if (((int32)ModeList[m].RefreshRate.Numerator < MinAllowableRefreshRate * ModeList[m].RefreshRate.Denominator) ||
+						((int32)ModeList[m].RefreshRate.Numerator > MaxAllowableRefreshRate * ModeList[m].RefreshRate.Denominator)
+						)
 					{
-						if (((int32)ModeList[m].RefreshRate.Numerator < MinAllowableRefreshRate * ModeList[m].RefreshRate.Denominator) ||
-							((int32)ModeList[m].RefreshRate.Numerator > MaxAllowableRefreshRate * ModeList[m].RefreshRate.Denominator)
-							)
-						{
-							continue;
-						}
+						continue;
 					}
-					else
+				}
+				else
+				{
+					// See if it is in the list already
+					for (int32 CheckIndex = 0; CheckIndex < Resolutions.Num(); CheckIndex++)
 					{
-						// See if it is in the list already
-						for (int32 CheckIndex = 0; CheckIndex < Resolutions.Num(); CheckIndex++)
+						FScreenResolutionRHI& CheckResolution = Resolutions[CheckIndex];
+						if ((CheckResolution.Width == ModeList[m].Width) &&
+							(CheckResolution.Height == ModeList[m].Height))
 						{
-							FScreenResolutionRHI& CheckResolution = Resolutions[CheckIndex];
-							if ((CheckResolution.Width == ModeList[m].Width) &&
-								(CheckResolution.Height == ModeList[m].Height))
-							{
-								// Already in the list...
-								bAddIt = false;
-								break;
-							}
+							// Already in the list...
+							bAddIt = false;
+							break;
 						}
 					}
+				}
 
-					if (bAddIt)
-					{
-						// Add the mode to the list
-						int32 Temp2Index = Resolutions.AddZeroed();
-						FScreenResolutionRHI& ScreenResolution = Resolutions[Temp2Index];
-
-						ScreenResolution.Width = ModeList[m].Width;
-						ScreenResolution.Height = ModeList[m].Height;
-						ScreenResolution.RefreshRate = ModeList[m].RefreshRate.Numerator / ModeList[m].RefreshRate.Denominator;
-					}
+				if (bAddIt)
+				{
+					FScreenResolutionRHI& ScreenResolution = Resolutions.Emplace_GetRef();
+					ScreenResolution.Width = ModeList[m].Width;
+					ScreenResolution.Height = ModeList[m].Height;
+					ScreenResolution.RefreshRate = ModeList[m].RefreshRate.Numerator / ModeList[m].RefreshRate.Denominator;
 				}
 			}
-
-			delete[] ModeList;
 		}
-		else
+	}
+	else
+	{
+		UE_LOG(LogD3D11RHI, Warning, TEXT("No display modes found for the standard format DXGI_FORMAT_R8G8B8A8_UNORM!"));
+	}
+
+	return !Resolutions.IsEmpty();
+}
+
+/**
+ *	Retrieve available screen resolutions.
+ *
+ *	@param	Resolutions			TArray<FScreenResolutionRHI> parameter that will be filled in.
+ *	@param	bIgnoreRefreshRate	If true, ignore refresh rates.
+ *
+ *	@return	bool				true if successfully filled the array
+ */
+bool FD3D11DynamicRHI::RHIGetAvailableResolutions(FScreenResolutionArray& Resolutions, bool bIgnoreRefreshRate)
+{
+	// TODO: The resolution information can be stale if the adapter's factory isn't current. Use RHIGetAvailableResolutionsForDisplay instead.
+	TRefCountPtr<IDXGIOutput> Output;
+	const HRESULT HResult = Adapter.DXGIAdapter->EnumOutputs(0, Output.GetInitReference());
+	return SUCCEEDED(HResult) ? GetAvailableResolutionsForOutput(Output, GetDevice(), 0, Resolutions, bIgnoreRefreshRate) : false;
+}
+
+bool FD3D11DynamicRHI::RHIGetAvailableResolutionsForDisplay(FScreenResolutionArray& Resolutions, bool bIgnoreRefreshRate, void* NativeMonitorHandle)
+{
+	HMONITOR MonitorHandle = static_cast<HMONITOR>(NativeMonitorHandle);
+	if (MonitorHandle == NULL)
+	{
+		return false;
+	}
+
+	if (!DXGIFactoryForResolutions.IsValid() || !DXGIFactoryForResolutions->IsCurrent())
+	{
+		SafeCreateDXGIFactory(DXGIFactoryForResolutions.GetInitReference(), GRHIGlobals.IsDebugLayerEnabled);
+	}
+
+	HRESULT HResult;
+	TRefCountPtr<IDXGIAdapter> DXGIAdapter;
+	for (int32 CurrentAdapter = 0; (HResult = DXGIFactoryForResolutions->EnumAdapters(CurrentAdapter, DXGIAdapter.GetInitReference())) != DXGI_ERROR_NOT_FOUND; ++CurrentAdapter)
+	{
+		if (FAILED(HResult))
 		{
-			UE_LOG(LogD3D11RHI, Warning, TEXT("No display modes found for the standard format DXGI_FORMAT_R8G8B8A8_UNORM!"));
+			return false;
 		}
 
-		++CurrentOutput;
+		TRefCountPtr<IDXGIOutput> Output;
+		for (int32 CurrentOutput = 0; (HResult = DXGIAdapter->EnumOutputs(CurrentOutput, Output.GetInitReference())) != DXGI_ERROR_NOT_FOUND; ++CurrentOutput)
+		{
+			if (FAILED(HResult))
+			{
+				return false;
+			}
 
-	// TODO: Cap at 1 for default output
-	} while(CurrentOutput < 1); //-V654
+			DXGI_OUTPUT_DESC OutputDesc{};
+			HRESULT OutputDescRes = Output->GetDesc(&OutputDesc);
+			check(SUCCEEDED(OutputDescRes));
+			if (MonitorHandle == OutputDesc.Monitor)
+			{
+				return GetAvailableResolutionsForOutput(Output, GetDevice(), CurrentOutput, Resolutions, bIgnoreRefreshRate);
+			}
+		}
+	}
 
-	return Resolutions.Num() > 0;
+	return false;
 }
