[Wasapi] Protect against device swap race condition
- Returns Wasapi behavior to interrupt current audio device swap requests in favor of new swap requests when simultaneous hardware device messages are received. Previous behavior was to ignore subsequent swap requests which led to a hung swap if the new device was invalid.
- Moves the critical section lock to cover the entire asynchronous PerformDeviceSwap() call to ensure the AudioDeviceManager in the DeviceSwapContext does not get reset in both the asynch swap task and a possible concurrent swap request on the hardware device callback.
- Performs an immediate device manager shut down to protect a slim race condition between game thread and asynch PerformDeviceSwap task.
- Resets the DeviceSwapContext before leaving the hardware device callback if the swap-to device is invalid.

#Synced-CL 50064980

diff --git a/Engine/Source/Runtime/Windows/AudioMixerWasapi/Private/AudioMixerWasapi.cpp b/Engine/Source/Runtime/Windows/AudioMixerWasapi/Private/AudioMixerWasapi.cpp
--- a/Engine/Source/Runtime/Windows/AudioMixerWasapi/Private/AudioMixerWasapi.cpp
+++ b/Engine/Source/Runtime/Windows/AudioMixerWasapi/Private/AudioMixerWasapi.cpp
@@ -503,11 +503,27 @@ namespace Audio
 		
 		if (DeviceSwapContext.IsValid())
 		{
-			UE_LOG(LogAudioMixer, Display, TEXT("FAudioMixerWasapi::InitDeviceSwapContextInternal - DeviceSwapContext in-flight, ignoring"));
-			return false;
+			UE_LOG(LogAudioMixer, Display, TEXT("FAudioMixerWasapi::InitDeviceSwapContextInternal - DeviceSwapContext in-flight, cancelling and starting a new swap."));
+
+			// If there is a previous device manager, this has been slated to be taken offline, so we need to shut the old one down first before resetting it
+			// Note, this can happen if a swap request comes in on the hardware device thread in between PreDeviceSwap() is called on the game thread and 
+			// PerformDeviceSwap is called on an asynchronous task. 
+			if (DeviceSwapContext->OldDeviceManager.IsValid())
+			{
+				const uint32 ThreadId = FPlatformTLS::GetCurrentThreadId(); // cross-platform thread id
+				UE_LOG(LogAudioMixer, Display, TEXT("FAudioMixerWasapi::InitDeviceSwapContextInternal - begin audio device manager shutdown for interrupted swap request. Thread Id %u"), ThreadId);
+
+				// Shutdown the current device manager
+				DeviceSwapContext->OldDeviceManager->StopAudioStream();
+				DeviceSwapContext->OldDeviceManager->CloseAudioStream();
+				DeviceSwapContext->OldDeviceManager->TeardownHardware();
+				DeviceSwapContext->OldDeviceManager.Reset();
+				UE_LOG(LogAudioMixer, Display, TEXT("FAudioMixerWasapi::InitDeviceSwapContextInternal - successfully shut down previous device manager before resetting in interrupted swap"));
+			}
 		}
 
 		// Create the new device swap context
+		DeviceSwapContext.Reset();
 		DeviceSwapContext = MakeUnique<FWasapiDeviceSwapContext>(InRequestedDeviceID, InReason);
 		if (!DeviceSwapContext.IsValid())
 		{
@@ -528,6 +544,7 @@ namespace Audio
 			{
 				UE_LOG(LogAudioMixer, Warning, TEXT("FAudioMixerWasapi::InitializeDeviceSwapContext - InitStreamParams() failed"));
 				DeviceSwapContext->NewDevice.Reset();
+				DeviceSwapContext.Reset();
 
 				return false;
 			}
@@ -549,15 +566,14 @@ namespace Audio
 		
 		TFunction<TUniquePtr<FDeviceSwapResult>()> AsyncDeviceSwap = [this]() mutable -> TUniquePtr<FDeviceSwapResult>
 		{
+			FScopeLock Lock(&DeviceSwapCriticalSection);
+				
 			// Transfer ownership of DeviceSwapContext to the async task.
 			TUniquePtr<FWasapiDeviceSwapContext> TempContext;
-			{
-				FScopeLock Lock(&DeviceSwapCriticalSection);
 
-				if (AudioStreamInfo.StreamState == EAudioOutputStreamState::SwappingDevice)
-				{
-					Swap(TempContext, DeviceSwapContext);
-				}
+			if (AudioStreamInfo.StreamState == EAudioOutputStreamState::SwappingDevice)
+			{
+				Swap(TempContext, DeviceSwapContext);
 			}
 
 			return PerformDeviceSwap(MoveTemp(TempContext));
@@ -590,7 +606,7 @@ namespace Audio
 		// This runs in an async task whose thread may not have initialized com
 		FScopedCoInitialize CoInitialize;
 
-		// No need to lock critical section here since this call has sole ownership of the context
+		// Critical section already locked in the async task that calls this function
 		if (InDeviceContext.IsValid())
 		{
 			UE_LOG(LogAudioMixer, Display, TEXT("FAudioMixerWasapi::PerformDeviceSwap - AsyncTask Start. Because=%s"), *InDeviceContext->DeviceSwapReason);
