Fix AMD-specific null dereference involving D3D11CreateDevice failure

#Synced-CL 50901530

diff --git a/Engine/Source/Runtime/Windows/D3D11RHI/Private/Windows/WindowsD3D11Device.cpp b/Engine/Source/Runtime/Windows/D3D11RHI/Private/Windows/WindowsD3D11Device.cpp
--- a/Engine/Source/Runtime/Windows/D3D11RHI/Private/Windows/WindowsD3D11Device.cpp
+++ b/Engine/Source/Runtime/Windows/D3D11RHI/Private/Windows/WindowsD3D11Device.cpp
@@ -1420,11 +1420,10 @@ void FD3D11DynamicRHI::InitD3DDevice()
 			// agsDriverExtensionsDX11_CreateDevice will not check for the DriverStore if there is not already a D3D11 device
 			// initialized. In order to use AMF we require libraries to also be loaded from the driver store so we temporarily
 			// initialize a device here and destroy it once the agsDriverExtensionsDX11_CreateDevice is run.
-			ID3D11Device* D3DDevicePreload = nullptr;
-			ID3D11DeviceContext* D3DDeviceContextPreload = nullptr;
+			TRefCountPtr<ID3D11Device> D3DDevicePreload;
+			TRefCountPtr<ID3D11DeviceContext> D3DDeviceContextPreload;
 
 			int32 NumAllowedFeatureLevels = 1;
-			D3D_FEATURE_LEVEL OutFeatureLevel = FeatureLevel;
 			HRESULT Result = D3D11CreateDevice(
 				Adapter.DXGIAdapter,
 				D3D_DRIVER_TYPE_UNKNOWN,
@@ -1433,12 +1432,12 @@ void FD3D11DynamicRHI::InitD3DDevice()
 				&FeatureLevel,
 				NumAllowedFeatureLevels,
 				D3D11_SDK_VERSION,
-				&D3DDevicePreload,
+				D3DDevicePreload.GetInitReference(),
 				&ActualFeatureLevel,
-				&D3DDeviceContextPreload);
+				D3DDeviceContextPreload.GetInitReference());
 			if (FAILED(Result))
 			{
-				UE_LOG(LogD3D11RHI, Error, TEXT("Failed to load the AMD DriverStore library"));
+				UE_LOG(LogD3D11RHI, Error, TEXT("Failed to load the AMD DriverStore library (0x%08X)"), Result);
 			}
 
 			AGSDX11ReturnedParams DeviceCreationReturnedParams;
@@ -1450,12 +1449,9 @@ void FD3D11DynamicRHI::InitD3DDevice()
 					&AmdExtensionParams,
 					&DeviceCreationReturnedParams);
 
-			// Destroy temporary device and context			
-			D3DDevicePreload->Release();
-			D3DDevicePreload = nullptr;
-
-			D3DDeviceContextPreload->Release();
-			D3DDeviceContextPreload = nullptr;
+			// Destroy temporary device and context
+			D3DDevicePreload.SafeRelease();
+			D3DDeviceContextPreload.SafeRelease();
 
 			if (DeviceCreation == AGS_SUCCESS)
 			{
@@ -1465,6 +1461,9 @@ void FD3D11DynamicRHI::InitD3DDevice()
 				AmdSupportedExtensionFlags = *(uint32*)&DeviceCreationReturnedParams.extensionsSupported;
 				bDeviceCreated = true;
 				UE_LOG(LogD3D11RHI, Log, TEXT("Created device via AGS, feature level %s, supported extensions %x."), GetFeatureLevelString(ActualFeatureLevel), AmdSupportedExtensionFlags);
+
+				GRHISupportsAtomicUInt64 = DeviceCreationReturnedParams.extensionsSupported.intrinsics19 != 0;  // "intrinsics19" includes AtomicU64
+				GSupportsDepthBoundsTest = DeviceCreationReturnedParams.extensionsSupported.depthBoundsTest != 0;
 			}
 			else
 			{
@@ -1475,9 +1474,6 @@ void FD3D11DynamicRHI::InitD3DDevice()
 				GRHIDeviceIsAMDPreGCNArchitecture = false;
 				UE_LOG(LogD3D11RHI, Warning, TEXT("Failed to create device via AGS, code %d."), DeviceCreation);
 			}
-
-			GRHISupportsAtomicUInt64 = DeviceCreationReturnedParams.extensionsSupported.intrinsics19 != 0;  // "intrinsics19" includes AtomicU64
-			GSupportsDepthBoundsTest = DeviceCreationReturnedParams.extensionsSupported.depthBoundsTest != 0;
 		}
 #endif //AMD_AGS_API
 
