Move D3D11 and D3D12 adapter selection logic into new module WindowsD3D

This allows us to make a better albeit imperfect decision about which GPU adapter will be chosen. This is the adapter that gets cross-referenced against DefaultHardware.ini and determines if we should run in perf mode, so it helps to choose correctly.

The known deficiencies are that we don't have knowledge of HMD adapters, and we prefer adapters that support D3D12 over D3D11. Trying to improve selection precision would likely make a larger mess.

#Synced-CL 50966172

diff --git a/Engine/Source/Runtime/ApplicationCore/ApplicationCore.Build.cs b/Engine/Source/Runtime/ApplicationCore/ApplicationCore.Build.cs
--- a/Engine/Source/Runtime/ApplicationCore/ApplicationCore.Build.cs
+++ b/Engine/Source/Runtime/ApplicationCore/ApplicationCore.Build.cs
@@ -27,7 +27,7 @@ public class ApplicationCore : ModuleRules
 			}
 		);
 
-		if ((Target.IsInPlatformGroup(UnrealPlatformGroup.Windows)))
+		if (Target.IsInPlatformGroup(UnrealPlatformGroup.Windows))
 		{
 			AddEngineThirdPartyPrivateStaticDependencies(Target,
 				"XInput"
@@ -42,11 +42,9 @@ public class ApplicationCore : ModuleRules
 				PublicDefinitions.Add("UE_WINDOWS_USING_UIA=0");
 			}
 
-			// Uses DXGI to query GPU hardware prior to RHI startup
-			PublicSystemLibraries.Add("DXGI.lib");
+			PrivateDependencyModuleNames.Add("WindowsD3D");
 
-			// DPI functions
-			PublicSystemLibraries.Add("Shcore.lib");
+			PublicSystemLibraries.Add("shcore.lib");
 		}
 		else if (Target.Platform == UnrealTargetPlatform.Mac)
 		{
diff --git a/Engine/Source/Runtime/ApplicationCore/Private/Windows/WindowsPlatformApplicationMisc.cpp b/Engine/Source/Runtime/ApplicationCore/Private/Windows/WindowsPlatformApplicationMisc.cpp
--- a/Engine/Source/Runtime/ApplicationCore/Private/Windows/WindowsPlatformApplicationMisc.cpp
+++ b/Engine/Source/Runtime/ApplicationCore/Private/Windows/WindowsPlatformApplicationMisc.cpp
@@ -14,6 +14,7 @@
 #include "Modules/ModuleManager.h"
 #include "Misc/ConfigCacheIni.h"
 #include "Misc/CoreDelegates.h"
+#include "Windows/WindowsD3D.h"
 #include "Windows/WindowsPlatformOutputDevices.h"
 #include "GenericPlatform/GenericPlatformCrashContext.h"
 #include "Templates/RefCounting.h"
@@ -29,10 +30,7 @@
 
 THIRD_PARTY_INCLUDES_START
 #include "Windows/AllowWindowsPlatformTypes.h"
-#include "dxgi1_3.h"
-#include "dxgi1_4.h"
-#include "dxgi1_6.h"
-#include "shellscalingapi.h"
+#include <shellscalingapi.h>
 #include "Windows/HideWindowsPlatformTypes.h"
 THIRD_PARTY_INCLUDES_END
 
@@ -391,30 +389,9 @@ int32 FWindowsPlatformApplicationMisc::GetMonitorDPI(void* NativeHandle)
 
 // Looks for an adapter with the most dedicated video memory
 FWindowsPlatformApplicationMisc::FGPUInfo FWindowsPlatformApplicationMisc::GetBestGPUInfo()
-{ 
-	TRefCountPtr<IDXGIFactory1> DXGIFactory1;
-	if (CreateDXGIFactory1(__uuidof(IDXGIFactory1), (void**)&DXGIFactory1) != S_OK || !DXGIFactory1)
-	{
-		return {};
-	}
-
-	DXGI_ADAPTER_DESC BestDesc = {};
-	TRefCountPtr<IDXGIAdapter> TempAdapter;
-	for (uint32 AdapterIndex = 0; DXGIFactory1->EnumAdapters(AdapterIndex, TempAdapter.GetInitReference()) != DXGI_ERROR_NOT_FOUND; ++AdapterIndex)
-	{
-		if (TempAdapter)
-		{
-			DXGI_ADAPTER_DESC Desc;
-			TempAdapter->GetDesc(&Desc);
-
-			if (Desc.DedicatedVideoMemory > BestDesc.DedicatedVideoMemory || AdapterIndex == 0)
-			{
-				BestDesc = Desc;
-			}
-		}
-	}
-
-	return FGPUInfo{ BestDesc.VendorId, BestDesc.DeviceId, BestDesc.DedicatedVideoMemory };
+{
+	const FWindowsGPUInfo Info = FWindowsD3D::GetGPUInfoByDedicatedMemory();
+	return FGPUInfo{ Info.VendorId, Info.DeviceId, Info.DedicatedVideoMemory };
 }
 
 float FWindowsPlatformApplicationMisc::GetDPIScaleFactorAtPoint(float X, float Y)
diff --git a/Engine/Source/Runtime/D3D12RHI/D3D12RHI.Build.cs b/Engine/Source/Runtime/D3D12RHI/D3D12RHI.Build.cs
--- a/Engine/Source/Runtime/D3D12RHI/D3D12RHI.Build.cs
+++ b/Engine/Source/Runtime/D3D12RHI/D3D12RHI.Build.cs
@@ -39,6 +39,7 @@ public class D3D12RHI : ModuleRules
 
 		if (Target.Platform.IsInGroup(UnrealPlatformGroup.Windows))
 		{
+			PublicDependencyModuleNames.Add("WindowsD3D");
 			PrivateDependencyModuleNames.Add("HeadMountedDisplay");
 
 			PublicDefinitions.Add("D3D12RHI_PLATFORM_HAS_CUSTOM_INTERFACE=0");
diff --git a/Engine/Source/Runtime/D3D12RHI/Private/Windows/WindowsD3D12Device.cpp b/Engine/Source/Runtime/D3D12RHI/Private/Windows/WindowsD3D12Device.cpp
--- a/Engine/Source/Runtime/D3D12RHI/Private/Windows/WindowsD3D12Device.cpp
+++ b/Engine/Source/Runtime/D3D12RHI/Private/Windows/WindowsD3D12Device.cpp
@@ -10,23 +10,21 @@
 #include "D3D12NvidiaExtensions.h"
 #include "WindowsD3D12Adapter.h"
 #include "Modules/ModuleManager.h"
-#include "Windows/AllowWindowsPlatformTypes.h"
 #include "Windows/WindowsPlatformCrashContext.h"
 #include "HAL/FileManager.h"
 #include "ProfilingDebugging/AssetMetadataTrace.h"
-#include <delayimp.h>
-#include "Windows/HideWindowsPlatformTypes.h"
-
 #include "HardwareInfo.h"
 #include "IHeadMountedDisplayModule.h"
 #include "GenericPlatform/GenericPlatformDriver.h"			// FGPUDriverInfo
 #include "RHIValidation.h"
 #include "RHIUtilities.h"
-
 #include "ShaderCompiler.h"
 #include "Misc/EngineVersion.h"
+#include "Windows/WindowsD3D.h"
 
-#pragma comment(lib, "d3d12.lib")
+#include "Windows/AllowWindowsPlatformTypes.h"
+#include <delayimp.h>
+#include "Windows/HideWindowsPlatformTypes.h"
 
 IMPLEMENT_MODULE(FD3D12DynamicRHIModule, D3D12RHI);
 
@@ -145,7 +143,7 @@ static void SafeCreateDXGIFactory(IDXGIFactory4** DXGIFactory)
 	{
 		bIsQuadBufferStereoEnabled = FParse::Param(FCommandLine::Get(), TEXT("quad_buffer_stereo"));
 
-		CreateDXGIFactory(__uuidof(IDXGIFactory4), (void**)DXGIFactory);
+		verify(SUCCEEDED(CreateDXGIFactory1(IID_PPV_ARGS(DXGIFactory))));
 	}
 	__except (IsDelayLoadException(GetExceptionInformation()))
 	{
@@ -525,7 +523,7 @@ inline bool GetSupportGPUUploadHeaps(ID3D12Device* device)
 
 /**
  * Attempts to create a D3D12 device for the adapter using at minimum MinFeatureLevel.
- * If creation is successful, true is returned and the max supported feature level is set in OutMaxFeatureLevel.
+ * If creation is successful, true is returned.
  */
 static bool SafeTestD3D12CreateDevice(IDXGIAdapter* Adapter, D3D_FEATURE_LEVEL MinFeatureLevel, FD3D12DeviceBasicInfo& OutInfo)
 {
@@ -874,34 +872,15 @@ void FD3D12DynamicRHIModule::FindAdapter()
 	}
 
 	TRefCountPtr<IDXGIFactory6> DXGIFactory6;
-	DXGIFactory4->QueryInterface(__uuidof(IDXGIFactory6), (void**)DXGIFactory6.GetInitReference());
-
-	bool bAllowPerfHUD = true;
-
-#if UE_BUILD_SHIPPING || UE_BUILD_TEST
-	bAllowPerfHUD = false;
-#endif
+	DXGIFactory4->QueryInterface(IID_PPV_ARGS(DXGIFactory6.GetInitReference()));
 
 	// Allow HMD to override which graphics adapter is chosen, so we pick the adapter where the HMD is connected
-	uint64 HmdGraphicsAdapterLuid = IHeadMountedDisplayModule::IsAvailable() ? IHeadMountedDisplayModule::Get().GetGraphicsAdapterLuid() : 0;
-	// Non-static as it is used only a few times
-	auto* CVarGraphicsAdapter = IConsoleManager::Get().FindTConsoleVariableDataInt(TEXT("r.GraphicsAdapter"));
-	int32 CVarExplicitAdapterValue = HmdGraphicsAdapterLuid == 0 ? (CVarGraphicsAdapter ? CVarGraphicsAdapter->GetValueOnGameThread() : -1) : -2;
-	FParse::Value(FCommandLine::Get(), TEXT("graphicsadapter="), CVarExplicitAdapterValue);
-
-	const bool bFavorDiscreteAdapter = CVarExplicitAdapterValue == -1;
+	const uint64 HmdGraphicsAdapterLuid = IHeadMountedDisplayModule::IsAvailable() ? IHeadMountedDisplayModule::Get().GetGraphicsAdapterLuid() : 0;
+	const uint64 AdapterLuid = FWindowsD3D::ChooseD3D12Adapter(HmdGraphicsAdapterLuid);
 
 	TRefCountPtr<IDXGIAdapter> TempAdapter;
 	const D3D_FEATURE_LEVEL MinRequiredFeatureLevel = GetRequiredD3DFeatureLevel();
 
-	FD3D12AdapterDesc PreferredAdapter;
-	FD3D12AdapterDesc FirstDiscreteAdapter;
-	FD3D12AdapterDesc BestMemoryAdapter;
-	FD3D12AdapterDesc FirstAdapter;
-
-	bool bRequestedWARP = D3D12RHI_ShouldCreateWithWarp();
-	bool bAllowSoftwareRendering = D3D12RHI_AllowSoftwareFallback();
-
 	int GpuPreferenceInt = DXGI_GPU_PREFERENCE_HIGH_PERFORMANCE;
 	FParse::Value(FCommandLine::Get(), TEXT("-gpupreference="), GpuPreferenceInt);
 	DXGI_GPU_PREFERENCE GpuPreference;
@@ -912,8 +891,6 @@ void FD3D12DynamicRHIModule::FindAdapter()
 	default: GpuPreference = DXGI_GPU_PREFERENCE_UNSPECIFIED; break;
 	}
 
-	const EGpuVendorId PreferredVendor = RHIGetPreferredAdapterVendor();
-
 	// Enumerate the DXGIFactory's adapters.
 	for (uint32 AdapterIndex = 0; FD3D12AdapterDesc::EnumAdapters(AdapterIndex, GpuPreference, DXGIFactory4, DXGIFactory6, TempAdapter.GetInitReference()) != DXGI_ERROR_NOT_FOUND; ++AdapterIndex)
 	{
@@ -921,7 +898,6 @@ void FD3D12DynamicRHIModule::FindAdapter()
 		if (TempAdapter)
 		{
 			FD3D12DeviceBasicInfo DeviceInfo;
-			// Log some information about the available D3D12 adapters.
 			DXGI_ADAPTER_DESC AdapterDesc{};
 			VERIFYD3D12RESULT(TempAdapter->GetDesc(&AdapterDesc));
 
@@ -996,84 +972,15 @@ void FD3D12DynamicRHIModule::FindAdapter()
 				CurrentAdapter.NumDeviceNodes = DeviceInfo.NumDeviceNodes;
 				CurrentAdapter.GpuPreference = GpuPreference;
 
-				// If requested WARP, then reject all other adapters. If WARP not requested, then reject the WARP device if software rendering support is disallowed
-				const bool bSkipWARP = (bRequestedWARP && !bIsWARP) || (!bRequestedWARP && bIsWARP && !bAllowSoftwareRendering);
-
-				// PerfHUD is for performance profiling
-				const bool bIsPerfHUD = !FCString::Stricmp(AdapterDesc.Description, TEXT("NVIDIA PerfHUD"));
-
-				// we don't allow the PerfHUD adapter
-				const bool bSkipPerfHUDAdapter = bIsPerfHUD && !bAllowPerfHUD;
-
-				// the HMD wants a specific adapter, not this one
-				const bool bSkipHmdGraphicsAdapter = HmdGraphicsAdapterLuid != 0 && FMemory::Memcmp(&HmdGraphicsAdapterLuid, &AdapterDesc.AdapterLuid, sizeof(LUID)) != 0;
-
-				// the user wants a specific adapter, not this one
-				const bool bSkipExplicitAdapter = CVarExplicitAdapterValue >= 0 && AdapterIndex != CVarExplicitAdapterValue;
-
-				const bool bSkipAdapter = bSkipWARP || bSkipPerfHUDAdapter || bSkipHmdGraphicsAdapter || bSkipExplicitAdapter || FParse::Param(FCommandLine::Get(), TEXT("ForceZeroAdapters"));
-
-				const bool bIsIntegrated = DeviceInfo.bUMA;
-
-				if (!bSkipAdapter && CurrentAdapter.IsValid())
+				if (!FMemory::Memcmp(&AdapterLuid, &AdapterDesc.AdapterLuid, sizeof(LUID)))
 				{
-					if (PreferredVendor != EGpuVendorId::Unknown && PreferredVendor == RHIConvertToGpuVendorId(AdapterDesc.VendorId) && !PreferredAdapter.IsValid())
-					{
-						PreferredAdapter = CurrentAdapter;
-					}
-					
-					if (!bIsWARP && !bIsIntegrated)
-					{
-						if (!FirstDiscreteAdapter.IsValid())
-						{
-							FirstDiscreteAdapter = CurrentAdapter;
-						}
-
-						if (CurrentAdapter.Desc.DedicatedVideoMemory > BestMemoryAdapter.Desc.DedicatedVideoMemory)
-						{
-							BestMemoryAdapter = CurrentAdapter;
-							if (PreferredVendor != EGpuVendorId::Unknown && PreferredVendor == RHIConvertToGpuVendorId(AdapterDesc.VendorId))
-							{
-								// Choose the best option of the preferred IHV devices
-								PreferredAdapter = BestMemoryAdapter;
-							}
-						}
-					}
-
-					if (!FirstAdapter.IsValid())
-					{
-						FirstAdapter = CurrentAdapter;
-					}
+					ChosenAdapter = MakeShareable(new FWindowsD3D12Adapter(CurrentAdapter));
+					// Don't break here so we can log all adapters.
 				}
 			}
 		}
 	}
 
-	if (bFavorDiscreteAdapter)
-	{
-		// We assume Intel is integrated graphics (slower than discrete) than NVIDIA or AMD cards and rather take a different one
-		if (PreferredAdapter.IsValid())
-		{
-			ChosenAdapter = MakeShareable(new FWindowsD3D12Adapter(PreferredAdapter));
-		}
-		else if (BestMemoryAdapter.IsValid())
-		{
-			ChosenAdapter = MakeShareable(new FWindowsD3D12Adapter(BestMemoryAdapter));
-		}
-		else if (FirstDiscreteAdapter.IsValid())
-		{
-			ChosenAdapter = MakeShareable(new FWindowsD3D12Adapter(FirstDiscreteAdapter));
-		}
-		else
-		{
-			ChosenAdapter = MakeShareable(new FWindowsD3D12Adapter(FirstAdapter));
-		}
-	}
-	else
-	{
-		ChosenAdapter = MakeShareable(new FWindowsD3D12Adapter(FirstAdapter));
-	}
-
 #if D3D12_CORE_ENABLED
 	UE_LOG(LogD3D12RHI, Log, TEXT("DirectX Agility SDK runtime %s."), CheckIfAgilitySDKLoaded() ? TEXT("found") : TEXT("not found"));
 #endif
diff --git a/Engine/Source/Runtime/D3D12RHI/Public/Windows/WindowsD3D12ThirdParty.h b/Engine/Source/Runtime/D3D12RHI/Public/Windows/WindowsD3D12ThirdParty.h
deleted file mode 100644
diff --git a/Engine/Source/Runtime/RHI/Private/Windows/WindowsDynamicRHI.cpp b/Engine/Source/Runtime/RHI/Private/Windows/WindowsDynamicRHI.cpp
--- a/Engine/Source/Runtime/RHI/Private/Windows/WindowsDynamicRHI.cpp
+++ b/Engine/Source/Runtime/RHI/Private/Windows/WindowsDynamicRHI.cpp
@@ -8,16 +8,10 @@
 #include "Misc/MessageDialog.h"
 #include "Containers/StaticArray.h"
 #include "DataDrivenShaderPlatformInfo.h"
-
+#include "Windows/WindowsD3D.h"
 #include "Windows/WindowsPlatformApplicationMisc.h"
 #include "Windows/WindowsConfigRules.h"
 
-THIRD_PARTY_INCLUDES_START
-#include "Windows/AllowWindowsPlatformTypes.h"
-#include <dxgi.h>
-#include "Windows/HideWindowsPlatformTypes.h"
-THIRD_PARTY_INCLUDES_END
-
 #define LOCTEXT_NAMESPACE "WindowsDynamicRHI"
 
 bool GDynamicRHIFailedToInitializeAdvancedPlatform = false;
@@ -385,45 +379,16 @@ TOptional<EWindowsRHI> ParseDefaultWindowsRHI()
 	return DefaultRHI;
 }
 
-struct FWindowsGPUInfo
-{
-	uint32 VendorId = 0;
-	uint32 DeviceId = 0;
-	uint64 DedicatedVideoMemory = 0;
-	FString AdapterName;
-};
-
-static void GetGPUInfo(FWindowsGPUInfo& Info)
+static FWindowsGPUInfo GetGPUInfo()
 {
-	DXGI_ADAPTER_DESC BestDesc = {};
-	TRefCountPtr<IDXGIFactory1> DXGIFactory1;
-	if (SUCCEEDED(CreateDXGIFactory1(IID_PPV_ARGS(DXGIFactory1.GetInitReference()))))
-	{
-		TRefCountPtr<IDXGIAdapter> TempAdapter;
-		for (uint32 AdapterIndex = 0; DXGIFactory1->EnumAdapters(AdapterIndex, TempAdapter.GetInitReference()) != DXGI_ERROR_NOT_FOUND; ++AdapterIndex)
-		{
-			if (TempAdapter)
-			{
-				DXGI_ADAPTER_DESC Desc;
-				TempAdapter->GetDesc(&Desc);
-
-				if (Desc.DedicatedVideoMemory > BestDesc.DedicatedVideoMemory || AdapterIndex == 0)
-				{
-					BestDesc = Desc;
-				}
-			}
-		}
-	}
-
-	Info.VendorId = BestDesc.VendorId;
-	Info.DeviceId = BestDesc.DeviceId;
-	Info.DedicatedVideoMemory = BestDesc.DedicatedVideoMemory;
-	Info.AdapterName = BestDesc.Description;
+	FWindowsGPUInfo Info = FWindowsD3D::GetExpectedGPUInfo();
 
 	if (Info.AdapterName.IsEmpty())
 	{
 		Info.AdapterName = FPlatformMisc::GetPrimaryGPUBrand();
 	}
+
+	return Info;
 }
 
 static bool IsGPUInHardwareListForRHI(const FWindowsGPUInfo& BestGPUInfo, const TCHAR* const ConfigKey, const TCHAR* const MatchingRHIName)
@@ -487,8 +452,7 @@ static bool ShouldAvoidD3D11_SM5()
 		!FParse::Param(FCommandLine::Get(), TEXT("dx11")) &&
 		!FParse::Param(FCommandLine::Get(), TEXT("sm5")))
 	{
-		FWindowsGPUInfo BestGPUInfo;
-		GetGPUInfo(BestGPUInfo);
+		const FWindowsGPUInfo BestGPUInfo = GetGPUInfo();
 
 		bAvoidD3D11SM5 = IsGPUInHardwareListForRHI(BestGPUInfo, TEXT("DeviceAvoidRHIList"), TEXT("D3D11_SM5"));
 		if (bAvoidD3D11SM5)
@@ -517,8 +481,7 @@ static bool ShouldAvoidD3D11_ES31()
 		!FParse::Param(FCommandLine::Get(), TEXT("FeatureLevelES31")) &&
 		!FParse::Param(FCommandLine::Get(), TEXT("FeatureLevelES3_1")))
 	{
-		FWindowsGPUInfo BestGPUInfo;
-		GetGPUInfo(BestGPUInfo);
+		const FWindowsGPUInfo BestGPUInfo = GetGPUInfo();
 
 		bAvoidD3D11ES31 = IsGPUInHardwareListForRHI(BestGPUInfo, TEXT("DeviceAvoidRHIList"), TEXT("D3D11_ES31"));
 		if (bAvoidD3D11ES31)
@@ -648,8 +611,7 @@ static bool IsForcedToD3D11_ES31()
 	}
 	else
 	{
-		FWindowsGPUInfo BestGPUInfo;
-		GetGPUInfo(BestGPUInfo);
+		const FWindowsGPUInfo BestGPUInfo = GetGPUInfo();
 
 		bRequireD3D11ES31 = IsGPUInHardwareListForRHI(BestGPUInfo, TEXT("DeviceForceRHIList"), TEXT("D3D11_ES31"));
 		if (bRequireD3D11ES31)
@@ -675,8 +637,7 @@ static bool ShouldDefaultToPerformanceMode(bool& bOutForceD3D11)
 		return ForceES31.GetValue();
 	}
 
-	FWindowsGPUInfo BestGPUInfo;
-	GetGPUInfo(BestGPUInfo);
+	const FWindowsGPUInfo BestGPUInfo = GetGPUInfo();
 
 	if (IsGPUInHardwareListForRHI(BestGPUInfo, TEXT("DeviceDefaultRHIList"), TEXT("D3D11_ES31")))
 	{
@@ -872,11 +833,11 @@ static bool ShouldDevicePreferSM6(uint32 SM6CapableDeviceId)
 
 static bool IsRHIAllowedAsDefault(EWindowsRHI InRHI, ERHIFeatureLevel::Type InFeatureLevel)
 {
-	static const FWindowsPlatformApplicationMisc::FGPUInfo BestGPUInfo = FWindowsPlatformApplicationMisc::GetBestGPUInfo();
+	static const FWindowsGPUInfo BestGPUInfo = GetGPUInfo();
 
 	if (InRHI == EWindowsRHI::D3D12 && InFeatureLevel == ERHIFeatureLevel::SM6)
 	{
-		if (ShouldDevicePreferSM6(BestGPUInfo.DeviceId)) 
+		if (ShouldDevicePreferSM6(BestGPUInfo.DeviceId))
 		{
 			return true;
 		}
diff --git a/Engine/Source/Runtime/RHI/RHI.Build.cs b/Engine/Source/Runtime/RHI/RHI.Build.cs
--- a/Engine/Source/Runtime/RHI/RHI.Build.cs
+++ b/Engine/Source/Runtime/RHI/RHI.Build.cs
@@ -15,6 +15,10 @@ public class RHI : ModuleRules
 		PrivateDependencyModuleNames.Add("ApplicationCore");
 		PrivateDependencyModuleNames.Add("Cbor");
 		PrivateDependencyModuleNames.Add("BuildSettings");
+		if (Target.Platform.IsInGroup(UnrealPlatformGroup.Windows))
+		{
+			PrivateDependencyModuleNames.Add("WindowsD3D");
+		}
 
 		PublicDefinitions.AddDefinition("WITH_MGPU", Target.Platform.IsInGroup(UnrealPlatformGroup.Windows) && Target.Platform.IsInGroup(UnrealPlatformGroup.Desktop));
 
@@ -30,18 +34,13 @@ public class RHI : ModuleRules
 					PublicDefinitions.Add("RHI_WANT_RESOURCE_INFO=1");
 				}
 
-				// UEBuildAndroid.cs adds VulkanRHI for Android builds if it is enabled
 				if (Target.Platform.IsInGroup(UnrealPlatformGroup.Windows))
 				{
 					DynamicallyLoadedModuleNames.Add("D3D11RHI");
-				}
-
-				if (Target.Platform.IsInGroup(UnrealPlatformGroup.Windows))
-				{
-					//#todo-rco: D3D12 requires different SDK headers not compatible with WinXP
 					DynamicallyLoadedModuleNames.Add("D3D12RHI");
 				}
 
+				// UEBuildAndroid.cs adds VulkanRHI for Android builds if it is enabled
 				if (Target.Platform.IsInGroup(UnrealPlatformGroup.Windows) || Target.IsInPlatformGroup(UnrealPlatformGroup.Unix))
 				{
 					DynamicallyLoadedModuleNames.Add("VulkanRHI");
diff --git a/Engine/Source/Runtime/Windows/D3D11RHI/D3D11RHI.Build.cs b/Engine/Source/Runtime/Windows/D3D11RHI/D3D11RHI.Build.cs
--- a/Engine/Source/Runtime/Windows/D3D11RHI/D3D11RHI.Build.cs
+++ b/Engine/Source/Runtime/Windows/D3D11RHI/D3D11RHI.Build.cs
@@ -31,6 +31,7 @@ public class D3D11RHI : ModuleRules
 
 		if (Target.IsInPlatformGroup(UnrealPlatformGroup.Windows))
 		{
+			PublicDependencyModuleNames.Add("WindowsD3D");
 			PrivateDependencyModuleNames.Add("HeadMountedDisplay");
 		}
 
diff --git a/Engine/Source/Runtime/Windows/D3D11RHI/Private/Windows/WindowsD3D11Device.cpp b/Engine/Source/Runtime/Windows/D3D11RHI/Private/Windows/WindowsD3D11Device.cpp
--- a/Engine/Source/Runtime/Windows/D3D11RHI/Private/Windows/WindowsD3D11Device.cpp
+++ b/Engine/Source/Runtime/Windows/D3D11RHI/Private/Windows/WindowsD3D11Device.cpp
@@ -3,23 +3,11 @@
 /*=============================================================================
 	WindowsD3D11Device.cpp: Windows D3D device RHI implementation.
 =============================================================================*/
-#include "Misc/EngineVersion.h"
 #include "D3D11RHIPrivate.h"
 #include "HAL/FileManager.h"
 #include "Misc/CommandLine.h"
 #include "Misc/EngineVersion.h"
-#include "Windows/AllowWindowsPlatformTypes.h"
 #include "Windows/WindowsPlatformCrashContext.h"
-#include <delayimp.h>
-	#if WITH_NVAPI
-		#include "nvapi.h"
-		#include "nvShaderExtnEnums.h"
-	#endif
-	#if WITH_AMD_AGS
-		#include "amd_ags.h"
-	#endif
-#include "Windows/HideWindowsPlatformTypes.h"
-
 #include "HardwareInfo.h"
 #include "IHeadMountedDisplayModule.h"
 #include "GenericPlatform/GenericPlatformDriver.h"			// FGPUDriverInfo
@@ -29,21 +17,23 @@
 #include "ShaderDiagnostics.h"
 #include "HDRHelper.h"
 #include "GlobalShader.h"
+#include "Windows/WindowsD3D.h"
+
+#include "Windows/AllowWindowsPlatformTypes.h"
+#include <delayimp.h>
+	#if WITH_NVAPI
+		#include "nvapi.h"
+		#include "nvShaderExtnEnums.h"
+	#endif
+	#if WITH_AMD_AGS
+		#include "amd_ags.h"
+	#endif
+#include "Windows/HideWindowsPlatformTypes.h"
 
 FD3D11DynamicRHI*	GD3D11RHI = nullptr;
 
 bool D3D11RHI_ShouldAllowAsyncResourceCreation();
 
-static bool D3D11RHI_AllowSoftwareFallback()
-{
-	if (FParse::Param(FCommandLine::Get(), TEXT("AllowSoftwareRendering")))
-	{
-		return true;
-	}
-
-	return false;
-}
-
 #ifdef AMD_AGS_API
 // Filled in during InitD3DDevice if IsRHIDeviceAMD
 struct AmdAgsInfo
@@ -254,65 +244,13 @@ static bool IsDelayLoadException(PEXCEPTION_POINTERS ExceptionPointers)
 }
 
 static bool bIsQuadBufferStereoEnabled = false;
-typedef HRESULT(WINAPI *FCreateDXGIFactory2)(UINT, REFIID, void **);
-static FCreateDXGIFactory2 CreateDXGIFactory2FnPtr = nullptr;
 
-/**
- * Since CreateDXGIFactory1 is a delay loaded import from the D3D11 DLL, if the user
- * doesn't have VistaSP2/DX10, calling CreateDXGIFactory1 will throw an exception.
- * We use SEH to detect that case and fail gracefully.
- */
 static void SafeCreateDXGIFactory(IDXGIFactory1** DXGIFactory1, bool bWithDebug)
 {
-#if !defined(D3D11_CUSTOM_VIEWPORT_CONSTRUCTOR) || !D3D11_CUSTOM_VIEWPORT_CONSTRUCTOR
-	__try
-	{
-		bool bQuadBufferStereoRequested = FParse::Param(FCommandLine::Get(), TEXT("quad_buffer_stereo"));
-		bool bIsWin8OrNewer = FPlatformMisc::VerifyWindowsVersion(8, 0);
-
-		if (bIsWin8OrNewer && (bQuadBufferStereoRequested || bWithDebug))
-		{
-			// CreateDXGIFactory2 is only available on Win8.1+, find it if it exists
-			HMODULE DxgiDLL = (HMODULE)FPlatformProcess::GetDllHandle(TEXT("dxgi.dll"));
-			if (DxgiDLL)
-			{
-#pragma warning(push)
-#pragma warning(disable: 4191) // disable the "unsafe conversion from 'FARPROC' to 'blah'" warning
-				CreateDXGIFactory2FnPtr = (FCreateDXGIFactory2)(GetProcAddress(DxgiDLL, "CreateDXGIFactory2"));
-#pragma warning(pop)
-				FPlatformProcess::FreeDllHandle(DxgiDLL);
-			}
-
-			if (bQuadBufferStereoRequested)
-			{
-				if (CreateDXGIFactory2FnPtr)
-				{
-					bIsQuadBufferStereoEnabled = true;
-				}
-				else
-				{
-					UE_LOG(LogD3D11RHI, Warning, TEXT("Win8.1 or above is required for quad_buffer_stereo support."));
-				}
-			}
-		}
+	bIsQuadBufferStereoEnabled = FParse::Param(FCommandLine::Get(), TEXT("quad_buffer_stereo"));
 
-		// IDXGIFactory2 required for dx11.1 active stereo and DXGI debug (dxgi1.3)
-		if (CreateDXGIFactory2FnPtr)
-		{
-			uint32 Flags = bWithDebug ? DXGI_CREATE_FACTORY_DEBUG : 0;
-			CreateDXGIFactory2FnPtr(Flags, __uuidof(IDXGIFactory2), (void**)DXGIFactory1);
-		}
-		else
-		{
-			CreateDXGIFactory1(__uuidof(IDXGIFactory1), (void**)DXGIFactory1);
-		}
-	}
-	__except(IsDelayLoadException(GetExceptionInformation()))
-	{
-		// We suppress warning C6322: Empty _except block. Appropriate checks are made upon returning. 
-		CA_SUPPRESS(6322);
-	}
-#endif	//!D3D11_CUSTOM_VIEWPORT_CONSTRUCTOR
+	const uint32 Flags = bWithDebug ? DXGI_CREATE_FACTORY_DEBUG : 0;
+	verify(SUCCEEDED(CreateDXGIFactory2(Flags, IID_PPV_ARGS(DXGIFactory1))));
 }
 
 /**
@@ -351,9 +289,9 @@ inline bool IsDeviceUMA(ID3D11Device* device)
 
 /**
  * Attempts to create a D3D11 device for the adapter using at most MaxFeatureLevel.
- * If creation is successful, true is returned and the supported feature level is set in OutFeatureLevel.
+ * If creation is successful, true is returned.
  */
-static bool SafeTestD3D11CreateDevice(IDXGIAdapter* Adapter, D3D_FEATURE_LEVEL MinFeatureLevel,D3D_FEATURE_LEVEL MaxFeatureLevel, FD3D11DeviceBasicInfo& OutInfo)
+static bool SafeTestD3D11CreateDevice(IDXGIAdapter* Adapter, D3D_FEATURE_LEVEL MinFeatureLevel, D3D_FEATURE_LEVEL MaxFeatureLevel, FD3D11DeviceBasicInfo& OutInfo)
 {
 	ID3D11Device* D3DDevice = nullptr;
 	ID3D11DeviceContext* D3DDeviceContext = nullptr;
@@ -434,10 +372,8 @@ static bool SafeTestD3D11CreateDevice(IDXGIAdapter* Adapter, D3D_FEATURE_LEVEL M
 		// Log any reason for failure to create test device. Extra debug help.
 		VERIFYD3D11RESULT_NOEXIT(Result);
 
-		bool bIsWin10 = FPlatformMisc::VerifyWindowsVersion(10, 0);
-
 		// Fatal error on 0x887A002D
-		if (DXGI_ERROR_SDK_COMPONENT_MISSING == Result && bIsWin10)
+		if (DXGI_ERROR_SDK_COMPONENT_MISSING == Result)
 		{
 			UE_LOG(LogD3D11RHI, Fatal, TEXT("-d3ddebug was used but optional Graphics Tools were not found. Install them through the Manage Optional Features in windows. See: https://docs.microsoft.com/en-us/windows/uwp/gaming/use-the-directx-runtime-and-visual-studio-graphics-diagnostic-features"));
 		}
@@ -745,22 +681,11 @@ void FD3D11DynamicRHIModule::FindAdapter()
 	}
 
 	TRefCountPtr<IDXGIFactory6> DXGIFactory6;
-	DXGIFactory1->QueryInterface(__uuidof(IDXGIFactory6), (void**)DXGIFactory6.GetInitReference());
-
-	bool bAllowPerfHUD = true;
-
-#if UE_BUILD_SHIPPING || UE_BUILD_TEST
-	bAllowPerfHUD = false;
-#endif
+	DXGIFactory1->QueryInterface(IID_PPV_ARGS(DXGIFactory6.GetInitReference()));
 
 	// Allow HMD to override which graphics adapter is chosen, so we pick the adapter where the HMD is connected
-	uint64 HmdGraphicsAdapterLuid  = IHeadMountedDisplayModule::IsAvailable() ? IHeadMountedDisplayModule::Get().GetGraphicsAdapterLuid() : 0;
-	// Non-static as it is used only a few times
-	auto* CVarGraphicsAdapter = IConsoleManager::Get().FindTConsoleVariableDataInt(TEXT("r.GraphicsAdapter"));
-	int32 CVarExplicitAdapterValue = HmdGraphicsAdapterLuid == 0 ? (CVarGraphicsAdapter ? CVarGraphicsAdapter->GetValueOnGameThread() : -1) : -2;
-	FParse::Value(FCommandLine::Get(), TEXT("graphicsadapter="), CVarExplicitAdapterValue);
-
-	const bool bFavorDiscreteAdapter = CVarExplicitAdapterValue == -1;
+	const uint64 HmdGraphicsAdapterLuid = IHeadMountedDisplayModule::IsAvailable() ? IHeadMountedDisplayModule::Get().GetGraphicsAdapterLuid() : 0;
+	const uint64 AdapterLuid = FWindowsD3D::ChooseD3D11Adapter(HmdGraphicsAdapterLuid);
 
 	TRefCountPtr<IDXGIAdapter> TempAdapter;
 	D3D_FEATURE_LEVEL MinAllowedFeatureLevel = GetMinAllowedD3DFeatureLevel();
@@ -769,19 +694,8 @@ void FD3D11DynamicRHIModule::FindAdapter()
 	UE_LOG(LogD3D11RHI, Log, TEXT("D3D11 min allowed feature level: %s"), GetFeatureLevelString(MinAllowedFeatureLevel));
 	UE_LOG(LogD3D11RHI, Log, TEXT("D3D11 max allowed feature level: %s"), GetFeatureLevelString(MaxAllowedFeatureLevel));
 
-	FD3D11Adapter PreferredAdapter;
-	FD3D11Adapter BestMemoryAdapter;
-	FD3D11Adapter FirstDiscreteAdapter;
-	FD3D11Adapter FirstAdapter;
-
-	SIZE_T BestDedicatedMemory = 0;
-
 	UE_LOG(LogD3D11RHI, Log, TEXT("D3D11 adapters:"));
 
-	const EGpuVendorId PreferredVendor = RHIGetPreferredAdapterVendor();
-	bool bAllowSoftwareFallback = D3D11RHI_AllowSoftwareFallback();
-
-
 	int GpuPreferenceInt = DXGI_GPU_PREFERENCE_HIGH_PERFORMANCE;
 	FParse::Value(FCommandLine::Get(), TEXT("-gpupreference="), GpuPreferenceInt);
 	DXGI_GPU_PREFERENCE GpuPreference;
@@ -812,129 +726,47 @@ void FD3D11DynamicRHIModule::FindAdapter()
 		{
 			UE_LOG(LogD3D11RHI, Log, TEXT("Testing D3D11 Adapter %u:"), AdapterIndex);
 			DXGI_ADAPTER_DESC AdapterDesc;
-			if (HRESULT DescResult = TempAdapter->GetDesc(&AdapterDesc); FAILED(DescResult))
-			{
-				UE_LOG(LogD3D11RHI, Warning, TEXT("Failed to get description for adapter %u."), AdapterIndex);
-			}
-			else
+			if (HRESULT DescResult = TempAdapter->GetDesc(&AdapterDesc); SUCCEEDED(DescResult))
 			{
 				LogDXGIAdapterDesc(AdapterDesc);
-			}
 
-			FD3D11DeviceBasicInfo DeviceInfo;
-			if(SafeTestD3D11CreateDevice(TempAdapter, MinAllowedFeatureLevel, MaxAllowedFeatureLevel, DeviceInfo))
-			{
-				// Log some information about the available D3D11 adapters.
-				
-				uint32 OutputCount = CountAdapterOutputs(TempAdapter);
-
-				UE_LOG(LogD3D11RHI, Log,
-					TEXT("  %2u. '%s' (Feature Level %s)"),
-					AdapterIndex,
-					AdapterDesc.Description,
-					GetFeatureLevelString(DeviceInfo.MaxFeatureLevel)
-					);
-				UE_LOG(LogD3D11RHI, Log,
-					TEXT("      %u/%u/%u MB DedicatedVideo/DedicatedSystem/SharedSystem, Outputs:%d, VendorId:0x%x UMA:%s"),
-					(uint32)(AdapterDesc.DedicatedVideoMemory / (1024*1024)),
-					(uint32)(AdapterDesc.DedicatedSystemMemory / (1024*1024)),
-					(uint32)(AdapterDesc.SharedSystemMemory / (1024*1024)),
-					OutputCount,
-					AdapterDesc.VendorId,
-					DeviceInfo.bUMA ? TEXT("true") : TEXT("false")
-					);
-
-				const bool bIsWARP = RHIConvertToGpuVendorId(AdapterDesc.VendorId) == EGpuVendorId::Microsoft;
-
-				// PerfHUD is for performance profiling
-				const bool bIsPerfHUD = !FCString::Stricmp(AdapterDesc.Description,TEXT("NVIDIA PerfHUD"));
-
-				FD3D11Adapter CurrentAdapter(TempAdapter, bIsWARP, DeviceInfo);
-
-				// Add special check to support HMDs, which do not have associated outputs.
-				// To reject the software emulation, unless the cvar wants it.
-				// https://msdn.microsoft.com/en-us/library/windows/desktop/bb205075(v=vs.85).aspx#WARP_new_for_Win8
-				// Before we tested for no output devices but that failed where a laptop had a Intel (with output) and NVidia (with no output)
-				const bool bSkipSoftwareAdapter = bIsWARP && !bAllowSoftwareFallback && CVarExplicitAdapterValue < 0 && HmdGraphicsAdapterLuid == 0;
-				
-				// we don't allow the PerfHUD adapter
-				const bool bSkipPerfHUDAdapter = bIsPerfHUD && !bAllowPerfHUD;
-
-				// the HMD wants a specific adapter, not this one
-				const bool bSkipHmdGraphicsAdapter = HmdGraphicsAdapterLuid != 0 && FMemory::Memcmp(&HmdGraphicsAdapterLuid, &AdapterDesc.AdapterLuid, sizeof(LUID)) != 0;
-
-				// the user wants a specific adapter, not this one
-				const bool bSkipExplicitAdapter = CVarExplicitAdapterValue >= 0 && AdapterIndex != CVarExplicitAdapterValue;
-				
-				const bool bSkipAdapter = bSkipSoftwareAdapter || bSkipPerfHUDAdapter || bSkipHmdGraphicsAdapter || bSkipExplicitAdapter;
-
-				const bool bIsIntegrated = DeviceInfo.bUMA;
-
-				if (!bSkipAdapter && CurrentAdapter.IsValid())
+				FD3D11DeviceBasicInfo DeviceInfo;
+				if(SafeTestD3D11CreateDevice(TempAdapter, MinAllowedFeatureLevel, MaxAllowedFeatureLevel, DeviceInfo))
 				{
-					if (PreferredVendor != EGpuVendorId::Unknown && PreferredVendor == RHIConvertToGpuVendorId(AdapterDesc.VendorId) && !PreferredAdapter.IsValid())
-					{
-						PreferredAdapter = CurrentAdapter;
-					}
-					
-					if (!bIsWARP && !bIsIntegrated)
-					{
-						if (!FirstDiscreteAdapter.IsValid())
-						{
-							FirstDiscreteAdapter = CurrentAdapter;
-						}
-
-						if (AdapterDesc.DedicatedVideoMemory > BestDedicatedMemory)
-						{
-							BestMemoryAdapter = CurrentAdapter;
-							BestDedicatedMemory = AdapterDesc.DedicatedVideoMemory;
-							if (PreferredVendor != EGpuVendorId::Unknown && PreferredVendor == RHIConvertToGpuVendorId(AdapterDesc.VendorId))
-							{
-								// Choose the best option of the preferred IHV devices
-								PreferredAdapter = BestMemoryAdapter;
-							}
-						}
-					}
-
-					if (!FirstAdapter.IsValid())
+					// Log some information about the available D3D11 adapters.
+
+					uint32 OutputCount = CountAdapterOutputs(TempAdapter);
+
+					UE_LOG(LogD3D11RHI, Log,
+						TEXT("  %2u. '%s' (Feature Level %s)"),
+						AdapterIndex,
+						AdapterDesc.Description,
+						GetFeatureLevelString(DeviceInfo.MaxFeatureLevel)
+						);
+					UE_LOG(LogD3D11RHI, Log,
+						TEXT("      %u/%u/%u MB DedicatedVideo/DedicatedSystem/SharedSystem, Outputs:%d, VendorId:0x%x UMA:%s"),
+						(uint32)(AdapterDesc.DedicatedVideoMemory / (1024*1024)),
+						(uint32)(AdapterDesc.DedicatedSystemMemory / (1024*1024)),
+						(uint32)(AdapterDesc.SharedSystemMemory / (1024*1024)),
+						OutputCount,
+						AdapterDesc.VendorId,
+						DeviceInfo.bUMA ? TEXT("true") : TEXT("false")
+						);
+
+					const bool bIsWARP = RHIConvertToGpuVendorId(AdapterDesc.VendorId) == EGpuVendorId::Microsoft;
+
+					if (!FMemory::Memcmp(&AdapterLuid, &AdapterDesc.AdapterLuid, sizeof(LUID)))
 					{
-						FirstAdapter = CurrentAdapter;
+						ChosenAdapter = FD3D11Adapter(TempAdapter, bIsWARP, DeviceInfo);
+						// Don't break here so we can log all adapters.
 					}
 				}
 			}
 			else
 			{
-				UE_LOG(LogD3D11RHI, Log, TEXT("  %2u. Unknown, failed to create test device."), AdapterIndex);
+				UE_LOG(LogD3D11RHI, Warning, TEXT("Failed to get description for adapter %u."), AdapterIndex);
 			}
 		}
-		else
-		{
-			UE_LOG(LogD3D11RHI, Log, TEXT("  %2u. Unknown, failed to create adapter."), AdapterIndex);
-		}
-	}
-
-	if (bFavorDiscreteAdapter)
-	{
-		if (PreferredAdapter.IsValid())
-		{
-			ChosenAdapter = PreferredAdapter;
-		}
-		else if (BestMemoryAdapter.IsValid())
-		{
-			ChosenAdapter = BestMemoryAdapter;
-		}
-		else if (FirstDiscreteAdapter.IsValid())
-		{
-			ChosenAdapter = FirstDiscreteAdapter;
-		}
-		else
-		{
-			ChosenAdapter = FirstAdapter;
-		}
-	}
-	else
-	{
-		ChosenAdapter = FirstAdapter;
 	}
 
 	if(ChosenAdapter.IsValid())
diff --git a/Engine/Source/Runtime/Windows/D3D11RHI/Public/Windows/D3D11ThirdParty.h b/Engine/Source/Runtime/Windows/D3D11RHI/Public/Windows/D3D11ThirdParty.h
deleted file mode 100644
diff --git a/Engine/Source/Runtime/Windows/WindowsD3D/Internal/Windows/WindowsD3D.h b/Engine/Source/Runtime/Windows/WindowsD3D/Internal/Windows/WindowsD3D.h
new file mode 100644
--- /dev/null
+++ b/Engine/Source/Runtime/Windows/WindowsD3D/Internal/Windows/WindowsD3D.h
@@ -0,0 +1,23 @@
+// Copyright Epic Games, Inc. All Rights Reserved.
+
+#pragma once
+
+#include "Containers/UnrealString.h"
+#include "HAL/Platform.h"
+
+struct FWindowsGPUInfo
+{
+	uint32 VendorId = 0;
+	uint32 DeviceId = 0;
+	uint64 DedicatedVideoMemory = 0;
+	FString AdapterName;
+};
+
+class FWindowsD3D
+{
+public:
+	static WINDOWSD3D_API FWindowsGPUInfo GetExpectedGPUInfo();
+	static WINDOWSD3D_API FWindowsGPUInfo GetGPUInfoByDedicatedMemory();
+	static WINDOWSD3D_API uint64 ChooseD3D11Adapter(uint64 DesiredAdapterLuid);
+	static WINDOWSD3D_API uint64 ChooseD3D12Adapter(uint64 DesiredAdapterLuid);
+};
diff --git a/Engine/Source/Runtime/Windows/WindowsD3D/Private/WindowsD3D.cpp b/Engine/Source/Runtime/Windows/WindowsD3D/Private/WindowsD3D.cpp
new file mode 100644
--- /dev/null
+++ b/Engine/Source/Runtime/Windows/WindowsD3D/Private/WindowsD3D.cpp
@@ -0,0 +1,652 @@
+// Copyright Epic Games, Inc. All Rights Reserved.
+
+#include "Windows/WindowsD3D.h"
+
+#include "HAL/IConsoleManager.h"
+#include "Logging/LogMacros.h"
+#include "Misc/CommandLine.h"
+#include "Misc/Parse.h"
+#include "Templates/RefCounting.h"
+#include "Windows/D3D11ThirdParty.h"
+#include "Windows/WindowsD3D12ThirdParty.h"
+
+THIRD_PARTY_INCLUDES_START
+#include "Windows/AllowWindowsPlatformTypes.h"
+#include <delayimp.h>
+#include "Windows/HideWindowsPlatformTypes.h"
+THIRD_PARTY_INCLUDES_END
+
+DEFINE_LOG_CATEGORY_STATIC(LogD3D, Log, All);
+
+enum class ED3DGpuVendorId : uint32
+{
+	Unknown = 0xffffffff,
+	NotQueried = 0,
+
+	Amd = 0x1002,
+	ImgTec = 0x1010,
+	Nvidia = 0x10DE,
+	Arm = 0x13B5,
+	Broadcom = 0x14E4,
+	Qualcomm = 0x5143,
+	Intel = 0x8086,
+	Apple = 0x106B,
+	Vivante = 0x7a05,
+	VeriSilicon = 0x1EB1,
+	SamsungAMD = 0x144D,
+	Microsoft = 0x1414,
+
+	Kazan = 0x10003,	// VkVendorId
+	Codeplay = 0x10004,	// VkVendorId
+	Mesa = 0x10005,	// VkVendorId
+};
+
+static ED3DGpuVendorId GetPreferredGpuVendor()
+{
+	if (FParse::Param(FCommandLine::Get(), TEXT("preferAMD")))
+	{
+		return ED3DGpuVendorId::Amd;
+	}
+
+	if (FParse::Param(FCommandLine::Get(), TEXT("preferIntel")))
+	{
+		return ED3DGpuVendorId::Intel;
+	}
+
+	if (FParse::Param(FCommandLine::Get(), TEXT("preferNvidia")))
+	{
+		return ED3DGpuVendorId::Nvidia;
+	}
+
+	if (FParse::Param(FCommandLine::Get(), TEXT("preferMS")) || FParse::Param(FCommandLine::Get(), TEXT("preferMicrosoft")))
+	{
+		return ED3DGpuVendorId::Microsoft;
+	}
+
+	return ED3DGpuVendorId::Unknown;
+}
+
+static ED3DGpuVendorId ConvertToGpuVendorId(uint32 VendorId)
+{
+	switch ((ED3DGpuVendorId)VendorId)
+	{
+	case ED3DGpuVendorId::NotQueried:
+		return ED3DGpuVendorId::NotQueried;
+
+	case ED3DGpuVendorId::Amd:
+	case ED3DGpuVendorId::Kazan:
+	case ED3DGpuVendorId::Codeplay:
+	case ED3DGpuVendorId::Mesa:
+	case ED3DGpuVendorId::ImgTec:
+	case ED3DGpuVendorId::Nvidia:
+	case ED3DGpuVendorId::Arm:
+	case ED3DGpuVendorId::Broadcom:
+	case ED3DGpuVendorId::Qualcomm:
+	case ED3DGpuVendorId::Intel:
+	case ED3DGpuVendorId::Apple:
+	case ED3DGpuVendorId::Vivante:
+	case ED3DGpuVendorId::VeriSilicon:
+	case ED3DGpuVendorId::SamsungAMD:
+	case ED3DGpuVendorId::Microsoft:
+		return (ED3DGpuVendorId)VendorId;
+
+	case ED3DGpuVendorId::Unknown:
+	default:
+		return ED3DGpuVendorId::Unknown;
+	}
+}
+
+static FWindowsGPUInfo GetSpecificGPUInfo(uint64 AdapterLuid)
+{
+	TRefCountPtr<IDXGIFactory1> DXGIFactory1;
+	if (FAILED(CreateDXGIFactory1(IID_PPV_ARGS(DXGIFactory1.GetInitReference()))) || !DXGIFactory1)
+	{
+		return {};
+	}
+
+	DXGI_ADAPTER_DESC BestDesc = {};
+	TRefCountPtr<IDXGIAdapter> TempAdapter;
+	for (uint32 AdapterIndex = 0; DXGIFactory1->EnumAdapters(AdapterIndex, TempAdapter.GetInitReference()) != DXGI_ERROR_NOT_FOUND; ++AdapterIndex)
+	{
+		if (TempAdapter)
+		{
+			DXGI_ADAPTER_DESC Desc;
+			if (SUCCEEDED(TempAdapter->GetDesc(&Desc)))
+			{
+				if (!FMemory::Memcmp(&AdapterLuid, &Desc.AdapterLuid, sizeof(LUID)))
+				{
+					BestDesc = Desc;
+					break;
+				}
+
+				if (AdapterIndex == 0)
+				{
+					BestDesc = Desc;
+				}
+			}
+		}
+	}
+
+	return FWindowsGPUInfo{ BestDesc.VendorId, BestDesc.DeviceId, BestDesc.DedicatedVideoMemory, BestDesc.Description };
+}
+
+// Make a best effort to anticipate the GPU adapter that will be chosen.
+FWindowsGPUInfo FWindowsD3D::GetExpectedGPUInfo()
+{
+	// TODO: Respect IHeadMountedDisplayModule::GetGraphicsAdapterLuid without creating a circular reference.
+	const uint64 DesiredAdapaterLuid = 0;
+
+	uint64 Luid = ChooseD3D12Adapter(DesiredAdapaterLuid);
+	if (!Luid)
+	{
+		Luid = ChooseD3D11Adapter(DesiredAdapaterLuid);
+	}
+
+	return Luid ? GetSpecificGPUInfo(Luid) : GetGPUInfoByDedicatedMemory();
+}
+
+FWindowsGPUInfo FWindowsD3D::GetGPUInfoByDedicatedMemory()
+{
+	TRefCountPtr<IDXGIFactory1> DXGIFactory1;
+	if (CreateDXGIFactory1(IID_PPV_ARGS(DXGIFactory1.GetInitReference())) != S_OK || !DXGIFactory1)
+	{
+		return {};
+	}
+
+	DXGI_ADAPTER_DESC BestDesc = {};
+	TRefCountPtr<IDXGIAdapter> TempAdapter;
+	for (uint32 AdapterIndex = 0; DXGIFactory1->EnumAdapters(AdapterIndex, TempAdapter.GetInitReference()) != DXGI_ERROR_NOT_FOUND; ++AdapterIndex)
+	{
+		if (TempAdapter)
+		{
+			DXGI_ADAPTER_DESC Desc;
+			if (SUCCEEDED(TempAdapter->GetDesc(&Desc)))
+			{
+				if (Desc.DedicatedVideoMemory > BestDesc.DedicatedVideoMemory || AdapterIndex == 0)
+				{
+					BestDesc = Desc;
+				}
+			}
+		}
+	}
+
+	return FWindowsGPUInfo{ BestDesc.VendorId, BestDesc.DeviceId, BestDesc.DedicatedVideoMemory, BestDesc.Description };
+}
+
+static HRESULT EnumAdapters(int32 AdapterIndex, DXGI_GPU_PREFERENCE GpuPreference, IDXGIFactory1* DxgiFactory1, IDXGIFactory6* DxgiFactory6, IDXGIAdapter** TempAdapter)
+{
+	if (!DxgiFactory6 || GpuPreference == DXGI_GPU_PREFERENCE_UNSPECIFIED)
+	{
+		return DxgiFactory1->EnumAdapters(AdapterIndex, TempAdapter);
+	}
+	else
+	{
+		return DxgiFactory6->EnumAdapterByGpuPreference(AdapterIndex, GpuPreference, IID_PPV_ARGS(TempAdapter));
+	}
+}
+
+/** This function is used as a SEH filter to catch only delay load exceptions. */
+static bool IsDelayLoadException(PEXCEPTION_POINTERS ExceptionPointers)
+{
+	switch (ExceptionPointers->ExceptionRecord->ExceptionCode)
+	{
+	case VcppException(ERROR_SEVERITY_ERROR, ERROR_MOD_NOT_FOUND):
+	case VcppException(ERROR_SEVERITY_ERROR, ERROR_PROC_NOT_FOUND):
+		return EXCEPTION_EXECUTE_HANDLER;
+	default:
+		return EXCEPTION_CONTINUE_SEARCH;
+	}
+}
+
+/**
+ * Attempts to create a D3D12 device for the adapter using at minimum MinFeatureLevel.
+ * If creation is successful, true is returned.
+ */
+static bool SafeTestD3D12CreateDevice(IDXGIAdapter* Adapter, D3D_FEATURE_LEVEL MinFeatureLevel, bool& bOutUMA)
+{
+	bOutUMA = false;
+
+	__try
+	{
+		ID3D12Device* Device = nullptr;
+		const HRESULT D3D12CreateDeviceResult = D3D12CreateDevice(Adapter, MinFeatureLevel, IID_PPV_ARGS(&Device));
+		if (SUCCEEDED(D3D12CreateDeviceResult))
+		{
+			D3D12_FEATURE_DATA_ARCHITECTURE data{};
+			bOutUMA = SUCCEEDED(Device->CheckFeatureSupport(D3D12_FEATURE_ARCHITECTURE, &data, sizeof(D3D12_FEATURE_DATA_ARCHITECTURE))) && data.UMA;
+
+			Device->Release();
+			return true;
+		}
+		else
+		{
+			UE_LOG(LogD3D, Log, TEXT("D3D12CreateDevice failed with code 0x%08X"), static_cast<int32>(D3D12CreateDeviceResult));
+		}
+	}
+	__except (IsDelayLoadException(GetExceptionInformation()))
+	{
+		// We suppress warning C6322: Empty _except block. Appropriate checks are made upon returning. 
+		CA_SUPPRESS(6322);
+	}
+
+	return false;
+}
+
+/**
+ * Attempts to create a D3D11 device for the adapter using at most MaxFeatureLevel.
+ * If creation is successful, true is returned.
+ */
+static bool SafeTestD3D11CreateDevice(IDXGIAdapter* Adapter, D3D_FEATURE_LEVEL MinFeatureLevel, D3D_FEATURE_LEVEL MaxFeatureLevel, bool& bOutUMA)
+{
+	ID3D11Device* D3DDevice = nullptr;
+	ID3D11DeviceContext* D3DDeviceContext = nullptr;
+	uint32 DeviceFlags = D3D11_CREATE_DEVICE_SINGLETHREADED;
+	// Use a debug device if specified on the command line.
+	if (FParse::Param(FCommandLine::Get(), TEXT("d3ddebug")) ||
+		FParse::Param(FCommandLine::Get(), TEXT("d3debug")) ||
+		FParse::Param(FCommandLine::Get(), TEXT("dxdebug")) ||
+		FParse::Param(FCommandLine::Get(), TEXT("d3dlogwarnings")) ||
+		FParse::Param(FCommandLine::Get(), TEXT("d3dbreakonwarning")) ||
+		FParse::Param(FCommandLine::Get(), TEXT("d3dcontinueonerrors")))
+	{
+		DeviceFlags |= D3D11_CREATE_DEVICE_DEBUG;
+	}
+
+	// @MIXEDREALITY_CHANGE : BEGIN - Add BGRA flag for Windows Mixed Reality HMD's
+	DeviceFlags |= D3D11_CREATE_DEVICE_BGRA_SUPPORT;
+	// @MIXEDREALITY_CHANGE : END
+
+	D3D_FEATURE_LEVEL RequestedFeatureLevels[] =
+	{
+		D3D_FEATURE_LEVEL_11_1,
+		D3D_FEATURE_LEVEL_11_0,
+	};
+
+	// Trim to allowed feature levels
+	int32 FirstAllowedFeatureLevel = 0;
+	int32 NumAllowedFeatureLevels = UE_ARRAY_COUNT(RequestedFeatureLevels);
+	int32 LastAllowedFeatureLevel = NumAllowedFeatureLevels - 1;
+
+	while (FirstAllowedFeatureLevel < NumAllowedFeatureLevels)
+	{
+		if (RequestedFeatureLevels[FirstAllowedFeatureLevel] == MaxFeatureLevel)
+		{
+			break;
+		}
+		FirstAllowedFeatureLevel++;
+	}
+
+	while (LastAllowedFeatureLevel > 0)
+	{
+		if (RequestedFeatureLevels[LastAllowedFeatureLevel] >= MinFeatureLevel)
+		{
+			break;
+		}
+		LastAllowedFeatureLevel--;
+	}
+
+	NumAllowedFeatureLevels = LastAllowedFeatureLevel - FirstAllowedFeatureLevel + 1;
+	if (MaxFeatureLevel < MinFeatureLevel || NumAllowedFeatureLevels <= 0)
+	{
+		return false;
+	}
+
+	__try
+	{
+		// We don't want software renderer. Ideally we specify D3D_DRIVER_TYPE_HARDWARE on creation but
+		// when we specify an adapter we need to specify D3D_DRIVER_TYPE_UNKNOWN (otherwise the call fails).
+		// We cannot check the device type later (seems this is missing functionality in D3D).
+		D3D_FEATURE_LEVEL DeviceFeatureLevel = (D3D_FEATURE_LEVEL)0;
+		HRESULT Result = D3D11CreateDevice(
+			Adapter,
+			D3D_DRIVER_TYPE_UNKNOWN,
+			nullptr,
+			DeviceFlags,
+			&RequestedFeatureLevels[FirstAllowedFeatureLevel],
+			NumAllowedFeatureLevels,
+			D3D11_SDK_VERSION,
+			&D3DDevice,
+			&DeviceFeatureLevel,
+			&D3DDeviceContext);
+
+		if (SUCCEEDED(Result))
+		{
+			D3D11_FEATURE_DATA_D3D11_OPTIONS2 FeatureData = {};
+			bOutUMA = SUCCEEDED(D3DDevice->CheckFeatureSupport(D3D11_FEATURE_D3D11_OPTIONS2, &FeatureData, sizeof(FeatureData))) && FeatureData.UnifiedMemoryArchitecture;
+
+			D3DDevice->Release();
+			D3DDeviceContext->Release();
+			return true;
+		}
+
+		// Log any reason for failure to create test device. Extra debug help.
+		UE_LOG(LogD3D, Log, TEXT("D3D11CreateDevice failed with code 0x%08X"), static_cast<uint32>(Result));
+
+		// Fatal error on 0x887A002D
+		if (DXGI_ERROR_SDK_COMPONENT_MISSING == Result)
+		{
+			UE_LOG(LogD3D, Fatal, TEXT("-d3ddebug was used but optional Graphics Tools were not found. Install them through the Manage Optional Features in windows. See: https://docs.microsoft.com/en-us/windows/uwp/gaming/use-the-directx-runtime-and-visual-studio-graphics-diagnostic-features"));
+		}
+	}
+	__except (IsDelayLoadException(GetExceptionInformation()))
+	{
+		// We suppress warning C6322: Empty _except block. Appropriate checks are made upon returning. 
+		CA_SUPPRESS(6322);
+	}
+
+	return false;
+}
+
+static uint64 ConvertLuidToUint64(LUID Luid)
+{
+	return (uint64_t)Luid.LowPart | ((uint64_t)Luid.HighPart << 32);
+}
+
+uint64 FWindowsD3D::ChooseD3D11Adapter(uint64 DesiredAdapterLuid)
+{
+	static uint64 CachedAdapterLuid = 0;
+	if (!DesiredAdapterLuid && CachedAdapterLuid)
+	{
+		return CachedAdapterLuid;
+	}
+
+	TRefCountPtr<IDXGIFactory1> DXGIFactory1;
+	CreateDXGIFactory1(IID_PPV_ARGS(DXGIFactory1.GetInitReference()));
+	if (!DXGIFactory1)
+	{
+		return 0;
+	}
+
+	TRefCountPtr<IDXGIFactory6> DXGIFactory6;
+	DXGIFactory1->QueryInterface(IID_PPV_ARGS(DXGIFactory6.GetInitReference()));
+
+	bool bAllowPerfHUD = true;
+
+#if UE_BUILD_SHIPPING || UE_BUILD_TEST
+	bAllowPerfHUD = false;
+#endif
+
+	// Non-static as it is used only a few times
+	const TConsoleVariableData<int32>* const CVarGraphicsAdapter = IConsoleManager::Get().FindTConsoleVariableDataInt(TEXT("r.GraphicsAdapter"));
+	int32 CVarExplicitAdapterValue = DesiredAdapterLuid == 0 ? (CVarGraphicsAdapter ? CVarGraphicsAdapter->GetValueOnGameThread() : -1) : -2;
+	FParse::Value(FCommandLine::Get(), TEXT("graphicsadapter="), CVarExplicitAdapterValue);
+
+	const bool bFavorDiscreteAdapter = CVarExplicitAdapterValue == -1;
+
+	TRefCountPtr<IDXGIAdapter> TempAdapter;
+	D3D_FEATURE_LEVEL MinAllowedFeatureLevel = D3D_FEATURE_LEVEL_11_0;
+	D3D_FEATURE_LEVEL MaxAllowedFeatureLevel = D3D_FEATURE_LEVEL_11_1;
+
+	TOptional<LUID> PreferredAdapterLuid;
+	TOptional<LUID> FirstDiscreteAdapterLuid;
+	TOptional<LUID> BestMemoryAdapterLuid;
+	TOptional<LUID> FirstAdapterLuid;
+	SIZE_T BestDedicatedVideoMemory = 0;
+
+	int GpuPreferenceInt = DXGI_GPU_PREFERENCE_HIGH_PERFORMANCE;
+	FParse::Value(FCommandLine::Get(), TEXT("-gpupreference="), GpuPreferenceInt);
+	DXGI_GPU_PREFERENCE GpuPreference;
+	switch (GpuPreferenceInt)
+	{
+	case 1: GpuPreference = DXGI_GPU_PREFERENCE_MINIMUM_POWER; break;
+	case 2: GpuPreference = DXGI_GPU_PREFERENCE_HIGH_PERFORMANCE; break;
+	default: GpuPreference = DXGI_GPU_PREFERENCE_UNSPECIFIED; break;
+	}
+
+	const ED3DGpuVendorId PreferredVendor = GetPreferredGpuVendor();
+	const bool bAllowSoftwareFallback = FParse::Param(FCommandLine::Get(), TEXT("AllowSoftwareRendering"));
+
+	// Enumerate the DXGIFactory's adapters.
+	for (uint32 AdapterIndex = 0; EnumAdapters(AdapterIndex, GpuPreference, DXGIFactory1, DXGIFactory6, TempAdapter.GetInitReference()) != DXGI_ERROR_NOT_FOUND; ++AdapterIndex)
+	{
+		// Check that if adapter supports D3D11.
+		if (TempAdapter)
+		{
+			DXGI_ADAPTER_DESC AdapterDesc{};
+			bool bUMA;
+			if (SUCCEEDED(TempAdapter->GetDesc(&AdapterDesc)) && SafeTestD3D11CreateDevice(TempAdapter, MinAllowedFeatureLevel, MaxAllowedFeatureLevel, bUMA))
+			{
+				const bool bIsWARP = ConvertToGpuVendorId(AdapterDesc.VendorId) == ED3DGpuVendorId::Microsoft;
+
+				// PerfHUD is for performance profiling
+				const bool bIsPerfHUD = !FCString::Stricmp(AdapterDesc.Description, TEXT("NVIDIA PerfHUD"));
+
+				// Add special check to support HMDs, which do not have associated outputs.
+				// To reject the software emulation, unless the cvar wants it.
+				// https://msdn.microsoft.com/en-us/library/windows/desktop/bb205075(v=vs.85).aspx#WARP_new_for_Win8
+				// Before we tested for no output devices but that failed where a laptop had a Intel (with output) and NVidia (with no output)
+				const bool bSkipSoftwareAdapter = bIsWARP && !bAllowSoftwareFallback && CVarExplicitAdapterValue < 0 && DesiredAdapterLuid == 0;
+
+				// we don't allow the PerfHUD adapter
+				const bool bSkipPerfHUDAdapter = bIsPerfHUD && !bAllowPerfHUD;
+
+				// the caller wants a specific adapter, not this one
+				const bool bSkipDesiredAdapter = DesiredAdapterLuid != 0 && FMemory::Memcmp(&DesiredAdapterLuid, &AdapterDesc.AdapterLuid, sizeof(LUID)) != 0;
+
+				// the user wants a specific adapter, not this one
+				const bool bSkipExplicitAdapter = CVarExplicitAdapterValue >= 0 && AdapterIndex != CVarExplicitAdapterValue;
+
+				const bool bSkipAdapter = bSkipSoftwareAdapter || bSkipPerfHUDAdapter || bSkipDesiredAdapter || bSkipExplicitAdapter || FParse::Param(FCommandLine::Get(), TEXT("ForceZeroAdapters"));
+
+				const bool bIsIntegrated = bUMA;
+
+				if (!bSkipAdapter)
+				{
+					if (PreferredVendor != ED3DGpuVendorId::Unknown && PreferredVendor == ConvertToGpuVendorId(AdapterDesc.VendorId) && !PreferredAdapterLuid.IsSet())
+					{
+						PreferredAdapterLuid = AdapterDesc.AdapterLuid;
+					}
+
+					if (!bIsWARP && !bIsIntegrated)
+					{
+						if (!FirstDiscreteAdapterLuid.IsSet())
+						{
+							FirstDiscreteAdapterLuid = AdapterDesc.AdapterLuid;
+						}
+
+						if (AdapterDesc.DedicatedVideoMemory > BestDedicatedVideoMemory)
+						{
+							BestMemoryAdapterLuid = AdapterDesc.AdapterLuid;
+							BestDedicatedVideoMemory = AdapterDesc.DedicatedVideoMemory;
+							if (PreferredVendor != ED3DGpuVendorId::Unknown && PreferredVendor == ConvertToGpuVendorId(AdapterDesc.VendorId))
+							{
+								// Choose the best option of the preferred IHV devices
+								PreferredAdapterLuid = BestMemoryAdapterLuid;
+							}
+						}
+					}
+
+					if (!FirstAdapterLuid.IsSet())
+					{
+						FirstAdapterLuid = AdapterDesc.AdapterLuid;
+					}
+				}
+			}
+		}
+	}
+
+	uint64 Luid = 0;
+	if (bFavorDiscreteAdapter)
+	{
+		if (PreferredAdapterLuid.IsSet())
+		{
+			Luid = ConvertLuidToUint64(PreferredAdapterLuid.GetValue());
+		}
+		else if (BestMemoryAdapterLuid.IsSet())
+		{
+			Luid = ConvertLuidToUint64(BestMemoryAdapterLuid.GetValue());
+		}
+		else if (FirstDiscreteAdapterLuid.IsSet())
+		{
+			Luid = ConvertLuidToUint64(FirstDiscreteAdapterLuid.GetValue());
+		}
+		else if (FirstAdapterLuid.IsSet())
+		{
+			Luid = ConvertLuidToUint64(FirstAdapterLuid.GetValue());
+		}
+	}
+	else if (FirstAdapterLuid.IsSet())
+	{
+		Luid = ConvertLuidToUint64(FirstAdapterLuid.GetValue());
+	}
+
+	if (!DesiredAdapterLuid)
+	{
+		CachedAdapterLuid = Luid;
+	}
+	return Luid;
+}
+
+uint64 FWindowsD3D::ChooseD3D12Adapter(uint64 DesiredAdapterLuid)
+{
+	static uint64 CachedAdapterLuid = 0;
+	if (!DesiredAdapterLuid && CachedAdapterLuid)
+	{
+		return CachedAdapterLuid;
+	}
+
+	TRefCountPtr<IDXGIFactory1> DXGIFactory1;
+	CreateDXGIFactory1(IID_PPV_ARGS(DXGIFactory1.GetInitReference()));
+	if (!DXGIFactory1)
+	{
+		return 0;
+	}
+
+	TRefCountPtr<IDXGIFactory6> DXGIFactory6;
+	DXGIFactory1->QueryInterface(__uuidof(IDXGIFactory6), (void**)DXGIFactory6.GetInitReference());
+
+	bool bAllowPerfHUD = true;
+
+#if UE_BUILD_SHIPPING || UE_BUILD_TEST
+	bAllowPerfHUD = false;
+#endif
+
+	// Non-static as it is used only a few times
+	const TConsoleVariableData<int32>* const CVarGraphicsAdapter = IConsoleManager::Get().FindTConsoleVariableDataInt(TEXT("r.GraphicsAdapter"));
+	int32 CVarExplicitAdapterValue = DesiredAdapterLuid == 0 ? (CVarGraphicsAdapter ? CVarGraphicsAdapter->GetValueOnGameThread() : -1) : -2;
+	FParse::Value(FCommandLine::Get(), TEXT("graphicsadapter="), CVarExplicitAdapterValue);
+
+	const bool bFavorDiscreteAdapter = CVarExplicitAdapterValue == -1;
+
+	TRefCountPtr<IDXGIAdapter> TempAdapter;
+	const D3D_FEATURE_LEVEL MinRequiredFeatureLevel = D3D_FEATURE_LEVEL_11_0;
+
+	TOptional<LUID> PreferredAdapterLuid;
+	TOptional<LUID> FirstDiscreteAdapterLuid;
+	TOptional<LUID> BestMemoryAdapterLuid;
+	TOptional<LUID> FirstAdapterLuid;
+	SIZE_T BestDedicatedVideoMemory = 0;
+
+	static bool bRequestedWARP = FParse::Param(FCommandLine::Get(), TEXT("warp"));
+	static bool bAllowSoftwareRendering = FParse::Param(FCommandLine::Get(), TEXT("AllowSoftwareRendering"));
+
+	int GpuPreferenceInt = DXGI_GPU_PREFERENCE_HIGH_PERFORMANCE;
+	FParse::Value(FCommandLine::Get(), TEXT("-gpupreference="), GpuPreferenceInt);
+	DXGI_GPU_PREFERENCE GpuPreference;
+	switch (GpuPreferenceInt)
+	{
+	case 1: GpuPreference = DXGI_GPU_PREFERENCE_MINIMUM_POWER; break;
+	case 2: GpuPreference = DXGI_GPU_PREFERENCE_HIGH_PERFORMANCE; break;
+	default: GpuPreference = DXGI_GPU_PREFERENCE_UNSPECIFIED; break;
+	}
+
+	const ED3DGpuVendorId PreferredVendor = GetPreferredGpuVendor();
+
+	// Enumerate the DXGIFactory's adapters.
+	for (uint32 AdapterIndex = 0; EnumAdapters(AdapterIndex, GpuPreference, DXGIFactory1, DXGIFactory6, TempAdapter.GetInitReference()) != DXGI_ERROR_NOT_FOUND; ++AdapterIndex)
+	{
+		// Check that if adapter supports D3D12.
+		if (TempAdapter)
+		{
+			bool bUMA;
+			DXGI_ADAPTER_DESC AdapterDesc{};
+			if (SUCCEEDED(TempAdapter->GetDesc(&AdapterDesc)) && SafeTestD3D12CreateDevice(TempAdapter, MinRequiredFeatureLevel, bUMA))
+			{
+				const bool bIsWARP = ConvertToGpuVendorId(AdapterDesc.VendorId) == ED3DGpuVendorId::Microsoft;
+
+				// If requested WARP, then reject all other adapters. If WARP not requested, then reject the WARP device if software rendering support is disallowed
+				const bool bSkipWARP = (bRequestedWARP && !bIsWARP) || (!bRequestedWARP && bIsWARP && !bAllowSoftwareRendering);
+
+				// PerfHUD is for performance profiling
+				const bool bIsPerfHUD = !FCString::Stricmp(AdapterDesc.Description, TEXT("NVIDIA PerfHUD"));
+
+				// we don't allow the PerfHUD adapter
+				const bool bSkipPerfHUDAdapter = bIsPerfHUD && !bAllowPerfHUD;
+
+				// the caller wants a specific adapter, not this one
+				const bool bSkipDesiredAdapter = DesiredAdapterLuid != 0 && FMemory::Memcmp(&DesiredAdapterLuid, &AdapterDesc.AdapterLuid, sizeof(LUID)) != 0;
+
+				// the user wants a specific adapter, not this one
+				const bool bSkipExplicitAdapter = CVarExplicitAdapterValue >= 0 && AdapterIndex != CVarExplicitAdapterValue;
+
+				const bool bSkipAdapter = bSkipWARP || bSkipPerfHUDAdapter || bSkipDesiredAdapter || bSkipExplicitAdapter || FParse::Param(FCommandLine::Get(), TEXT("ForceZeroAdapters"));
+
+				const bool bIsIntegrated = bUMA;
+
+				if (!bSkipAdapter)
+				{
+					if (PreferredVendor != ED3DGpuVendorId::Unknown && PreferredVendor == ConvertToGpuVendorId(AdapterDesc.VendorId) && !PreferredAdapterLuid.IsSet())
+					{
+						PreferredAdapterLuid = AdapterDesc.AdapterLuid;
+					}
+
+					if (!bIsWARP && !bIsIntegrated)
+					{
+						if (!FirstDiscreteAdapterLuid.IsSet())
+						{
+							FirstDiscreteAdapterLuid = AdapterDesc.AdapterLuid;
+						}
+
+						if (AdapterDesc.DedicatedVideoMemory > BestDedicatedVideoMemory)
+						{
+							BestMemoryAdapterLuid = AdapterDesc.AdapterLuid;
+							BestDedicatedVideoMemory = AdapterDesc.DedicatedVideoMemory;
+							if (PreferredVendor != ED3DGpuVendorId::Unknown && PreferredVendor == ConvertToGpuVendorId(AdapterDesc.VendorId))
+							{
+								// Choose the best option of the preferred IHV devices
+								PreferredAdapterLuid = BestMemoryAdapterLuid;
+							}
+						}
+					}
+
+					if (!FirstAdapterLuid.IsSet())
+					{
+						FirstAdapterLuid = AdapterDesc.AdapterLuid;
+					}
+				}
+			}
+		}
+	}
+
+	uint64 Luid = 0;
+	if (bFavorDiscreteAdapter)
+	{
+		// We assume Intel is integrated graphics (slower than discrete) than NVIDIA or AMD cards and rather take a different one
+		if (PreferredAdapterLuid.IsSet())
+		{
+			Luid = ConvertLuidToUint64(PreferredAdapterLuid.GetValue());
+		}
+		else if (BestMemoryAdapterLuid.IsSet())
+		{
+			Luid = ConvertLuidToUint64(BestMemoryAdapterLuid.GetValue());
+		}
+		else if (FirstDiscreteAdapterLuid.IsSet())
+		{
+			Luid = ConvertLuidToUint64(FirstDiscreteAdapterLuid.GetValue());
+		}
+		else if (FirstAdapterLuid.IsSet())
+		{
+			Luid = ConvertLuidToUint64(FirstAdapterLuid.GetValue());
+		}
+	}
+	else if (FirstAdapterLuid.IsSet())
+	{
+		Luid = ConvertLuidToUint64(FirstAdapterLuid.GetValue());
+	}
+
+	if (!DesiredAdapterLuid)
+	{
+		CachedAdapterLuid = Luid;
+	}
+	return Luid;
+}
diff --git a/Engine/Source/Runtime/Windows/WindowsD3D/Public/Windows/D3D11ThirdParty.h b/Engine/Source/Runtime/Windows/WindowsD3D/Public/Windows/D3D11ThirdParty.h
new file mode 100644
--- /dev/null
+++ b/Engine/Source/Runtime/Windows/WindowsD3D/Public/Windows/D3D11ThirdParty.h
@@ -0,0 +1,20 @@
+// Copyright Epic Games, Inc. All Rights Reserved.
+
+#pragma once
+
+#include "HAL/Platform.h"
+#include "Microsoft/AllowMicrosoftPlatformTypes.h"
+
+#pragma pack(push,16)
+
+THIRD_PARTY_INCLUDES_START
+	#include <d3d11_4.h>
+	#include <dxgi1_6.h>
+	#include <dxgidebug.h>
+THIRD_PARTY_INCLUDES_END
+
+#undef DrawText
+
+#pragma pack(pop)
+
+#include "Microsoft/HideMicrosoftPlatformTypes.h"
diff --git a/Engine/Source/Runtime/Windows/WindowsD3D/Public/Windows/WindowsD3D12ThirdParty.h b/Engine/Source/Runtime/Windows/WindowsD3D/Public/Windows/WindowsD3D12ThirdParty.h
new file mode 100644
--- /dev/null
+++ b/Engine/Source/Runtime/Windows/WindowsD3D/Public/Windows/WindowsD3D12ThirdParty.h
@@ -0,0 +1,22 @@
+// Copyright Epic Games, Inc. All Rights Reserved.
+
+#pragma once
+
+#include "HAL/Platform.h"
+
+#include "Microsoft/AllowMicrosoftPlatformTypes.h"
+
+#pragma pack(push,16)
+#define INITGUID
+THIRD_PARTY_INCLUDES_START
+	#include <d3d12.h>
+	#include <d3dx12.h>
+	#include <dxgi1_6.h>
+	#include <dxgidebug.h>
+THIRD_PARTY_INCLUDES_END
+
+#undef DrawText
+
+#pragma pack(pop)
+
+#include "Microsoft/HideMicrosoftPlatformTypes.h"
diff --git a/Engine/Source/Runtime/Windows/WindowsD3D/WindowsD3D.Build.cs b/Engine/Source/Runtime/Windows/WindowsD3D/WindowsD3D.Build.cs
new file mode 100644
--- /dev/null
+++ b/Engine/Source/Runtime/Windows/WindowsD3D/WindowsD3D.Build.cs
@@ -0,0 +1,17 @@
+// Copyright Epic Games, Inc. All Rights Reserved.
+
+using UnrealBuildTool;
+
+public class WindowsD3D : ModuleRules
+{
+	public WindowsD3D(ReadOnlyTargetRules Target) : base(Target)
+	{
+		PublicDependencyModuleNames.AddRange(
+			new string[] {
+				"Core"
+			}
+		);
+
+		AddEngineThirdPartyPrivateStaticDependencies(Target, "DX11", "DX12");
+	}
+}
