AVEncoder: Set keyframe flag on video samples that are keyframes and on all audio samples (which are always keyframes) so that the mp4 writer writes proper sync sample information to the file

#rnx

#Synced-CL 51333923

diff --git a/Engine/Source/Runtime/AVEncoder/Private/IMFSampleWrapper.cpp b/Engine/Source/Runtime/AVEncoder/Private/IMFSampleWrapper.cpp
--- a/Engine/Source/Runtime/AVEncoder/Private/IMFSampleWrapper.cpp
+++ b/Engine/Source/Runtime/AVEncoder/Private/IMFSampleWrapper.cpp
@@ -71,6 +71,11 @@ void FIMFSampleWrapper::SetDuration(FTimespan Duration)
 	verify(SUCCEEDED(Sample->SetSampleDuration(Duration.GetTicks())));
 }
 
+void FIMFSampleWrapper::SetIsKeyFrame()
+{
+	verify(SUCCEEDED(Sample->SetUINT32(MFSampleExtension_CleanPoint, 1)));
+}
+
 bool FIMFSampleWrapper::IsVideoKeyFrame() const
 {
 	PRAGMA_DISABLE_DEPRECATION_WARNINGS
diff --git a/Engine/Source/Runtime/AVEncoder/Public/IMFSampleWrapper.h b/Engine/Source/Runtime/AVEncoder/Public/IMFSampleWrapper.h
--- a/Engine/Source/Runtime/AVEncoder/Public/IMFSampleWrapper.h
+++ b/Engine/Source/Runtime/AVEncoder/Public/IMFSampleWrapper.h
@@ -46,6 +46,8 @@ public:
 
 	AVENCODER_API void SetDuration(FTimespan Duration);
 
+	AVENCODER_API void SetIsKeyFrame();
+
 	AVENCODER_API bool IsVideoKeyFrame() const;
 
 	AVENCODER_API int GetBufferCount() const;
diff --git a/Engine/Source/Runtime/Windows/WindowsPlatformFeatures/Private/WmfMp4Writer.cpp b/Engine/Source/Runtime/Windows/WindowsPlatformFeatures/Private/WmfMp4Writer.cpp
--- a/Engine/Source/Runtime/Windows/WindowsPlatformFeatures/Private/WmfMp4Writer.cpp
+++ b/Engine/Source/Runtime/Windows/WindowsPlatformFeatures/Private/WmfMp4Writer.cpp
@@ -126,6 +126,10 @@ bool FWmfMp4Writer::Write(const AVEncoder::FMediaPacket& InSample, DWORD StreamI
 	Sample.GetSample()->AddBuffer(WmfBuffer);
 	Sample.SetTime(InSample.Timestamp);
 	Sample.SetDuration(InSample.Duration);
+	if(InSample.IsVideoKeyFrame() || InSample.Type == AVEncoder::EPacketType::Audio)
+	{
+		Sample.SetIsKeyFrame();
+	}
 
 	CHECK_HR(Writer->WriteSample(StreamIndex, const_cast<IMFSample*>(Sample.GetSample())));
 
