Fixing D3D11 RHI selection
- D3D11RHI::IsSupported(ERHIFeatureLevel::SM6) was returning true since only the IsSupported() method was being overridden.

#Synced-CL 51502459

diff --git a/Engine/Source/Runtime/Windows/D3D11RHI/Private/D3D11RHIPrivate.h b/Engine/Source/Runtime/Windows/D3D11RHI/Private/D3D11RHIPrivate.h
--- a/Engine/Source/Runtime/Windows/D3D11RHI/Private/D3D11RHIPrivate.h
+++ b/Engine/Source/Runtime/Windows/D3D11RHI/Private/D3D11RHIPrivate.h
@@ -797,6 +797,7 @@ public:
 
 	// IDynamicRHIModule
 	virtual bool IsSupported() override;
+	virtual bool IsSupported(ERHIFeatureLevel::Type RequestedFeatureLevel) override;
 	virtual FDynamicRHI* CreateRHI(ERHIFeatureLevel::Type RequestedFeatureLevel = ERHIFeatureLevel::Num) override;
 
 private:
diff --git a/Engine/Source/Runtime/Windows/D3D11RHI/Private/Windows/WindowsD3D11Device.cpp b/Engine/Source/Runtime/Windows/D3D11RHI/Private/Windows/WindowsD3D11Device.cpp
--- a/Engine/Source/Runtime/Windows/D3D11RHI/Private/Windows/WindowsD3D11Device.cpp
+++ b/Engine/Source/Runtime/Windows/D3D11RHI/Private/Windows/WindowsD3D11Device.cpp
@@ -621,6 +621,10 @@ bool FD3D11DynamicRHIModule::IsSupported()
 		&& ChosenAdapter.MaxSupportedFeatureLevel >= D3D_FEATURE_LEVEL_11_0;
 }
 
+bool FD3D11DynamicRHIModule::IsSupported(ERHIFeatureLevel::Type RequestedFeatureLevel)
+{
+	return RequestedFeatureLevel <= ERHIFeatureLevel::SM5 && IsSupported();
+}
 
 const TCHAR* GetFeatureLevelString(D3D_FEATURE_LEVEL FeatureLevel)
 {
