Fixes for floating-point precision swapchain in SDR mode (backbuffer renderer setting).

* Ensure floating-point swapchain color space is set to G10 in SDR. Attempting to set G22 on a float16 backbuffer is actively rejected: "DXGI ERROR: IDXGISwapChain3::SetColorSpace1: Unsupported ColorSpace specified. [ MISCELLANEOUS ERROR #262: ]" In our case, the "CheckColorSpaceSupport" call fails and we print an "unabled to set color space" warning. In my testing, ignoring this error yields unpredictable behavior with the Desktop Window Manager, depending on whether displays support Advanced Color modes. The only reliable path with DXGI_FORMAT_R16G16B16A16_FLOAT is to use linear colors.
* Ensure that secondary windows retain the (linear) floating point pixel format.

#Synced-CL 51749116

diff --git a/Engine/Source/Runtime/D3D12RHI/Private/D3D12Viewport.h b/Engine/Source/Runtime/D3D12RHI/Private/D3D12Viewport.h
--- a/Engine/Source/Runtime/D3D12RHI/Private/D3D12Viewport.h
+++ b/Engine/Source/Runtime/D3D12RHI/Private/D3D12Viewport.h
@@ -198,6 +198,9 @@ private:
 	void ShutdownHDR();
 
 #if D3D12RHI_USE_DXGI_COLOR_SPACE
+	/** Ensure the correct color space is set on the swap chain */
+	void EnsureColorSpace(DXGI_COLOR_SPACE_TYPE NewColorSpace);
+
 	/** Ensure the correct color space is set on the swap chain */
 	void EnsureColorSpace(EDisplayColorGamut DisplayGamut, EDisplayOutputFormat OutputDevice);
 #endif
diff --git a/Engine/Source/Runtime/D3D12RHI/Private/Windows/WindowsD3D12Viewport.cpp b/Engine/Source/Runtime/D3D12RHI/Private/Windows/WindowsD3D12Viewport.cpp
--- a/Engine/Source/Runtime/D3D12RHI/Private/Windows/WindowsD3D12Viewport.cpp
+++ b/Engine/Source/Runtime/D3D12RHI/Private/Windows/WindowsD3D12Viewport.cpp
@@ -408,6 +408,13 @@ void FD3D12Viewport::ShutdownHDR()
 {
 	// Make sure to set the appropriate color space even if GRHISupportsHDROutput is false because we 
 	// might have toggled HDR on and off in the windows settings
+	if (PixelFormat == PF_FloatRGBA)
+	{
+		// Per DXGI spec, float-valued swapchains must contain linear data.
+		// The swapchain color space must therefore be G10 instead of G22.
+		EnsureColorSpace(DXGI_COLOR_SPACE_RGB_FULL_G10_NONE_P709);
+	}
+	else
 	{
 		// Ensure we have the correct color space set.
 		EnsureColorSpace(EDisplayColorGamut::sRGB_D65, EDisplayOutputFormat::SDR_sRGB);
@@ -472,13 +479,33 @@ static const FString GetDXGIColorSpaceString(DXGI_COLOR_SPACE_TYPE ColorSpace)
 	return FString::FromInt(ColorSpace);
 };
 
-void FD3D12Viewport::EnsureColorSpace(EDisplayColorGamut DisplayGamut, EDisplayOutputFormat OutputDevice)
+void FD3D12Viewport::EnsureColorSpace(DXGI_COLOR_SPACE_TYPE NewColorSpace)
 {
 	if (!SwapChain4.GetReference())
 	{
 		return;
 	}
 
+	if (ColorSpace != NewColorSpace)
+	{
+		uint32 ColorSpaceSupport = 0;
+		HRESULT hr = SwapChain4->CheckColorSpaceSupport(NewColorSpace, &ColorSpaceSupport);
+		FString NewColorSpaceName = GetDXGIColorSpaceString(NewColorSpace);
+		if (SUCCEEDED(hr) && ((ColorSpaceSupport & DXGI_SWAP_CHAIN_COLOR_SPACE_SUPPORT_FLAG_PRESENT) != 0))
+		{
+			VERIFYD3D12RESULT(SwapChain4->SetColorSpace1(NewColorSpace));
+			UE_LOG(LogD3D12RHI, Verbose, TEXT("Setting color space on swap chain (%p): %s"), SwapChain4.GetReference(), *NewColorSpaceName);
+			ColorSpace = NewColorSpace;
+		}
+		else
+		{
+			UE_LOG(LogD3D12RHI, Error, TEXT("Warning: unable to set color space %s to the swapchain: verify EDisplayOutputFormat / swapchain format"), *NewColorSpaceName);
+		}
+	}
+}
+
+void FD3D12Viewport::EnsureColorSpace(EDisplayColorGamut DisplayGamut, EDisplayOutputFormat OutputDevice)
+{
 	DXGI_COLOR_SPACE_TYPE NewColorSpace = DXGI_COLOR_SPACE_RGB_FULL_G22_NONE_P709;	// sRGB;
 	const bool bPrimaries2020 = (DisplayGamut == EDisplayColorGamut::Rec2020_D65);
 
@@ -506,22 +533,7 @@ void FD3D12Viewport::EnsureColorSpace(EDisplayColorGamut DisplayGamut, EDisplayO
 		break;
 	}
 
-	if (ColorSpace != NewColorSpace)
-	{
-		uint32 ColorSpaceSupport = 0;
-		HRESULT hr = SwapChain4->CheckColorSpaceSupport(NewColorSpace, &ColorSpaceSupport);
-		FString NewColorSpaceName = GetDXGIColorSpaceString(NewColorSpace);
-		if (SUCCEEDED(hr) && ((ColorSpaceSupport & DXGI_SWAP_CHAIN_COLOR_SPACE_SUPPORT_FLAG_PRESENT) != 0))
-		{
-			VERIFYD3D12RESULT(SwapChain4->SetColorSpace1(NewColorSpace));
-			UE_LOG(LogD3D12RHI, Verbose, TEXT("Setting color space on swap chain (%#016llx): %s"), SwapChain4.GetReference(), *NewColorSpaceName);
-			ColorSpace = NewColorSpace;
-		}
-		else
-		{
-			UE_LOG(LogD3D12RHI, Error, TEXT("Warning: unabled to set color space %s to the swapchain: verify EDisplayOutputFormat / swapchain format"), *NewColorSpaceName);
-		}
-	}
+	EnsureColorSpace(NewColorSpace);
 }
 
 void FD3D12Viewport::OnResumeRendering()
diff --git a/Engine/Source/Runtime/SlateRHIRenderer/Private/SlateRHIRenderer.cpp b/Engine/Source/Runtime/SlateRHIRenderer/Private/SlateRHIRenderer.cpp
--- a/Engine/Source/Runtime/SlateRHIRenderer/Private/SlateRHIRenderer.cpp
+++ b/Engine/Source/Runtime/SlateRHIRenderer/Private/SlateRHIRenderer.cpp
@@ -742,7 +742,7 @@ EPixelFormat FSlateRHIRenderer::GetViewportPixelFormat(const SWindow& Window, bo
 	// Use a known default format in VR / Mobile / Transparent Window SDR configurations.
 	if (bIsStandaloneStereoOnlyDevice
 #if ALPHA_BLENDED_WINDOWS
-		|| Window.GetTransparencySupport() == EWindowTransparency::PerPixel
+		|| ((Window.GetTransparencySupport() == EWindowTransparency::PerPixel) && !IsDefaultBackBufferLinearSDR())
 #endif
 		)
 	{
diff --git a/Engine/Source/Runtime/Windows/D3D11RHI/Private/Windows/WindowsD3D11Viewport.cpp b/Engine/Source/Runtime/Windows/D3D11RHI/Private/Windows/WindowsD3D11Viewport.cpp
--- a/Engine/Source/Runtime/Windows/D3D11RHI/Private/Windows/WindowsD3D11Viewport.cpp
+++ b/Engine/Source/Runtime/Windows/D3D11RHI/Private/Windows/WindowsD3D11Viewport.cpp
@@ -303,7 +303,7 @@ static const FString GetDXGIColorSpaceString(DXGI_COLOR_SPACE_TYPE ColorSpace)
 	return FString::FromInt(ColorSpace);
 };
 
-inline void EnsureColorSpace(IDXGISwapChain* SwapChain, EDisplayColorGamut DisplayGamut, EDisplayOutputFormat OutputDevice, EPixelFormat PixelFormat)
+inline void EnsureColorSpace(IDXGISwapChain* SwapChain, DXGI_COLOR_SPACE_TYPE NewColorSpace)
 {
 	if (SwapChain == nullptr)
 	{
@@ -316,6 +316,24 @@ inline void EnsureColorSpace(IDXGISwapChain* SwapChain, EDisplayColorGamut Displ
 		return;
 	}
 
+	{
+		UINT ColorSpaceSupport = 0;
+		HRESULT hr = swapChain3->CheckColorSpaceSupport(NewColorSpace, &ColorSpaceSupport);
+		FString ColorSpaceName = GetDXGIColorSpaceString(NewColorSpace);
+		if (SUCCEEDED(hr) && (ColorSpaceSupport & DXGI_SWAP_CHAIN_COLOR_SPACE_SUPPORT_FLAG_PRESENT))
+		{
+			swapChain3->SetColorSpace1(NewColorSpace);
+			UE_LOG(LogD3D11RHI, Verbose, TEXT("Setting color space on swap chain: %s"), *ColorSpaceName);
+		}
+		else
+		{
+			UE_LOG(LogD3D11RHI, Error, TEXT("Warning: unable to set color space %s to the swapchain: verify EDisplayOutputFormat / swapchain format"), *ColorSpaceName);
+		}
+	}
+}
+
+inline void EnsureColorSpace(IDXGISwapChain* SwapChain, EDisplayColorGamut DisplayGamut, EDisplayOutputFormat OutputDevice, EPixelFormat PixelFormat)
+{
 	DXGI_COLOR_SPACE_TYPE NewColorSpace = DXGI_COLOR_SPACE_RGB_FULL_G22_NONE_P709;	// sRGB;
 	const bool bPrimaries2020 = (DisplayGamut == EDisplayColorGamut::Rec2020_D65);
 
@@ -345,20 +363,7 @@ inline void EnsureColorSpace(IDXGISwapChain* SwapChain, EDisplayColorGamut Displ
 		break;
 	}
 
-	{
-		UINT ColorSpaceSupport = 0;
-		HRESULT hr = swapChain3->CheckColorSpaceSupport(NewColorSpace, &ColorSpaceSupport);
-		FString ColorSpaceName = GetDXGIColorSpaceString(NewColorSpace);
-		if (SUCCEEDED(hr) && (ColorSpaceSupport & DXGI_SWAP_CHAIN_COLOR_SPACE_SUPPORT_FLAG_PRESENT))
-		{
-			swapChain3->SetColorSpace1(NewColorSpace);
-			UE_LOG(LogD3D11RHI, Verbose, TEXT("Setting color space on swap chain: %s"), *ColorSpaceName);
-		}
-		else
-		{
-			UE_LOG(LogD3D11RHI, Error, TEXT("Warning: unable to set color space %s to the swapchain: verify EDisplayOutputFormat / swapchain format"), *ColorSpaceName);
-		}
-	}
+	EnsureColorSpace(SwapChain, NewColorSpace);
 }
 
 // When a window has moved or resized we need to check whether it is on a HDR monitor or not. Set the correct color space of the monitor
@@ -379,6 +384,13 @@ void FD3D11Viewport::ShutdownHDR()
 {
 	// Make sure to set the appropriate color space even if GRHISupportsHDROutput is false because we 
 	// might have toggled HDR on and off in the windows settings
+	if (PixelFormat == PF_FloatRGBA)
+	{
+		// Per DXGI spec, float-valued swapchains must contain linear data.
+		// The swapchain color space must therefore be G10 instead of G22.
+		EnsureColorSpace(SwapChain, DXGI_COLOR_SPACE_RGB_FULL_G10_NONE_P709);
+	}
+	else
 	{
 		// Ensure we have the correct color space set.
 		EnsureColorSpace(SwapChain, EDisplayColorGamut::sRGB_D65, EDisplayOutputFormat::SDR_sRGB, PixelFormat);
