[Audio] WASAPI: Fall back to null device when hardware device open fails

When a WASAPI audio device is found but cannot be opened (e.g. exclusive
mode lock by another application, COM error, driver issue), fall back to
null device rendering instead of failing the entire audio subsystem.

Previously, DeviceManager->InitializeHardware() or OpenAudioStream()
failures in FAudioMixerWasapi::OpenAudioStream() returned false, which
propagated up through FMixerDevice::InitializeHardware() and
FAudioDevice::Init(), leaving FAudioDevice null. This caused crashes in
every editor path that touches audio (MetaSound editor, WaveformEditor,
MovieSceneAudioSystem, etc.).

The null device mechanism (StartRunningNullDevice) already existed for
runtime device loss and the "no device found" case. This change extends
it to cover "device found but can't be opened" & the audio engine ticks
normally in software (MetaSounds process, mixers run) but outputs silence.

[RN] Fix for exclusive mode audio device initialization. Audio initializes with null renderer vs crashing.
#Synced-CL 51786389

diff --git a/Engine/Source/Runtime/Windows/AudioMixerWasapi/Private/AudioMixerWasapi.cpp b/Engine/Source/Runtime/Windows/AudioMixerWasapi/Private/AudioMixerWasapi.cpp
--- a/Engine/Source/Runtime/Windows/AudioMixerWasapi/Private/AudioMixerWasapi.cpp
+++ b/Engine/Source/Runtime/Windows/AudioMixerWasapi/Private/AudioMixerWasapi.cpp
@@ -213,6 +213,13 @@ namespace Audio
 		bool bIsAggregateDevice = false;
 		TArray<FWasapiRenderStreamParams> StreamParams;
 
+		// Track whether we successfully opened a hardware device. If any step in the
+		// hardware init chain fails, we fall through to null device setup rather than
+		// failing the entire audio subsystem. A null device ticks the audio engine
+		// (MetaSounds, mixers, etc.) in software but outputs silence and strictly better
+		// than crashing every editor path that touches audio.
+		bool bUsingHardwareDevice = false;
+
 		if (InitStreamParams(OpenStreamParams.OutputDeviceIndex, OpenStreamParams.NumFrames, OpenStreamParams.NumBuffers, OpenStreamParams.SampleRate, StreamParams))
 		{
 			// Adopt the first device info. In the case of an aggregate device, all of the sub-devices will
@@ -236,31 +243,58 @@ namespace Audio
 			// The ReadNextBufferCallback life cycle is tied to this object. It is ultimately bound to a delegate
 			// in the render stream object which will be unbound in TeardownHardware, prior to 'this' being deallocated.
 			TFunction<void()> ReadNextBufferCallback = [this]() { ReadNextBuffer(); };
-			if (!ensure(DeviceManager.IsValid() && DeviceManager->InitializeHardware(StreamParams, ReadNextBufferCallback)))
-			{
-				AUDIO_PLATFORM_LOG_ONCE(TEXT("FAudioMixerWasapi::OpenAudioStream DeviceManager->InitializeHardware() failed"), Warning);
-				return false;
+			if (!DeviceManager.IsValid() || !DeviceManager->InitializeHardware(StreamParams, ReadNextBufferCallback))
+			{
+				// Device was found but hardware init failed. e.g. another application has
+				// exclusive WASAPI access, COM error, driver issue, etc. Teardown any
+				// partial COM state, then reset the device manager so StartAudioStream()
+				// will pick up the null device path.
+				UE_LOG(LogAudioMixer, Warning, TEXT("FAudioMixerWasapi::OpenAudioStream: Audio device '%s' was found but "
+					"InitializeHardware() failed. Falling back to null device (no audio output). "
+					"Check if another application has exclusive access to this device."), *CurrentDeviceName);
+				if (DeviceManager.IsValid())
+				{
+					DeviceManager->TeardownHardware();
+				}
+				DeviceManager.Reset();
 			}
-
-			// Assign the total number of direct out channels based on the device manager. For WasapiDefaultDeviceMgr this will be 0 
-			// and for WasapiAggregateDeviceMgr this will be the total of all the channels less the main outs (first 8 channels).
-			AudioStreamInfo.DeviceInfo.NumDirectOutChannels = DeviceManager->GetNumDirectOutChannels();
-
-			if (!ensure(DeviceManager.IsValid() && DeviceManager->OpenAudioStream(StreamParams)))
+			else
 			{
-				UE_LOG(LogAudioMixer, Error, TEXT("FAudioMixerWasapi::OpenAudioStream DeviceManager->OpenAudioStream() failed"));
-				return false;
-			}
+				// Assign the total number of direct out channels based on the device manager. For WasapiDefaultDeviceMgr this will be 0
+				// and for WasapiAggregateDeviceMgr this will be the total of all the channels less the main outs (first 8 channels).
+				AudioStreamInfo.DeviceInfo.NumDirectOutChannels = DeviceManager->GetNumDirectOutChannels();
 
-			// Store the device ID here in case it is removed. We can switch back if the device comes back.
-			if (Params.bRestoreIfRemoved)
-			{
-				SetOriginalAudioDeviceId(AudioStreamInfo.DeviceInfo.DeviceId);
+				if (!DeviceManager.IsValid() || !DeviceManager->OpenAudioStream(StreamParams))
+				{
+					// Hardware initialized but the audio stream couldn't be opened. e.g.
+					// WASAPI session creation failed, buffer allocation error, format mismatch.
+					// Teardown COM resources before discarding the device manager.
+					UE_LOG(LogAudioMixer, Warning, TEXT("FAudioMixerWasapi::OpenAudioStream: Audio device '%s' initialized but "
+						"OpenAudioStream() failed. Falling back to null device (no audio output)."), *CurrentDeviceName);
+					DeviceManager->CloseAudioStream();
+					DeviceManager->TeardownHardware();
+					DeviceManager.Reset();
+				}
+				else
+				{
+					// Full hardware path succeeded.
+					bUsingHardwareDevice = true;
+
+					// Store the device ID here in case it is removed. We can switch back if the device comes back.
+					if (Params.bRestoreIfRemoved)
+					{
+						SetOriginalAudioDeviceId(AudioStreamInfo.DeviceInfo.DeviceId);
+					}
+				}
 			}
 		}
-		else
+
+		if (!bUsingHardwareDevice)
 		{
-			// Setup for running null device.
+			// No hardware device available — either no device was found, or the device
+			// was found but couldn't be opened. Set up a stereo null device so the audio
+			// engine can still run in software. StartAudioStream() will detect the null
+			// DeviceManager and call StartRunningNullDevice().
 			AudioStreamInfo.NumOutputFrames = OpenStreamParams.NumFrames;
 			AudioStreamInfo.DeviceInfo.OutputChannelArray = { EAudioMixerChannel::FrontLeft, EAudioMixerChannel::FrontRight };
 			AudioStreamInfo.DeviceInfo.NumChannels = 2;
