Persist D3D11 devices after creation until we've chosen the adapter

#Synced-CL 52232959

diff --git a/Engine/Source/Runtime/D3D12RHI/Private/Windows/WindowsD3D12Device.cpp b/Engine/Source/Runtime/D3D12RHI/Private/Windows/WindowsD3D12Device.cpp
--- a/Engine/Source/Runtime/D3D12RHI/Private/Windows/WindowsD3D12Device.cpp
+++ b/Engine/Source/Runtime/D3D12RHI/Private/Windows/WindowsD3D12Device.cpp
@@ -24,10 +24,6 @@
 #include "Misc/EngineVersion.h"
 #include "Windows/WindowsD3D.h"
 
-#include "Windows/AllowWindowsPlatformTypes.h"
-#include <delayimp.h>
-#include "Windows/HideWindowsPlatformTypes.h"
-
 IMPLEMENT_MODULE(FD3D12DynamicRHIModule, D3D12RHI);
 
 extern bool D3D12RHI_ShouldCreateWithWarp();
@@ -130,41 +126,6 @@ using namespace D3D12RHI;
 
 static bool bIsQuadBufferStereoEnabled = false;
 
-/** This function is used as a SEH filter to catch only delay load exceptions. */
-static bool IsDelayLoadException(PEXCEPTION_POINTERS ExceptionPointers)
-{
-	switch (ExceptionPointers->ExceptionRecord->ExceptionCode)
-	{
-	case VcppException(ERROR_SEVERITY_ERROR, ERROR_MOD_NOT_FOUND):
-	case VcppException(ERROR_SEVERITY_ERROR, ERROR_PROC_NOT_FOUND):
-		return EXCEPTION_EXECUTE_HANDLER;
-	default:
-		return EXCEPTION_CONTINUE_SEARCH;
-	}
-}
-
-/**
- * Since CreateDXGIFactory is a delay loaded import from the DXGI DLL, if the user
- * doesn't have Vista/DX10, calling CreateDXGIFactory will throw an exception.
- * We use SEH to detect that case and fail gracefully.
- */
-static void SafeCreateDXGIFactory(IDXGIFactory4** DXGIFactory)
-{
-#if !defined(D3D12_CUSTOM_VIEWPORT_CONSTRUCTOR) || !D3D12_CUSTOM_VIEWPORT_CONSTRUCTOR
-	__try
-	{
-		bIsQuadBufferStereoEnabled = FParse::Param(FCommandLine::Get(), TEXT("quad_buffer_stereo"));
-
-		verify(SUCCEEDED(CreateDXGIFactory1(IID_PPV_ARGS(DXGIFactory))));
-	}
-	__except (IsDelayLoadException(GetExceptionInformation()))
-	{
-		// We suppress warning C6322: Empty _except block. Appropriate checks are made upon returning. 
-		CA_SUPPRESS(6322);
-	}
-#endif	//!D3D12_CUSTOM_VIEWPORT_CONSTRUCTOR
-}
-
 static D3D_FEATURE_LEVEL FindHighestFeatureLevel(ID3D12Device* Device, D3D_FEATURE_LEVEL MinFeatureLevel)
 {
 	const D3D_FEATURE_LEVEL FeatureLevels[] =
@@ -531,18 +492,20 @@ static bool SafeTestD3D12CreateDevice(IDXGIAdapter* Adapter, LUID AdapterLuid, F
 {
 	constexpr D3D_FEATURE_LEVEL MinFeatureLevel = D3D_FEATURE_LEVEL_11_0;
 
-	uint64 Luid = BitCast<uint64>(AdapterLuid);
-	TRefCountPtr<ID3D12Device> Device = FWindowsD3D::FindCachedD3D12AdapterDevice(Luid);
-	if (!Device.IsValid())
+	const uint64 Luid = BitCast<uint64>(AdapterLuid);
+	ID3D12Device* Device = FWindowsD3D::FindCachedD3D12AdapterDevice(Luid);
+	if (!Device)
 	{
-		const HRESULT D3D12CreateDeviceResult = D3D12CreateDevice(Adapter, MinFeatureLevel, IID_PPV_ARGS(Device.GetInitReference()));
-		if (!SUCCEEDED(D3D12CreateDeviceResult))
+		TRefCountPtr<ID3D12Device> DevicePtr;
+		const HRESULT D3D12CreateDeviceResult = D3D12CreateDevice(Adapter, MinFeatureLevel, IID_PPV_ARGS(DevicePtr.GetInitReference()));
+		if (FAILED(D3D12CreateDeviceResult))
 		{
 			UE_LOG(LogD3D12RHI, Log, TEXT("D3D12CreateDevice failed with code 0x%08X"), static_cast<int32>(D3D12CreateDeviceResult));
 			return false;
 		}
 
-		FWindowsD3D::CacheD3D12AdapterDevice(Luid, Device);
+		Device = DevicePtr.GetReference();
+		FWindowsD3D::CacheD3D12AdapterDevice(Luid, MoveTemp(DevicePtr));
 	}
 
 	OutInfo.MaxFeatureLevel = FindHighestFeatureLevel(Device, MinFeatureLevel);
@@ -865,9 +828,12 @@ void FD3D12DynamicRHIModule::FindAdapter()
 	}
 #endif
 
-	// Try to create the DXGIFactory.  This will fail if we're not running Vista.
+	// Try to create the DXGIFactory.
 	TRefCountPtr<IDXGIFactory4> DXGIFactory4;
-	SafeCreateDXGIFactory(DXGIFactory4.GetInitReference());
+#if !defined(D3D12_CUSTOM_VIEWPORT_CONSTRUCTOR) || !D3D12_CUSTOM_VIEWPORT_CONSTRUCTOR
+	bIsQuadBufferStereoEnabled = FParse::Param(FCommandLine::Get(), TEXT("quad_buffer_stereo"));
+	verify(SUCCEEDED(CreateDXGIFactory1(IID_PPV_ARGS(DXGIFactory4.GetInitReference()))));
+#endif
 	if (!DXGIFactory4)
 	{
 		return;
diff --git a/Engine/Source/Runtime/RHI/Private/Windows/WindowsDynamicRHI.cpp b/Engine/Source/Runtime/RHI/Private/Windows/WindowsDynamicRHI.cpp
--- a/Engine/Source/Runtime/RHI/Private/Windows/WindowsDynamicRHI.cpp
+++ b/Engine/Source/Runtime/RHI/Private/Windows/WindowsDynamicRHI.cpp
@@ -1392,7 +1392,7 @@ FDynamicRHI* PlatformCreateDynamicRHI()
 	const TCHAR* LoadedRHIModuleName;
 	IDynamicRHIModule* DynamicRHIModule = LoadDynamicRHIModule(RequestedFeatureLevel, LoadedRHIModuleName);
 
-	FWindowsD3D::ReleaseCachedD3D12AdapterDevices();
+	FWindowsD3D::ReleaseCachedAdapterDevices();
 
 	if (DynamicRHIModule)
 	{
diff --git a/Engine/Source/Runtime/Windows/D3D11RHI/Private/Windows/WindowsD3D11Device.cpp b/Engine/Source/Runtime/Windows/D3D11RHI/Private/Windows/WindowsD3D11Device.cpp
--- a/Engine/Source/Runtime/Windows/D3D11RHI/Private/Windows/WindowsD3D11Device.cpp
+++ b/Engine/Source/Runtime/Windows/D3D11RHI/Private/Windows/WindowsD3D11Device.cpp
@@ -21,7 +21,6 @@
 #include "Windows/WindowsD3D.h"
 
 #include "Windows/AllowWindowsPlatformTypes.h"
-#include <delayimp.h>
 	#if WITH_NVAPI
 		#include "nvapi.h"
 		#include "nvShaderExtnEnums.h"
@@ -231,19 +230,6 @@ static LONG __stdcall D3D11VectoredExceptionHandler(EXCEPTION_POINTERS* InInfo)
 	return EXCEPTION_CONTINUE_SEARCH;
 }
 
-/** This function is used as a SEH filter to catch only delay load exceptions. */
-static bool IsDelayLoadException(PEXCEPTION_POINTERS ExceptionPointers)
-{
-	switch(ExceptionPointers->ExceptionRecord->ExceptionCode)
-	{
-	case VcppException(ERROR_SEVERITY_ERROR, ERROR_MOD_NOT_FOUND):
-	case VcppException(ERROR_SEVERITY_ERROR, ERROR_PROC_NOT_FOUND):
-		return EXCEPTION_EXECUTE_HANDLER;
-	default:
-		return EXCEPTION_CONTINUE_SEARCH;
-	}
-}
-
 static bool bIsQuadBufferStereoEnabled = false;
 
 static void SafeCreateDXGIFactory(IDXGIFactory1** DXGIFactory1, bool bWithDebug)
@@ -289,64 +275,70 @@ inline bool IsDeviceUMA(ID3D11Device* device)
 }
 
 /**
- * Attempts to create a D3D11 device for the adapter using at most MaxFeatureLevel.
+ * Attempts to create a D3D11 device for the adapter using at most MaxFeatureLevel or find an existing one.
  * If creation is successful, true is returned.
  */
-static bool SafeTestD3D11CreateDevice(IDXGIAdapter* Adapter, D3D_FEATURE_LEVEL MinFeatureLevel, D3D_FEATURE_LEVEL MaxFeatureLevel, FD3D11DeviceBasicInfo& OutInfo)
+static bool SafeTestD3D11CreateDevice(IDXGIAdapter* Adapter, LUID AdapterLuid, D3D_FEATURE_LEVEL MinFeatureLevel, D3D_FEATURE_LEVEL MaxFeatureLevel, FD3D11DeviceBasicInfo& OutInfo)
 {
-	ID3D11Device* D3DDevice = nullptr;
-	ID3D11DeviceContext* D3DDeviceContext = nullptr;
-	uint32 DeviceFlags = D3D11_CREATE_DEVICE_SINGLETHREADED;
-	// Use a debug device if specified on the command line.
-	if(GRHIGlobals.IsDebugLayerEnabled)
+	const uint64 Luid = BitCast<uint64>(AdapterLuid);
+	int32 CachedFeatureLevel;
+	ID3D11Device* D3DDevice = FWindowsD3D::FindCachedD3D11AdapterDevice(Luid, CachedFeatureLevel);
+	if (D3DDevice)
 	{
-		DeviceFlags |= D3D11_CREATE_DEVICE_DEBUG;
+		OutInfo.MaxFeatureLevel = static_cast<D3D_FEATURE_LEVEL>(CachedFeatureLevel);
 	}
+	else
+	{
+		uint32 DeviceFlags = D3D11_CREATE_DEVICE_SINGLETHREADED;
+		// Use a debug device if specified on the command line.
+		if(GRHIGlobals.IsDebugLayerEnabled)
+		{
+			DeviceFlags |= D3D11_CREATE_DEVICE_DEBUG;
+		}
 
-	// @MIXEDREALITY_CHANGE : BEGIN - Add BGRA flag for Windows Mixed Reality HMD's
-	DeviceFlags |= D3D11_CREATE_DEVICE_BGRA_SUPPORT;
-	// @MIXEDREALITY_CHANGE : END
+		// @MIXEDREALITY_CHANGE : BEGIN - Add BGRA flag for Windows Mixed Reality HMD's
+		DeviceFlags |= D3D11_CREATE_DEVICE_BGRA_SUPPORT;
+		// @MIXEDREALITY_CHANGE : END
 
-	D3D_FEATURE_LEVEL RequestedFeatureLevels[] =
-	{
-		D3D_FEATURE_LEVEL_11_1,
-		D3D_FEATURE_LEVEL_11_0,
-	};
+		D3D_FEATURE_LEVEL RequestedFeatureLevels[] =
+		{
+			D3D_FEATURE_LEVEL_11_1,
+			D3D_FEATURE_LEVEL_11_0,
+		};
 	
-	// Trim to allowed feature levels
-	int32 FirstAllowedFeatureLevel = 0;
-	int32 NumAllowedFeatureLevels = UE_ARRAY_COUNT(RequestedFeatureLevels);
-	int32 LastAllowedFeatureLevel = NumAllowedFeatureLevels - 1;
+		// Trim to allowed feature levels
+		int32 FirstAllowedFeatureLevel = 0;
+		int32 NumAllowedFeatureLevels = UE_ARRAY_COUNT(RequestedFeatureLevels);
+		int32 LastAllowedFeatureLevel = NumAllowedFeatureLevels - 1;
 	
-	while (FirstAllowedFeatureLevel < NumAllowedFeatureLevels)
-	{
-		if (RequestedFeatureLevels[FirstAllowedFeatureLevel] == MaxFeatureLevel)
+		while (FirstAllowedFeatureLevel < NumAllowedFeatureLevels)
 		{
-			break;
+			if (RequestedFeatureLevels[FirstAllowedFeatureLevel] == MaxFeatureLevel)
+			{
+				break;
+			}
+			FirstAllowedFeatureLevel++;
 		}
-		FirstAllowedFeatureLevel++;
-	}
 
-	while (LastAllowedFeatureLevel > 0)
-	{
-		if (RequestedFeatureLevels[LastAllowedFeatureLevel] >= MinFeatureLevel)
+		while (LastAllowedFeatureLevel > 0)
 		{
-			break;
+			if (RequestedFeatureLevels[LastAllowedFeatureLevel] >= MinFeatureLevel)
+			{
+				break;
+			}
+			LastAllowedFeatureLevel--;
 		}
-		LastAllowedFeatureLevel--;
-	}
 	
-	NumAllowedFeatureLevels = LastAllowedFeatureLevel - FirstAllowedFeatureLevel + 1;
-	if (MaxFeatureLevel < MinFeatureLevel || NumAllowedFeatureLevels <= 0)
-	{
-		return false;
-	}
+		NumAllowedFeatureLevels = LastAllowedFeatureLevel - FirstAllowedFeatureLevel + 1;
+		if (MaxFeatureLevel < MinFeatureLevel || NumAllowedFeatureLevels <= 0)
+		{
+			return false;
+		}
 
-	__try
-	{
 		// We don't want software renderer. Ideally we specify D3D_DRIVER_TYPE_HARDWARE on creation but
 		// when we specify an adapter we need to specify D3D_DRIVER_TYPE_UNKNOWN (otherwise the call fails).
 		// We cannot check the device type later (seems this is missing functionality in D3D).
+		TRefCountPtr<ID3D11Device> DevicePtr;
 		D3D_FEATURE_LEVEL DeviceFeatureLevel = (D3D_FEATURE_LEVEL)0;
 		HRESULT Result = D3D11CreateDevice(
 			Adapter,
@@ -356,36 +348,31 @@ static bool SafeTestD3D11CreateDevice(IDXGIAdapter* Adapter, D3D_FEATURE_LEVEL M
 			&RequestedFeatureLevels[FirstAllowedFeatureLevel],
 			NumAllowedFeatureLevels,
 			D3D11_SDK_VERSION,
-			&D3DDevice,
+			DevicePtr.GetInitReference(),
 			&DeviceFeatureLevel,
-			&D3DDeviceContext);
+			nullptr);
 
-		if (SUCCEEDED(Result))
+		if (FAILED(Result))
 		{
-			OutInfo.MaxFeatureLevel = DeviceFeatureLevel;
-			OutInfo.bUMA = IsDeviceUMA(D3DDevice);
+			// Log any reason for failure to create test device. Extra debug help.
+			VERIFYD3D11RESULT_NOEXIT(Result);
 
-			D3DDevice->Release();
-			D3DDeviceContext->Release();
-			return true;
-		}
-
-		// Log any reason for failure to create test device. Extra debug help.
-		VERIFYD3D11RESULT_NOEXIT(Result);
+			// Fatal error on 0x887A002D
+			if (DXGI_ERROR_SDK_COMPONENT_MISSING == Result)
+			{
+				UE_LOG(LogD3D11RHI, Fatal, TEXT("-d3ddebug was used but optional Graphics Tools were not found. Install them through the Manage Optional Features in windows. See: https://docs.microsoft.com/en-us/windows/uwp/gaming/use-the-directx-runtime-and-visual-studio-graphics-diagnostic-features"));
+			}
 
-		// Fatal error on 0x887A002D
-		if (DXGI_ERROR_SDK_COMPONENT_MISSING == Result)
-		{
-			UE_LOG(LogD3D11RHI, Fatal, TEXT("-d3ddebug was used but optional Graphics Tools were not found. Install them through the Manage Optional Features in windows. See: https://docs.microsoft.com/en-us/windows/uwp/gaming/use-the-directx-runtime-and-visual-studio-graphics-diagnostic-features"));
+			return false;
 		}
-	}
-	__except(IsDelayLoadException(GetExceptionInformation()))
-	{
-		// We suppress warning C6322: Empty _except block. Appropriate checks are made upon returning. 
-		CA_SUPPRESS(6322);
+
+		D3DDevice = DevicePtr.GetReference();
+		FWindowsD3D::CacheD3D11AdapterDevice(Luid, MoveTemp(DevicePtr), DeviceFeatureLevel);
+		OutInfo.MaxFeatureLevel = DeviceFeatureLevel;
 	}
 
-	return false;
+	OutInfo.bUMA = IsDeviceUMA(D3DDevice);
+	return true;
 }
 
 struct DisplayChromacities
@@ -736,7 +723,7 @@ void FD3D11DynamicRHIModule::FindAdapter()
 				LogDXGIAdapterDesc(AdapterDesc);
 
 				FD3D11DeviceBasicInfo DeviceInfo;
-				if(SafeTestD3D11CreateDevice(TempAdapter, MinAllowedFeatureLevel, MaxAllowedFeatureLevel, DeviceInfo))
+				if(SafeTestD3D11CreateDevice(TempAdapter, AdapterDesc.AdapterLuid, MinAllowedFeatureLevel, MaxAllowedFeatureLevel, DeviceInfo))
 				{
 					// Log some information about the available D3D11 adapters.
 
@@ -1376,7 +1363,6 @@ void FD3D11DynamicRHI::InitD3DDevice()
 		}
 #endif //AMD_AGS_API
 
-
 		if (!bDeviceCreated)
 		{
 			UE_LOG(LogD3D11RHI, Log, TEXT("Creating D3DDevice using adapter:"));
diff --git a/Engine/Source/Runtime/Windows/WindowsD3D/Internal/Windows/WindowsD3D.h b/Engine/Source/Runtime/Windows/WindowsD3D/Internal/Windows/WindowsD3D.h
--- a/Engine/Source/Runtime/Windows/WindowsD3D/Internal/Windows/WindowsD3D.h
+++ b/Engine/Source/Runtime/Windows/WindowsD3D/Internal/Windows/WindowsD3D.h
@@ -6,6 +6,7 @@
 #include "HAL/Platform.h"
 #include "Templates/RefCounting.h"
 
+struct ID3D11Device;
 struct ID3D12Device;
 
 struct FWindowsGPUInfo
@@ -24,7 +25,9 @@ public:
 	static WINDOWSD3D_API uint64 ChooseD3D11Adapter(uint64 DesiredAdapterLuid);
 	static WINDOWSD3D_API uint64 ChooseD3D12Adapter(uint64 DesiredAdapterLuid);
 
-	static WINDOWSD3D_API TRefCountPtr<ID3D12Device> FindCachedD3D12AdapterDevice(uint64 AdapterLuid);
-	static WINDOWSD3D_API void CacheD3D12AdapterDevice(uint64 AdapterLuid, TRefCountPtr<ID3D12Device> Device);
-	static WINDOWSD3D_API void ReleaseCachedD3D12AdapterDevices();
+	static WINDOWSD3D_API ID3D11Device* FindCachedD3D11AdapterDevice(uint64 AdapterLuid, int32& FeatureLevel);
+	static WINDOWSD3D_API void CacheD3D11AdapterDevice(uint64 AdapterLuid, TRefCountPtr<ID3D11Device>&& Device, int32 FeatureLevel);
+	static WINDOWSD3D_API ID3D12Device* FindCachedD3D12AdapterDevice(uint64 AdapterLuid);
+	static WINDOWSD3D_API void CacheD3D12AdapterDevice(uint64 AdapterLuid, TRefCountPtr<ID3D12Device>&& Device);
+	static WINDOWSD3D_API void ReleaseCachedAdapterDevices();
 };
diff --git a/Engine/Source/Runtime/Windows/WindowsD3D/Private/WindowsD3D.cpp b/Engine/Source/Runtime/Windows/WindowsD3D/Private/WindowsD3D.cpp
--- a/Engine/Source/Runtime/Windows/WindowsD3D/Private/WindowsD3D.cpp
+++ b/Engine/Source/Runtime/Windows/WindowsD3D/Private/WindowsD3D.cpp
@@ -10,14 +10,9 @@
 #include "Windows/D3D11ThirdParty.h"
 #include "Windows/WindowsD3D12ThirdParty.h"
 
-THIRD_PARTY_INCLUDES_START
-#include "Windows/AllowWindowsPlatformTypes.h"
-#include <delayimp.h>
-#include "Windows/HideWindowsPlatformTypes.h"
-THIRD_PARTY_INCLUDES_END
-
 DEFINE_LOG_CATEGORY_STATIC(LogD3D, Log, All);
 
+static TMap<uint64, TPair<TRefCountPtr<ID3D11Device>, int32>> D3D11AdapterDevices;
 static TMap<uint64, TRefCountPtr<ID3D12Device>> D3D12AdapterDevices;
 
 enum class ED3DGpuVendorId : uint32
@@ -187,19 +182,6 @@ static HRESULT EnumAdapters(int32 AdapterIndex, DXGI_GPU_PREFERENCE GpuPreferenc
 	}
 }
 
-/** This function is used as a SEH filter to catch only delay load exceptions. */
-static bool IsDelayLoadException(PEXCEPTION_POINTERS ExceptionPointers)
-{
-	switch (ExceptionPointers->ExceptionRecord->ExceptionCode)
-	{
-	case VcppException(ERROR_SEVERITY_ERROR, ERROR_MOD_NOT_FOUND):
-	case VcppException(ERROR_SEVERITY_ERROR, ERROR_PROC_NOT_FOUND):
-		return EXCEPTION_EXECUTE_HANDLER;
-	default:
-		return EXCEPTION_CONTINUE_SEARCH;
-	}
-}
-
 /**
  * Attempts to create a D3D12 device for the adapter or find an existing one.
  * If creation is successful, true is returned.
@@ -210,18 +192,20 @@ static bool SafeTestD3D12CreateDevice(IDXGIAdapter* Adapter, LUID AdapterLuid, b
 
 	constexpr D3D_FEATURE_LEVEL MinFeatureLevel = D3D_FEATURE_LEVEL_11_0;
 
-	uint64 Luid = BitCast<uint64>(AdapterLuid);
-	TRefCountPtr<ID3D12Device> Device = FWindowsD3D::FindCachedD3D12AdapterDevice(Luid);
+	const uint64 Luid = BitCast<uint64>(AdapterLuid);
+	ID3D12Device* Device = FWindowsD3D::FindCachedD3D12AdapterDevice(Luid);
 	if (!Device)
 	{
-		const HRESULT D3D12CreateDeviceResult = D3D12CreateDevice(Adapter, MinFeatureLevel, IID_PPV_ARGS(Device.GetInitReference()));
-		if (!SUCCEEDED(D3D12CreateDeviceResult))
+		TRefCountPtr<ID3D12Device> DevicePtr;
+		const HRESULT D3D12CreateDeviceResult = D3D12CreateDevice(Adapter, MinFeatureLevel, IID_PPV_ARGS(DevicePtr.GetInitReference()));
+		if (FAILED(D3D12CreateDeviceResult))
 		{
 			UE_LOG(LogD3D, Log, TEXT("D3D12CreateDevice failed with code 0x%08X"), static_cast<int32>(D3D12CreateDeviceResult));
 			return false;
 		}
 
-		FWindowsD3D::CacheD3D12AdapterDevice(Luid, Device);
+		Device = DevicePtr.GetReference();
+		FWindowsD3D::CacheD3D12AdapterDevice(Luid, MoveTemp(DevicePtr));
 	}
 
 	D3D12_FEATURE_DATA_ARCHITECTURE data{};
@@ -230,13 +214,11 @@ static bool SafeTestD3D12CreateDevice(IDXGIAdapter* Adapter, LUID AdapterLuid, b
 }
 
 /**
- * Attempts to create a D3D11 device for the adapter using at most MaxFeatureLevel.
+ * Attempts to create a D3D11 device for the adapter using at most MaxFeatureLevel or find an existing one.
  * If creation is successful, true is returned.
  */
-static bool SafeTestD3D11CreateDevice(IDXGIAdapter* Adapter, D3D_FEATURE_LEVEL MinFeatureLevel, D3D_FEATURE_LEVEL MaxFeatureLevel, bool& bOutUMA)
+static bool SafeTestD3D11CreateDevice(IDXGIAdapter* Adapter, LUID AdapterLuid, D3D_FEATURE_LEVEL MinFeatureLevel, D3D_FEATURE_LEVEL MaxFeatureLevel, bool& bOutUMA)
 {
-	ID3D11Device* D3DDevice = nullptr;
-	ID3D11DeviceContext* D3DDeviceContext = nullptr;
 	uint32 DeviceFlags = D3D11_CREATE_DEVICE_SINGLETHREADED;
 	// Use a debug device if specified on the command line.
 	if (FParse::Param(FCommandLine::Get(), TEXT("d3ddebug")) ||
@@ -288,11 +270,15 @@ static bool SafeTestD3D11CreateDevice(IDXGIAdapter* Adapter, D3D_FEATURE_LEVEL M
 		return false;
 	}
 
-	__try
+	const uint64 Luid = BitCast<uint64>(AdapterLuid);
+	int32 Unused;
+	ID3D11Device* D3DDevice = FWindowsD3D::FindCachedD3D11AdapterDevice(Luid, Unused);
+	if (!D3DDevice)
 	{
 		// We don't want software renderer. Ideally we specify D3D_DRIVER_TYPE_HARDWARE on creation but
 		// when we specify an adapter we need to specify D3D_DRIVER_TYPE_UNKNOWN (otherwise the call fails).
 		// We cannot check the device type later (seems this is missing functionality in D3D).
+		TRefCountPtr<ID3D11Device> DevicePtr;
 		D3D_FEATURE_LEVEL DeviceFeatureLevel = (D3D_FEATURE_LEVEL)0;
 		HRESULT Result = D3D11CreateDevice(
 			Adapter,
@@ -302,36 +288,32 @@ static bool SafeTestD3D11CreateDevice(IDXGIAdapter* Adapter, D3D_FEATURE_LEVEL M
 			&RequestedFeatureLevels[FirstAllowedFeatureLevel],
 			NumAllowedFeatureLevels,
 			D3D11_SDK_VERSION,
-			&D3DDevice,
+			DevicePtr.GetInitReference(),
 			&DeviceFeatureLevel,
-			&D3DDeviceContext);
+			nullptr);
 
-		if (SUCCEEDED(Result))
+		if (FAILED(Result))
 		{
-			D3D11_FEATURE_DATA_D3D11_OPTIONS2 FeatureData = {};
-			bOutUMA = SUCCEEDED(D3DDevice->CheckFeatureSupport(D3D11_FEATURE_D3D11_OPTIONS2, &FeatureData, sizeof(FeatureData))) && FeatureData.UnifiedMemoryArchitecture;
+			// Log any reason for failure to create test device. Extra debug help.
+			UE_LOG(LogD3D, Log, TEXT("D3D11CreateDevice failed with code 0x%08X"), static_cast<uint32>(Result));
 
-			D3DDevice->Release();
-			D3DDeviceContext->Release();
-			return true;
-		}
-
-		// Log any reason for failure to create test device. Extra debug help.
-		UE_LOG(LogD3D, Log, TEXT("D3D11CreateDevice failed with code 0x%08X"), static_cast<uint32>(Result));
+			// Fatal error on 0x887A002D
+			if (DXGI_ERROR_SDK_COMPONENT_MISSING == Result)
+			{
+				UE_LOG(LogD3D, Fatal, TEXT("-d3ddebug was used but optional Graphics Tools were not found. Install them through the Manage Optional Features in windows. See: https://docs.microsoft.com/en-us/windows/uwp/gaming/use-the-directx-runtime-and-visual-studio-graphics-diagnostic-features"));
+			}
 
-		// Fatal error on 0x887A002D
-		if (DXGI_ERROR_SDK_COMPONENT_MISSING == Result)
-		{
-			UE_LOG(LogD3D, Fatal, TEXT("-d3ddebug was used but optional Graphics Tools were not found. Install them through the Manage Optional Features in windows. See: https://docs.microsoft.com/en-us/windows/uwp/gaming/use-the-directx-runtime-and-visual-studio-graphics-diagnostic-features"));
+			return false;
 		}
-	}
-	__except (IsDelayLoadException(GetExceptionInformation()))
-	{
-		// We suppress warning C6322: Empty _except block. Appropriate checks are made upon returning. 
-		CA_SUPPRESS(6322);
+
+		D3DDevice = DevicePtr.GetReference();
+		FWindowsD3D::CacheD3D11AdapterDevice(Luid, MoveTemp(DevicePtr), DeviceFeatureLevel);
 	}
 
-	return false;
+	D3D11_FEATURE_DATA_D3D11_OPTIONS2 FeatureData = {};
+	bOutUMA = SUCCEEDED(D3DDevice->CheckFeatureSupport(D3D11_FEATURE_D3D11_OPTIONS2, &FeatureData, sizeof(FeatureData))) && FeatureData.UnifiedMemoryArchitecture;
+
+	return true;
 }
 
 static uint64 ConvertLuidToUint64(LUID Luid)
@@ -401,7 +383,7 @@ uint64 FWindowsD3D::ChooseD3D11Adapter(uint64 DesiredAdapterLuid)
 		{
 			DXGI_ADAPTER_DESC AdapterDesc{};
 			bool bUMA;
-			if (SUCCEEDED(TempAdapter->GetDesc(&AdapterDesc)) && SafeTestD3D11CreateDevice(TempAdapter, MinAllowedFeatureLevel, MaxAllowedFeatureLevel, bUMA))
+			if (SUCCEEDED(TempAdapter->GetDesc(&AdapterDesc)) && SafeTestD3D11CreateDevice(TempAdapter, AdapterDesc.AdapterLuid, MinAllowedFeatureLevel, MaxAllowedFeatureLevel, bUMA))
 			{
 				const bool bIsWARP = ConvertToGpuVendorId(AdapterDesc.VendorId) == ED3DGpuVendorId::Microsoft;
 
@@ -647,23 +629,38 @@ uint64 FWindowsD3D::ChooseD3D12Adapter(uint64 DesiredAdapterLuid)
 	return Luid;
 }
 
-TRefCountPtr<ID3D12Device> FWindowsD3D::FindCachedD3D12AdapterDevice(uint64 AdapterLuid)
+ID3D11Device* FWindowsD3D::FindCachedD3D11AdapterDevice(uint64 AdapterLuid, int32& OutFeatureLevel)
 {
-	TRefCountPtr<ID3D12Device>* DevicePtr = D3D12AdapterDevices.Find(AdapterLuid);
-	if (DevicePtr)
+	ID3D11Device* Device = nullptr;
+
+	const TPair<TRefCountPtr<ID3D11Device>, int32>* const PairPtr = D3D11AdapterDevices.Find(AdapterLuid);
+	if (PairPtr)
 	{
-		return *DevicePtr;
+		OutFeatureLevel = PairPtr->Value;
+		Device = PairPtr->Key.GetReference();
 	}
 
-	return nullptr;
+	return Device;
+}
+
+void FWindowsD3D::CacheD3D11AdapterDevice(uint64 AdapterLuid, TRefCountPtr<ID3D11Device>&& Device, int32 FeatureLevel)
+{
+	D3D11AdapterDevices.Emplace(AdapterLuid, TPair<TRefCountPtr<ID3D11Device>, int32>(MoveTemp(Device), FeatureLevel));
+}
+
+ID3D12Device* FWindowsD3D::FindCachedD3D12AdapterDevice(uint64 AdapterLuid)
+{
+	TRefCountPtr<ID3D12Device>* DevicePtr = D3D12AdapterDevices.Find(AdapterLuid);
+	return DevicePtr ? DevicePtr->GetReference() : nullptr;
 }
 
-void FWindowsD3D::CacheD3D12AdapterDevice(uint64 AdapterLuid, TRefCountPtr<ID3D12Device> Device)
+void FWindowsD3D::CacheD3D12AdapterDevice(uint64 AdapterLuid, TRefCountPtr<ID3D12Device>&& Device)
 {
-	D3D12AdapterDevices.Emplace(MoveTemp(AdapterLuid), MoveTemp(Device));
+	D3D12AdapterDevices.Emplace(AdapterLuid, MoveTemp(Device));
 }
 
-void FWindowsD3D::ReleaseCachedD3D12AdapterDevices()
+void FWindowsD3D::ReleaseCachedAdapterDevices()
 {
+	D3D11AdapterDevices.Reset();
 	D3D12AdapterDevices.Reset();
 }
