Disable D3D device caching and persistence due to vendor-specific issues.

#rnx
#tests locally and by QA on affected machine
#Synced-CL 52411106

diff --git a/Engine/Source/Runtime/D3D12RHI/Private/Windows/WindowsD3D12Device.cpp b/Engine/Source/Runtime/D3D12RHI/Private/Windows/WindowsD3D12Device.cpp
--- a/Engine/Source/Runtime/D3D12RHI/Private/Windows/WindowsD3D12Device.cpp
+++ b/Engine/Source/Runtime/D3D12RHI/Private/Windows/WindowsD3D12Device.cpp
@@ -494,9 +494,9 @@ static bool SafeTestD3D12CreateDevice(IDXGIAdapter* Adapter, LUID AdapterLuid, F
 
 	const uint64 Luid = BitCast<uint64>(AdapterLuid);
 	ID3D12Device* Device = FWindowsD3D::FindCachedD3D12AdapterDevice(Luid);
+	TRefCountPtr<ID3D12Device> DevicePtr;
 	if (!Device)
 	{
-		TRefCountPtr<ID3D12Device> DevicePtr;
 		const HRESULT D3D12CreateDeviceResult = D3D12CreateDevice(Adapter, MinFeatureLevel, IID_PPV_ARGS(DevicePtr.GetInitReference()));
 		if (FAILED(D3D12CreateDeviceResult))
 		{
@@ -505,7 +505,10 @@ static bool SafeTestD3D12CreateDevice(IDXGIAdapter* Adapter, LUID AdapterLuid, F
 		}
 
 		Device = DevicePtr.GetReference();
-		FWindowsD3D::CacheD3D12AdapterDevice(Luid, MoveTemp(DevicePtr));
+		if (FWindowsD3D::ShouldCacheAdapterDevices())
+		{
+			FWindowsD3D::CacheD3D12AdapterDevice(Luid, MoveTemp(DevicePtr));
+		}
 	}
 
 	OutInfo.MaxFeatureLevel = FindHighestFeatureLevel(Device, MinFeatureLevel);
diff --git a/Engine/Source/Runtime/Windows/D3D11RHI/Private/Windows/WindowsD3D11Device.cpp b/Engine/Source/Runtime/Windows/D3D11RHI/Private/Windows/WindowsD3D11Device.cpp
--- a/Engine/Source/Runtime/Windows/D3D11RHI/Private/Windows/WindowsD3D11Device.cpp
+++ b/Engine/Source/Runtime/Windows/D3D11RHI/Private/Windows/WindowsD3D11Device.cpp
@@ -283,6 +283,7 @@ static bool SafeTestD3D11CreateDevice(IDXGIAdapter* Adapter, LUID AdapterLuid, D
 	const uint64 Luid = BitCast<uint64>(AdapterLuid);
 	int32 CachedFeatureLevel;
 	ID3D11Device* D3DDevice = FWindowsD3D::FindCachedD3D11AdapterDevice(Luid, CachedFeatureLevel);
+	TRefCountPtr<ID3D11Device> DevicePtr;
 	if (D3DDevice)
 	{
 		OutInfo.MaxFeatureLevel = static_cast<D3D_FEATURE_LEVEL>(CachedFeatureLevel);
@@ -338,7 +339,6 @@ static bool SafeTestD3D11CreateDevice(IDXGIAdapter* Adapter, LUID AdapterLuid, D
 		// We don't want software renderer. Ideally we specify D3D_DRIVER_TYPE_HARDWARE on creation but
 		// when we specify an adapter we need to specify D3D_DRIVER_TYPE_UNKNOWN (otherwise the call fails).
 		// We cannot check the device type later (seems this is missing functionality in D3D).
-		TRefCountPtr<ID3D11Device> DevicePtr;
 		D3D_FEATURE_LEVEL DeviceFeatureLevel = (D3D_FEATURE_LEVEL)0;
 		HRESULT Result = D3D11CreateDevice(
 			Adapter,
@@ -367,7 +367,10 @@ static bool SafeTestD3D11CreateDevice(IDXGIAdapter* Adapter, LUID AdapterLuid, D
 		}
 
 		D3DDevice = DevicePtr.GetReference();
-		FWindowsD3D::CacheD3D11AdapterDevice(Luid, MoveTemp(DevicePtr), DeviceFeatureLevel);
+		if (FWindowsD3D::ShouldCacheAdapterDevices())
+		{
+			FWindowsD3D::CacheD3D11AdapterDevice(Luid, MoveTemp(DevicePtr), DeviceFeatureLevel);
+		}
 		OutInfo.MaxFeatureLevel = DeviceFeatureLevel;
 	}
 
diff --git a/Engine/Source/Runtime/Windows/WindowsD3D/Internal/Windows/WindowsD3D.h b/Engine/Source/Runtime/Windows/WindowsD3D/Internal/Windows/WindowsD3D.h
--- a/Engine/Source/Runtime/Windows/WindowsD3D/Internal/Windows/WindowsD3D.h
+++ b/Engine/Source/Runtime/Windows/WindowsD3D/Internal/Windows/WindowsD3D.h
@@ -25,6 +25,7 @@ public:
 	static WINDOWSD3D_API uint64 ChooseD3D11Adapter(uint64 DesiredAdapterLuid);
 	static WINDOWSD3D_API uint64 ChooseD3D12Adapter(uint64 DesiredAdapterLuid);
 
+	static WINDOWSD3D_API bool ShouldCacheAdapterDevices();
 	static WINDOWSD3D_API ID3D11Device* FindCachedD3D11AdapterDevice(uint64 AdapterLuid, int32& FeatureLevel);
 	static WINDOWSD3D_API void CacheD3D11AdapterDevice(uint64 AdapterLuid, TRefCountPtr<ID3D11Device>&& Device, int32 FeatureLevel);
 	static WINDOWSD3D_API ID3D12Device* FindCachedD3D12AdapterDevice(uint64 AdapterLuid);
diff --git a/Engine/Source/Runtime/Windows/WindowsD3D/Private/WindowsD3D.cpp b/Engine/Source/Runtime/Windows/WindowsD3D/Private/WindowsD3D.cpp
--- a/Engine/Source/Runtime/Windows/WindowsD3D/Private/WindowsD3D.cpp
+++ b/Engine/Source/Runtime/Windows/WindowsD3D/Private/WindowsD3D.cpp
@@ -15,6 +15,10 @@ DEFINE_LOG_CATEGORY_STATIC(LogD3D, Log, All);
 static TMap<uint64, TPair<TRefCountPtr<ID3D11Device>, int32>> D3D11AdapterDevices;
 static TMap<uint64, TRefCountPtr<ID3D12Device>> D3D12AdapterDevices;
 
+// Some vendors are seeing initializing issues with this optimization enabled.
+// Default to off until those problems are resolved.
+static bool GCacheD3DDevices = false;
+
 enum class ED3DGpuVendorId : uint32
 {
 	Unknown = 0xffffffff,
@@ -194,9 +198,9 @@ static bool SafeTestD3D12CreateDevice(IDXGIAdapter* Adapter, LUID AdapterLuid, b
 
 	const uint64 Luid = BitCast<uint64>(AdapterLuid);
 	ID3D12Device* Device = FWindowsD3D::FindCachedD3D12AdapterDevice(Luid);
+	TRefCountPtr<ID3D12Device> DevicePtr;
 	if (!Device)
 	{
-		TRefCountPtr<ID3D12Device> DevicePtr;
 		const HRESULT D3D12CreateDeviceResult = D3D12CreateDevice(Adapter, MinFeatureLevel, IID_PPV_ARGS(DevicePtr.GetInitReference()));
 		if (FAILED(D3D12CreateDeviceResult))
 		{
@@ -205,7 +209,10 @@ static bool SafeTestD3D12CreateDevice(IDXGIAdapter* Adapter, LUID AdapterLuid, b
 		}
 
 		Device = DevicePtr.GetReference();
-		FWindowsD3D::CacheD3D12AdapterDevice(Luid, MoveTemp(DevicePtr));
+		if (FWindowsD3D::ShouldCacheAdapterDevices())
+		{
+			FWindowsD3D::CacheD3D12AdapterDevice(Luid, MoveTemp(DevicePtr));
+		}
 	}
 
 	D3D12_FEATURE_DATA_ARCHITECTURE data{};
@@ -273,12 +280,12 @@ static bool SafeTestD3D11CreateDevice(IDXGIAdapter* Adapter, LUID AdapterLuid, D
 	const uint64 Luid = BitCast<uint64>(AdapterLuid);
 	int32 Unused;
 	ID3D11Device* D3DDevice = FWindowsD3D::FindCachedD3D11AdapterDevice(Luid, Unused);
+	TRefCountPtr<ID3D11Device> DevicePtr;
 	if (!D3DDevice)
 	{
 		// We don't want software renderer. Ideally we specify D3D_DRIVER_TYPE_HARDWARE on creation but
 		// when we specify an adapter we need to specify D3D_DRIVER_TYPE_UNKNOWN (otherwise the call fails).
 		// We cannot check the device type later (seems this is missing functionality in D3D).
-		TRefCountPtr<ID3D11Device> DevicePtr;
 		D3D_FEATURE_LEVEL DeviceFeatureLevel = (D3D_FEATURE_LEVEL)0;
 		HRESULT Result = D3D11CreateDevice(
 			Adapter,
@@ -307,7 +314,10 @@ static bool SafeTestD3D11CreateDevice(IDXGIAdapter* Adapter, LUID AdapterLuid, D
 		}
 
 		D3DDevice = DevicePtr.GetReference();
-		FWindowsD3D::CacheD3D11AdapterDevice(Luid, MoveTemp(DevicePtr), DeviceFeatureLevel);
+		if (FWindowsD3D::ShouldCacheAdapterDevices())
+		{
+			FWindowsD3D::CacheD3D11AdapterDevice(Luid, MoveTemp(DevicePtr), DeviceFeatureLevel);
+		}
 	}
 
 	D3D11_FEATURE_DATA_D3D11_OPTIONS2 FeatureData = {};
@@ -629,6 +639,11 @@ uint64 FWindowsD3D::ChooseD3D12Adapter(uint64 DesiredAdapterLuid)
 	return Luid;
 }
 
+bool FWindowsD3D::ShouldCacheAdapterDevices()
+{
+	return GCacheD3DDevices;
+}
+
 ID3D11Device* FWindowsD3D::FindCachedD3D11AdapterDevice(uint64 AdapterLuid, int32& OutFeatureLevel)
 {
 	ID3D11Device* Device = nullptr;
