Re-enable D3D device caching during initialization for some vendors.

Intel is having issues with initialization due to the vendor-specific app discovery mechanism. Leave the caching disabled for those devices until the issue is resolved or worked around.

#rnx
#Synced-CL 52411106

diff --git a/Engine/Source/Runtime/D3D12RHI/Private/Windows/WindowsD3D12Device.cpp b/Engine/Source/Runtime/D3D12RHI/Private/Windows/WindowsD3D12Device.cpp
--- a/Engine/Source/Runtime/D3D12RHI/Private/Windows/WindowsD3D12Device.cpp
+++ b/Engine/Source/Runtime/D3D12RHI/Private/Windows/WindowsD3D12Device.cpp
@@ -488,11 +488,11 @@ inline bool GetSupportGPUUploadHeaps(ID3D12Device* device)
  * Attempts to create a D3D12 device for the adapter or find an existing one.
  * If creation is successful, true is returned.
  */
-static bool SafeTestD3D12CreateDevice(IDXGIAdapter* Adapter, LUID AdapterLuid, FD3D12DeviceBasicInfo& OutInfo)
+static bool SafeTestD3D12CreateDevice(IDXGIAdapter* Adapter, const DXGI_ADAPTER_DESC& AdapterDesc, FD3D12DeviceBasicInfo& OutInfo)
 {
 	constexpr D3D_FEATURE_LEVEL MinFeatureLevel = D3D_FEATURE_LEVEL_11_0;
 
-	const uint64 Luid = BitCast<uint64>(AdapterLuid);
+	const uint64 Luid = BitCast<uint64>(AdapterDesc.AdapterLuid);
 	ID3D12Device* Device = FWindowsD3D::FindCachedD3D12AdapterDevice(Luid);
 	TRefCountPtr<ID3D12Device> DevicePtr;
 	if (!Device)
@@ -505,7 +505,7 @@ static bool SafeTestD3D12CreateDevice(IDXGIAdapter* Adapter, LUID AdapterLuid, F
 		}
 
 		Device = DevicePtr.GetReference();
-		if (FWindowsD3D::ShouldCacheAdapterDevices())
+		if (FWindowsD3D::ShouldCacheAdapterDevices(AdapterDesc.VendorId))
 		{
 			FWindowsD3D::CacheD3D12AdapterDevice(Luid, MoveTemp(DevicePtr));
 		}
@@ -885,7 +885,7 @@ void FD3D12DynamicRHIModule::FindAdapter()
 			}
 #endif
 
-			if (SafeTestD3D12CreateDevice(TempAdapter, AdapterDesc.AdapterLuid, DeviceInfo))
+			if (SafeTestD3D12CreateDevice(TempAdapter, AdapterDesc, DeviceInfo))
 			{
 				check(DeviceInfo.NumDeviceNodes > 0);
 
diff --git a/Engine/Source/Runtime/Windows/D3D11RHI/Private/Windows/WindowsD3D11Device.cpp b/Engine/Source/Runtime/Windows/D3D11RHI/Private/Windows/WindowsD3D11Device.cpp
--- a/Engine/Source/Runtime/Windows/D3D11RHI/Private/Windows/WindowsD3D11Device.cpp
+++ b/Engine/Source/Runtime/Windows/D3D11RHI/Private/Windows/WindowsD3D11Device.cpp
@@ -278,9 +278,9 @@ inline bool IsDeviceUMA(ID3D11Device* device)
  * Attempts to create a D3D11 device for the adapter using at most MaxFeatureLevel or find an existing one.
  * If creation is successful, true is returned.
  */
-static bool SafeTestD3D11CreateDevice(IDXGIAdapter* Adapter, LUID AdapterLuid, D3D_FEATURE_LEVEL MinFeatureLevel, D3D_FEATURE_LEVEL MaxFeatureLevel, FD3D11DeviceBasicInfo& OutInfo)
+static bool SafeTestD3D11CreateDevice(IDXGIAdapter* Adapter, const DXGI_ADAPTER_DESC& AdapterDesc, D3D_FEATURE_LEVEL MinFeatureLevel, D3D_FEATURE_LEVEL MaxFeatureLevel, FD3D11DeviceBasicInfo& OutInfo)
 {
-	const uint64 Luid = BitCast<uint64>(AdapterLuid);
+	const uint64 Luid = BitCast<uint64>(AdapterDesc.AdapterLuid);
 	int32 CachedFeatureLevel;
 	ID3D11Device* D3DDevice = FWindowsD3D::FindCachedD3D11AdapterDevice(Luid, CachedFeatureLevel);
 	TRefCountPtr<ID3D11Device> DevicePtr;
@@ -367,7 +367,7 @@ static bool SafeTestD3D11CreateDevice(IDXGIAdapter* Adapter, LUID AdapterLuid, D
 		}
 
 		D3DDevice = DevicePtr.GetReference();
-		if (FWindowsD3D::ShouldCacheAdapterDevices())
+		if (FWindowsD3D::ShouldCacheAdapterDevices(AdapterDesc.VendorId))
 		{
 			FWindowsD3D::CacheD3D11AdapterDevice(Luid, MoveTemp(DevicePtr), DeviceFeatureLevel);
 		}
@@ -726,7 +726,7 @@ void FD3D11DynamicRHIModule::FindAdapter()
 				LogDXGIAdapterDesc(AdapterDesc);
 
 				FD3D11DeviceBasicInfo DeviceInfo;
-				if(SafeTestD3D11CreateDevice(TempAdapter, AdapterDesc.AdapterLuid, MinAllowedFeatureLevel, MaxAllowedFeatureLevel, DeviceInfo))
+				if(SafeTestD3D11CreateDevice(TempAdapter, AdapterDesc, MinAllowedFeatureLevel, MaxAllowedFeatureLevel, DeviceInfo))
 				{
 					// Log some information about the available D3D11 adapters.
 
diff --git a/Engine/Source/Runtime/Windows/WindowsD3D/Internal/Windows/WindowsD3D.h b/Engine/Source/Runtime/Windows/WindowsD3D/Internal/Windows/WindowsD3D.h
--- a/Engine/Source/Runtime/Windows/WindowsD3D/Internal/Windows/WindowsD3D.h
+++ b/Engine/Source/Runtime/Windows/WindowsD3D/Internal/Windows/WindowsD3D.h
@@ -25,7 +25,7 @@ public:
 	static WINDOWSD3D_API uint64 ChooseD3D11Adapter(uint64 DesiredAdapterLuid);
 	static WINDOWSD3D_API uint64 ChooseD3D12Adapter(uint64 DesiredAdapterLuid);
 
-	static WINDOWSD3D_API bool ShouldCacheAdapterDevices();
+	static WINDOWSD3D_API bool ShouldCacheAdapterDevices(uint32 GPUVendorId);
 	static WINDOWSD3D_API ID3D11Device* FindCachedD3D11AdapterDevice(uint64 AdapterLuid, int32& FeatureLevel);
 	static WINDOWSD3D_API void CacheD3D11AdapterDevice(uint64 AdapterLuid, TRefCountPtr<ID3D11Device>&& Device, int32 FeatureLevel);
 	static WINDOWSD3D_API ID3D12Device* FindCachedD3D12AdapterDevice(uint64 AdapterLuid);
diff --git a/Engine/Source/Runtime/Windows/WindowsD3D/Private/WindowsD3D.cpp b/Engine/Source/Runtime/Windows/WindowsD3D/Private/WindowsD3D.cpp
--- a/Engine/Source/Runtime/Windows/WindowsD3D/Private/WindowsD3D.cpp
+++ b/Engine/Source/Runtime/Windows/WindowsD3D/Private/WindowsD3D.cpp
@@ -15,10 +15,6 @@ DEFINE_LOG_CATEGORY_STATIC(LogD3D, Log, All);
 static TMap<uint64, TPair<TRefCountPtr<ID3D11Device>, int32>> D3D11AdapterDevices;
 static TMap<uint64, TRefCountPtr<ID3D12Device>> D3D12AdapterDevices;
 
-// Some vendors are seeing initializing issues with this optimization enabled.
-// Default to off until those problems are resolved.
-static bool GCacheD3DDevices = false;
-
 enum class ED3DGpuVendorId : uint32
 {
 	Unknown = 0xffffffff,
@@ -190,13 +186,13 @@ static HRESULT EnumAdapters(int32 AdapterIndex, DXGI_GPU_PREFERENCE GpuPreferenc
  * Attempts to create a D3D12 device for the adapter or find an existing one.
  * If creation is successful, true is returned.
  */
-static bool SafeTestD3D12CreateDevice(IDXGIAdapter* Adapter, LUID AdapterLuid, bool& bOutUMA)
+static bool SafeTestD3D12CreateDevice(IDXGIAdapter* Adapter, const DXGI_ADAPTER_DESC& AdapterDesc, bool& bOutUMA)
 {
 	bOutUMA = false;
 
 	constexpr D3D_FEATURE_LEVEL MinFeatureLevel = D3D_FEATURE_LEVEL_11_0;
 
-	const uint64 Luid = BitCast<uint64>(AdapterLuid);
+	const uint64 Luid = BitCast<uint64>(AdapterDesc.AdapterLuid);
 	ID3D12Device* Device = FWindowsD3D::FindCachedD3D12AdapterDevice(Luid);
 	TRefCountPtr<ID3D12Device> DevicePtr;
 	if (!Device)
@@ -209,7 +205,7 @@ static bool SafeTestD3D12CreateDevice(IDXGIAdapter* Adapter, LUID AdapterLuid, b
 		}
 
 		Device = DevicePtr.GetReference();
-		if (FWindowsD3D::ShouldCacheAdapterDevices())
+		if (FWindowsD3D::ShouldCacheAdapterDevices(AdapterDesc.VendorId))
 		{
 			FWindowsD3D::CacheD3D12AdapterDevice(Luid, MoveTemp(DevicePtr));
 		}
@@ -224,7 +220,7 @@ static bool SafeTestD3D12CreateDevice(IDXGIAdapter* Adapter, LUID AdapterLuid, b
  * Attempts to create a D3D11 device for the adapter using at most MaxFeatureLevel or find an existing one.
  * If creation is successful, true is returned.
  */
-static bool SafeTestD3D11CreateDevice(IDXGIAdapter* Adapter, LUID AdapterLuid, D3D_FEATURE_LEVEL MinFeatureLevel, D3D_FEATURE_LEVEL MaxFeatureLevel, bool& bOutUMA)
+static bool SafeTestD3D11CreateDevice(IDXGIAdapter* Adapter, const DXGI_ADAPTER_DESC& AdapterDesc, D3D_FEATURE_LEVEL MinFeatureLevel, D3D_FEATURE_LEVEL MaxFeatureLevel, bool& bOutUMA)
 {
 	uint32 DeviceFlags = D3D11_CREATE_DEVICE_SINGLETHREADED;
 	// Use a debug device if specified on the command line.
@@ -277,7 +273,7 @@ static bool SafeTestD3D11CreateDevice(IDXGIAdapter* Adapter, LUID AdapterLuid, D
 		return false;
 	}
 
-	const uint64 Luid = BitCast<uint64>(AdapterLuid);
+	const uint64 Luid = BitCast<uint64>(AdapterDesc.AdapterLuid);
 	int32 Unused;
 	ID3D11Device* D3DDevice = FWindowsD3D::FindCachedD3D11AdapterDevice(Luid, Unused);
 	TRefCountPtr<ID3D11Device> DevicePtr;
@@ -314,7 +310,7 @@ static bool SafeTestD3D11CreateDevice(IDXGIAdapter* Adapter, LUID AdapterLuid, D
 		}
 
 		D3DDevice = DevicePtr.GetReference();
-		if (FWindowsD3D::ShouldCacheAdapterDevices())
+		if (FWindowsD3D::ShouldCacheAdapterDevices(AdapterDesc.VendorId))
 		{
 			FWindowsD3D::CacheD3D11AdapterDevice(Luid, MoveTemp(DevicePtr), DeviceFeatureLevel);
 		}
@@ -393,7 +389,7 @@ uint64 FWindowsD3D::ChooseD3D11Adapter(uint64 DesiredAdapterLuid)
 		{
 			DXGI_ADAPTER_DESC AdapterDesc{};
 			bool bUMA;
-			if (SUCCEEDED(TempAdapter->GetDesc(&AdapterDesc)) && SafeTestD3D11CreateDevice(TempAdapter, AdapterDesc.AdapterLuid, MinAllowedFeatureLevel, MaxAllowedFeatureLevel, bUMA))
+			if (SUCCEEDED(TempAdapter->GetDesc(&AdapterDesc)) && SafeTestD3D11CreateDevice(TempAdapter, AdapterDesc, MinAllowedFeatureLevel, MaxAllowedFeatureLevel, bUMA))
 			{
 				const bool bIsWARP = ConvertToGpuVendorId(AdapterDesc.VendorId) == ED3DGpuVendorId::Microsoft;
 
@@ -548,7 +544,7 @@ uint64 FWindowsD3D::ChooseD3D12Adapter(uint64 DesiredAdapterLuid)
 		{
 			bool bUMA;
 			DXGI_ADAPTER_DESC AdapterDesc{};
-			if (SUCCEEDED(TempAdapter->GetDesc(&AdapterDesc)) && SafeTestD3D12CreateDevice(TempAdapter, AdapterDesc.AdapterLuid, bUMA))
+			if (SUCCEEDED(TempAdapter->GetDesc(&AdapterDesc)) && SafeTestD3D12CreateDevice(TempAdapter, AdapterDesc, bUMA))
 			{
 				const bool bIsWARP = ConvertToGpuVendorId(AdapterDesc.VendorId) == ED3DGpuVendorId::Microsoft;
 
@@ -639,9 +635,10 @@ uint64 FWindowsD3D::ChooseD3D12Adapter(uint64 DesiredAdapterLuid)
 	return Luid;
 }
 
-bool FWindowsD3D::ShouldCacheAdapterDevices()
+bool FWindowsD3D::ShouldCacheAdapterDevices(uint32 GPUVendorId)
 {
-	return GCacheD3DDevices;
+	// Do not cache Intel devices until order initialization issue with Intel App Discovery is resolved.
+	return GPUVendorId != static_cast<uint32>(ED3DGpuVendorId::Intel);
 }
 
 ID3D11Device* FWindowsD3D::FindCachedD3D11AdapterDevice(uint64 AdapterLuid, int32& OutFeatureLevel)
