Fix assertion in DX11 ReadSurfaceData that was generating false positives on TYPELESS formats (which can be generated when running the editor over Parsec/RDP).

#Synced-CL 52671970

diff --git a/Engine/Source/Runtime/Windows/D3D11RHI/Private/D3D11RenderTarget.cpp b/Engine/Source/Runtime/Windows/D3D11RHI/Private/D3D11RenderTarget.cpp
--- a/Engine/Source/Runtime/Windows/D3D11RHI/Private/D3D11RenderTarget.cpp
+++ b/Engine/Source/Runtime/Windows/D3D11RHI/Private/D3D11RenderTarget.cpp
@@ -442,7 +442,7 @@ void FD3D11DynamicRHI::RHIReadSurfaceData(FRHITexture* TextureRHI,FIntRect InRec
 	//	verify against TextureDesc.Format
 
 	EPixelFormat Format = TextureRHI->GetFormat();
-	check( GPixelFormats[Format].PlatformFormat == TextureDesc.Format );
+	check( ConvertTypelessToUnorm((DXGI_FORMAT)GPixelFormats[Format].PlatformFormat) == ConvertTypelessToUnorm(TextureDesc.Format) );
 	check( GPixelFormats[Format].BlockBytes == D3D11RT_ComputeBytesPerPixel(TextureDesc.Format) );
 
 	// ConvertDXGIToFColor switches on the hardware format, not the EPixelFormat :
@@ -716,7 +716,7 @@ void FD3D11DynamicRHI::RHIReadSurfaceData(FRHITexture* TextureRHI, FIntRect InRe
 	uint32 SrcPitch = SizeX * BytesPerPixel;
 	EPixelFormat Format = TextureRHI->GetFormat();
 
-	check( GPixelFormats[Format].PlatformFormat == TextureDesc.Format );
+	check( ConvertTypelessToUnorm((DXGI_FORMAT)GPixelFormats[Format].PlatformFormat) == ConvertTypelessToUnorm(TextureDesc.Format) );
 
 	// switching on the EPixelFormat is risky if the mapping is not what you expect
 	//	verify against TextureDesc.Format
