Prevent rescursive calls to FD3D1XViewport::Resize on change of swapchain fullscreen state

#Synced-CL 52720542

diff --git a/Engine/Source/Runtime/D3D12RHI/Private/D3D12Viewport.cpp b/Engine/Source/Runtime/D3D12RHI/Private/D3D12Viewport.cpp
--- a/Engine/Source/Runtime/D3D12RHI/Private/D3D12Viewport.cpp
+++ b/Engine/Source/Runtime/D3D12RHI/Private/D3D12Viewport.cpp
@@ -277,6 +277,17 @@ void FD3D12Viewport::InitializeBackBufferArrays()
 
 void FD3D12Viewport::Resize(uint32 InSizeX, uint32 InSizeY, bool bInIsFullscreen, EPixelFormat PreferredPixelFormat)
 {
+	if (bResizeInProgress)
+	{
+		return;
+	}
+
+	bResizeInProgress = true;
+	ON_SCOPE_EXIT
+	{
+		bResizeInProgress = false;
+	};
+
 	FD3D12Adapter* Adapter = GetParentAdapter();
 
 #if !UE_BUILD_SHIPPING
diff --git a/Engine/Source/Runtime/D3D12RHI/Private/D3D12Viewport.h b/Engine/Source/Runtime/D3D12RHI/Private/D3D12Viewport.h
--- a/Engine/Source/Runtime/D3D12RHI/Private/D3D12Viewport.h
+++ b/Engine/Source/Runtime/D3D12RHI/Private/D3D12Viewport.h
@@ -232,6 +232,7 @@ private:
 	bool bAllowTearing = true;
 	bool bNeedSwapChain = false;
 	bool bLastFramePresented = true;
+	bool bResizeInProgress = false;
 
 	uint32 CheckedPresentFailureCounter = 0;
 
diff --git a/Engine/Source/Runtime/Windows/D3D11RHI/Private/D3D11Viewport.cpp b/Engine/Source/Runtime/Windows/D3D11RHI/Private/D3D11Viewport.cpp
--- a/Engine/Source/Runtime/Windows/D3D11RHI/Private/D3D11Viewport.cpp
+++ b/Engine/Source/Runtime/Windows/D3D11RHI/Private/D3D11Viewport.cpp
@@ -242,6 +242,17 @@ DXGI_MODE_DESC FD3D11Viewport::SetupDXGI_MODE_DESC() const
 
 void FD3D11Viewport::Resize(uint32 InSizeX, uint32 InSizeY, bool bInIsFullscreen, EPixelFormat PreferredPixelFormat)
 {
+	if (bResizeInProgress)
+	{
+		return;
+	}
+
+	bResizeInProgress = true;
+	ON_SCOPE_EXIT
+	{
+		bResizeInProgress = false;
+	};
+
 	// Unbind any dangling references to resources
 	D3DRHI->SetRenderTargets(0, nullptr, nullptr);
 	D3DRHI->ClearState();
diff --git a/Engine/Source/Runtime/Windows/D3D11RHI/Public/D3D11Viewport.h b/Engine/Source/Runtime/Windows/D3D11RHI/Public/D3D11Viewport.h
--- a/Engine/Source/Runtime/Windows/D3D11RHI/Public/D3D11Viewport.h
+++ b/Engine/Source/Runtime/Windows/D3D11RHI/Public/D3D11Viewport.h
@@ -118,6 +118,7 @@ protected:
 	EDisplayColorGamut DisplayColorGamut = EDisplayColorGamut::sRGB_D65;
 	EDisplayOutputFormat DisplayOutputFormat = EDisplayOutputFormat::SDR_sRGB;
 	bool bAllowTearing = false;
+	bool bResizeInProgress = false;
 
 	static D3D11RHI_API uint32 GSwapChainFlags;
 
